Added several LiveLink blueprint functions which were missing from initial merge.

diff --git a/Engine/Plugins/Animation/LiveLink/Source/LiveLink/Private/LiveLinkBlueprintLibrary.cpp b/Engine/Plugins/Animation/LiveLink/Source/LiveLink/Private/LiveLinkBlueprintLibrary.cpp
--- a/Engine/Plugins/Animation/LiveLink/Source/LiveLink/Private/LiveLinkBlueprintLibrary.cpp
+++ b/Engine/Plugins/Animation/LiveLink/Source/LiveLink/Private/LiveLinkBlueprintLibrary.cpp
@@ -14,6 +14,11 @@ int ULiveLinkBlueprintLibrary::NumberOfTransforms(UPARAM(ref) FSubjectFrameHandl
 	return SubjectFrameHandle.GetNumberOfTransforms();
 };
 
+void ULiveLinkBlueprintLibrary::TransformNames(UPARAM(ref) FSubjectFrameHandle& SubjectFrameHandle, TArray<FName>& TransformNames)
+{
+	return SubjectFrameHandle.GetTransformNames(TransformNames);
+}
+
 void ULiveLinkBlueprintLibrary::GetRootTransform(UPARAM(ref) FSubjectFrameHandle& SubjectFrameHandle, FLiveLinkTransform& LiveLinkTransform)
 {
 	SubjectFrameHandle.GetRootTransform(LiveLinkTransform);
@@ -24,6 +29,11 @@ void ULiveLinkBlueprintLibrary::GetTransformByIndex(UPARAM(ref) FSubjectFrameHan
 	SubjectFrameHandle.GetTransformByIndex(TransformIndex, LiveLinkTransform);
 };
 
+void ULiveLinkBlueprintLibrary::GetTransformByName(UPARAM(ref) FSubjectFrameHandle& SubjectFrameHandle, FName TransformName, FLiveLinkTransform& LiveLinkTransform)
+{
+	SubjectFrameHandle.GetTransformByName(TransformName, LiveLinkTransform);
+};
+
 void ULiveLinkBlueprintLibrary::GetMetadata(UPARAM(ref) FSubjectFrameHandle& SubjectFrameHandle, FSubjectMetadata& Metadata)
 {
 	SubjectFrameHandle.GetSubjectMetadata(Metadata);
diff --git a/Engine/Plugins/Animation/LiveLink/Source/LiveLink/Private/LiveLinkBlueprintStructs.cpp b/Engine/Plugins/Animation/LiveLink/Source/LiveLink/Private/LiveLinkBlueprintStructs.cpp
--- a/Engine/Plugins/Animation/LiveLink/Source/LiveLink/Private/LiveLinkBlueprintStructs.cpp
+++ b/Engine/Plugins/Animation/LiveLink/Source/LiveLink/Private/LiveLinkBlueprintStructs.cpp
@@ -51,6 +51,11 @@ int FCachedSubjectFrame::GetNumberOfTransforms()
 	return SourceFrame.Transforms.Num();
 };
 
+void FCachedSubjectFrame::GetTransformNames(TArray<FName>& TransformNames)
+{
+	TransformNames = SourceFrame.RefSkeleton.GetBoneNames();
+};
+
 void FCachedSubjectFrame::GetTransformName(const int TransformIndex, FName& Name)
 {
 	if (IsValidTransformIndex(TransformIndex))
@@ -63,6 +68,11 @@ void FCachedSubjectFrame::GetTransformName(const int TransformIndex, FName& Name
 	}
 };
 
+int FCachedSubjectFrame::GetTransformIndexFromName(FName TransformName)
+{
+	return SourceFrame.RefSkeleton.GetBoneNames().IndexOfByKey(TransformName);
+};
+
 int FCachedSubjectFrame::GetParentTransformIndex(const int TransformIndex)
 {
 	if (IsValidTransformIndex(TransformIndex))
@@ -297,7 +307,15 @@ int FSubjectFrameHandle::GetNumberOfTransforms()
 	{
 		return 0;
 	}
-}
+};
+
+void FSubjectFrameHandle::GetTransformNames(TArray<FName>& TransformNames)
+{
+	if (CachedFrame.IsValid())
+	{
+		CachedFrame->GetTransformNames(TransformNames);
+	}
+};
 
 void FSubjectFrameHandle::GetRootTransform(FLiveLinkTransform& LiveLinkTransform)
 {
@@ -314,6 +332,13 @@ void FSubjectFrameHandle::GetTransformByIndex(int TransformIndex, FLiveLinkTrans
 	LiveLinkTransform.SetTransformIndex(TransformIndex);
 };
 
+void FSubjectFrameHandle::GetTransformByName(FName TransformName, FLiveLinkTransform& LiveLinkTransform)
+{
+	int TransformIndex = CachedFrame->GetTransformIndexFromName(TransformName);
+	LiveLinkTransform.SetCachedFrame(CachedFrame);
+	LiveLinkTransform.SetTransformIndex(TransformIndex);
+};
+
 void FSubjectFrameHandle::SetCachedFrame(TSharedPtr<FCachedSubjectFrame> InCachedFrame)
 {
 	CachedFrame = InCachedFrame;
diff --git a/Engine/Plugins/Animation/LiveLink/Source/LiveLink/Private/LiveLinkComponent.cpp b/Engine/Plugins/Animation/LiveLink/Source/LiveLink/Private/LiveLinkComponent.cpp
--- a/Engine/Plugins/Animation/LiveLink/Source/LiveLink/Private/LiveLinkComponent.cpp
+++ b/Engine/Plugins/Animation/LiveLink/Source/LiveLink/Private/LiveLinkComponent.cpp
@@ -61,6 +61,15 @@ bool ULiveLinkComponent::HasLiveLinkClient()
 	return (LiveLinkClient != nullptr);
 }
 
+void ULiveLinkComponent::GetAvailableSubjectNames(TArray<FName>& SubjectNames)
+{
+	SubjectNames.Empty();
+	if (HasLiveLinkClient())
+	{
+		LiveLinkClient->GetSubjectNames(SubjectNames);
+	}
+}
+
 void ULiveLinkComponent::GetSubjectData(const FName SubjectName, bool& bSuccess, FSubjectFrameHandle& SubjectFrameHandle)
 {
 	bSuccess = false;
diff --git a/Engine/Plugins/Animation/LiveLink/Source/LiveLink/Public/LiveLinkBlueprintLibrary.h b/Engine/Plugins/Animation/LiveLink/Source/LiveLink/Public/LiveLinkBlueprintLibrary.h
--- a/Engine/Plugins/Animation/LiveLink/Source/LiveLink/Public/LiveLinkBlueprintLibrary.h
+++ b/Engine/Plugins/Animation/LiveLink/Source/LiveLink/Public/LiveLinkBlueprintLibrary.h
@@ -22,6 +22,10 @@ class LIVELINK_API ULiveLinkBlueprintLibrary : public UBlueprintFunctionLibrary
 	UFUNCTION(BlueprintPure, Category = "LiveLink")
 	static int NumberOfTransforms(UPARAM(ref) FSubjectFrameHandle& SubjectFrameHandle);
 
+	// Returns an array of Transform Names stored in the Subject Frame
+	UFUNCTION(BlueprintPure, Category = "LiveLink")
+	static void TransformNames(UPARAM(ref) FSubjectFrameHandle& SubjectFrameHandle, TArray<FName>& TransformNames);
+
 	// Returns the Root Transform for the Subject Frame as a LiveLink Transform or the Identity if there are no transforms.
 	UFUNCTION(BlueprintPure, Category = "LiveLink")
 	static void GetRootTransform(UPARAM(ref) FSubjectFrameHandle& SubjectFrameHandle, FLiveLinkTransform& LiveLinkTransform);
@@ -30,6 +34,10 @@ class LIVELINK_API ULiveLinkBlueprintLibrary : public UBlueprintFunctionLibrary
 	UFUNCTION(BlueprintPure, Category = "LiveLink")
 	static void GetTransformByIndex(UPARAM(ref) FSubjectFrameHandle& SubjectFrameHandle, int TransformIndex, FLiveLinkTransform& LiveLinkTransform);
 
+	// Returns the LiveLink Transform stored in a Subject Frame with a given name. Returns an Identity transform if Transform Name is invalid.
+	UFUNCTION(BlueprintPure, Category = "LiveLink")
+	static void GetTransformByName(UPARAM(ref) FSubjectFrameHandle& SubjectFrameHandle, FName TransformName, FLiveLinkTransform& LiveLinkTransform);
+
 	// Returns the Subject Metadata structure stored in the Subject Frame
 	UFUNCTION(BlueprintPure, Category = "LiveLink")
 	static void GetMetadata(UPARAM(ref) FSubjectFrameHandle& SubjectFrameHandle, FSubjectMetadata& Metadata);
diff --git a/Engine/Plugins/Animation/LiveLink/Source/LiveLink/Public/LiveLinkBlueprintStructs.h b/Engine/Plugins/Animation/LiveLink/Source/LiveLink/Public/LiveLinkBlueprintStructs.h
--- a/Engine/Plugins/Animation/LiveLink/Source/LiveLink/Public/LiveLinkBlueprintStructs.h
+++ b/Engine/Plugins/Animation/LiveLink/Source/LiveLink/Public/LiveLinkBlueprintStructs.h
@@ -41,8 +41,12 @@ struct FCachedSubjectFrame
 
 	int GetNumberOfTransforms();
 
+	void GetTransformNames(TArray<FName>& TransformNames);
+
 	void GetTransformName(const int TransformIndex, FName& Name);
 
+	int GetTransformIndexFromName(FName TransformName);
+
 	int GetParentTransformIndex(const int TransformIndex);
 
 	void GetChildTransformIndices(const int TransformIndex, TArray<int>& ChildIndices);
@@ -115,10 +119,14 @@ struct FSubjectFrameHandle
 
 	int GetNumberOfTransforms();
 
+	void GetTransformNames(TArray<FName>& TransformNames);
+
 	void GetRootTransform(FLiveLinkTransform& LiveLinkTransform);
 
 	void GetTransformByIndex(int TransformIndex, FLiveLinkTransform& LiveLinkTransform);
 
+	void GetTransformByName(FName TransformName, FLiveLinkTransform& LiveLinkTransform);
+
 	void SetCachedFrame(TSharedPtr<FCachedSubjectFrame> InCachedFrame);
 private:
 	TSharedPtr<FCachedSubjectFrame> CachedFrame;
diff --git a/Engine/Plugins/Animation/LiveLink/Source/LiveLink/Public/LiveLinkComponent.h b/Engine/Plugins/Animation/LiveLink/Source/LiveLink/Public/LiveLinkComponent.h
--- a/Engine/Plugins/Animation/LiveLink/Source/LiveLink/Public/LiveLinkComponent.h
+++ b/Engine/Plugins/Animation/LiveLink/Source/LiveLink/Public/LiveLinkComponent.h
@@ -33,6 +33,10 @@ public:
 	// This Event is triggered any time new LiveLink data is available, including in the editor
 	UPROPERTY(BlueprintAssignable, Category = "LiveLink")
 	FLiveLinkTickSignature OnLiveLinkUpdated;
+
+	// Returns a list of available Subject Names for LiveLink
+	UFUNCTION(BlueprintCallable, Category = "LiveLink")
+	void GetAvailableSubjectNames(TArray<FName>& SubjectNames);
 	
 	// Returns a handle to the current frame of data in LiveLink for a given subject along with a boolean for whether a frame was found.
 	// Returns a handle to an empty frame if no frame of data is found.
