Fixed LiveLink Client Panel not refreshing for sources not added through UI

diff --git a/Engine/Plugins/Animation/LiveLink/Source/LiveLink/Private/LiveLinkClient.cpp b/Engine/Plugins/Animation/LiveLink/Source/LiveLink/Private/LiveLinkClient.cpp
--- a/Engine/Plugins/Animation/LiveLink/Source/LiveLink/Private/LiveLinkClient.cpp
+++ b/Engine/Plugins/Animation/LiveLink/Source/LiveLink/Private/LiveLinkClient.cpp
@@ -1,4 +1,4 @@
-﻿// Copyright 1998-2018 Epic Games, Inc. All Rights Reserved.
+// Copyright 1998-2018 Epic Games, Inc. All Rights Reserved.
 
 #include "LiveLinkClient.h"
 #include "Misc/ScopeLock.h"
@@ -402,6 +402,8 @@ void FLiveLinkClient::AddSource(TSharedPtr<ILiveLinkSource> InSource)
 	
 	InSource->ReceiveClient(this, SourceGuids.Last());
 	InSource->InitializeSettings(NewSettings);
+
+	OnLiveLinkSourcesChanged.Broadcast();
 }
 
 void FLiveLinkClient::AddVirtualSubjectSource()
diff --git a/Engine/Plugins/Animation/LiveLink/Source/LiveLinkEditor/Private/LiveLinkClientPanel.cpp b/Engine/Plugins/Animation/LiveLink/Source/LiveLinkEditor/Private/LiveLinkClientPanel.cpp
--- a/Engine/Plugins/Animation/LiveLink/Source/LiveLinkEditor/Private/LiveLinkClientPanel.cpp
+++ b/Engine/Plugins/Animation/LiveLink/Source/LiveLinkEditor/Private/LiveLinkClientPanel.cpp
@@ -592,8 +592,6 @@ FReply SLiveLinkClientPanel::OnCloseSourceSelectionPanel(ULiveLinkSourceFactory*
 	if (Source.IsValid())
 	{
 		Client->AddSource(Source);
-
-		RefreshSourceData(true);
 	}
 	
 	FSlateApplication::Get().DismissAllMenus();
