Change Live Link tab to nomad

diff --git a/Engine/Plugins/Animation/LiveLink/Source/LiveLinkEditor/Private/LiveLinkEditorModule.cpp b/Engine/Plugins/Animation/LiveLink/Source/LiveLinkEditor/Private/LiveLinkEditorModule.cpp
--- a/Engine/Plugins/Animation/LiveLink/Source/LiveLinkEditor/Private/LiveLinkEditorModule.cpp
+++ b/Engine/Plugins/Animation/LiveLink/Source/LiveLinkEditor/Private/LiveLinkEditorModule.cpp
@@ -59,9 +59,11 @@ public:
 		static FName LiveLinkStyle(TEXT("LiveLinkStyle"));
 		StyleSet = MakeShareable(new FSlateStyleSet(LiveLinkStyle));
 
+		bHasRegisteredTabSpawners = false;
+
 		if (FModuleManager::Get().IsModuleLoaded(LevelEditorModuleName))
 		{
-			RegisterTab();
+			RegisterTabSpawner();
 		}
 		else
 		{
@@ -97,34 +99,13 @@ public:
 	{
 		if (ReasonForChange == EModuleChangeReason::ModuleLoaded && ModuleName == LevelEditorModuleName)
 		{
-			RegisterTab();
+			RegisterTabSpawner();
 		}
 	}
 
-	void RegisterTab()
-	{
-		FLevelEditorModule& LevelEditorModule = FModuleManager::LoadModuleChecked<FLevelEditorModule>(LevelEditorModuleName);
-
-		TSharedPtr<FSlateStyleSet> StyleSetPtr = StyleSet;
-
-		//Register our UI
-		LevelEditorTabManagerChangedHandle = LevelEditorModule.OnTabManagerChanged().AddLambda([StyleSetPtr]()
-		{
-			FLevelEditorModule& LocalLevelEditorModule = FModuleManager::GetModuleChecked<FLevelEditorModule>(LevelEditorModuleName);
-			LocalLevelEditorModule.GetLevelEditorTabManager()->RegisterTabSpawner(LiveLinkClientTabName, FOnSpawnTab::CreateStatic(&FLiveLinkEditorModule::SpawnLiveLinkTab, StyleSetPtr))
-				.SetGroup(WorkspaceMenu::GetMenuStructure().GetLevelEditorCategory())
-				.SetDisplayName(LOCTEXT("LiveLinkTabTitle", "Live Link"))
-				.SetTooltipText(LOCTEXT("SequenceRecorderTooltipText", "Open the Live Link streaming manager tab."))
-				.SetIcon(FSlateIcon(StyleSetPtr->GetStyleSetName(), "LiveLinkClient.Common.Icon.Small"));
-		});
-	}
-
 	virtual void ShutdownModule() override
 	{
-		if (FSlateApplication::IsInitialized())
-		{
-			FGlobalTabmanager::Get()->UnregisterTabSpawner(LiveLinkClientTabName);
-		}
+		UnregisterTabSpawner();
 
 		FModuleManager::Get().OnModulesChanged().Remove(ModulesChangedHandle);
 
@@ -166,12 +147,38 @@ public:
 
 private:
 
+	void RegisterTabSpawner()
+	{
+		if (bHasRegisteredTabSpawners)
+		{
+			UnregisterTabSpawner();
+		}
+
+		TSharedPtr<FSlateStyleSet> StyleSetPtr = StyleSet;
+		FTabSpawnerEntry& SpawnerEntry = FGlobalTabmanager::Get()->RegisterNomadTabSpawner(LiveLinkClientTabName, FOnSpawnTab::CreateStatic(&FLiveLinkEditorModule::SpawnLiveLinkTab, StyleSetPtr))
+			.SetDisplayName(LOCTEXT("LiveLinkTabTitle", "Live Link"))
+			.SetTooltipText(LOCTEXT("SequenceRecorderTooltipText", "Open the Live Link streaming manager tab."))
+			.SetIcon(FSlateIcon(StyleSetPtr->GetStyleSetName(), "LiveLinkClient.Common.Icon.Small"));
+
+		const IWorkspaceMenuStructure& MenuStructure = WorkspaceMenu::GetMenuStructure();
+		SpawnerEntry.SetGroup(MenuStructure.GetLevelEditorCategory());
+
+		bHasRegisteredTabSpawners = true;
+	}
+
+	void UnregisterTabSpawner()
+	{
+		FGlobalTabmanager::Get()->UnregisterNomadTabSpawner(LiveLinkClientTabName);
+		bHasRegisteredTabSpawners = false;
+	}
+
 	FDelegateHandle LevelEditorTabManagerChangedHandle;
 	FDelegateHandle ModulesChangedHandle;
 	FDelegateHandle CreateLiveLinkPropertyTrackEditorHandle;
 	FMovieSceneLiveLinkSectionRecorderFactory MovieSceneLiveLinkRecorder;
 
-
+	// Track if we have registered
+	bool bHasRegisteredTabSpawners;
 };
 
 IMPLEMENT_MODULE(FLiveLinkEditorModule, LiveLinkEditor);
