Fix possible crash due to retarget asset being cleaned up on ULiveLinkInstance

diff --git a/Engine/Plugins/Animation/LiveLink/Source/LiveLink/Private/LiveLinkInstance.cpp b/Engine/Plugins/Animation/LiveLink/Source/LiveLink/Private/LiveLinkInstance.cpp
--- a/Engine/Plugins/Animation/LiveLink/Source/LiveLink/Private/LiveLinkInstance.cpp
+++ b/Engine/Plugins/Animation/LiveLink/Source/LiveLink/Private/LiveLinkInstance.cpp
@@ -24,10 +24,16 @@ void FLiveLinkInstanceProxy::UpdateAnimationNode(float DeltaSeconds)
 
 	FAnimationUpdateContext UpdateContext(this, DeltaSeconds);
 	PoseNode.Update_AnyThread(UpdateContext);
+	
+	if(ULiveLinkInstance* Instance = Cast<ULiveLinkInstance>(GetAnimInstanceObject()))
+	{
+		Instance->CurrentRetargetAsset = PoseNode.CurrentRetargetAsset; //Cache for GC
+	}
 }
 
 ULiveLinkInstance::ULiveLinkInstance(const FObjectInitializer& Initializer)
 	: Super(Initializer)
+	, CurrentRetargetAsset(nullptr)
 {
 
 }
@@ -35,4 +41,10 @@ ULiveLinkInstance::ULiveLinkInstance(const FObjectInitializer& Initializer)
 FAnimInstanceProxy* ULiveLinkInstance::CreateAnimInstanceProxy()
 {
 	return new FLiveLinkInstanceProxy(this);
+}
+
+void ULiveLinkInstance::DestroyAnimInstanceProxy(FAnimInstanceProxy* InProxy)
+{
+	Super::DestroyAnimInstanceProxy(InProxy);
+	CurrentRetargetAsset = nullptr;
 }
\ No newline at end of file
diff --git a/Engine/Plugins/Animation/LiveLink/Source/LiveLink/Public/LiveLinkInstance.h b/Engine/Plugins/Animation/LiveLink/Source/LiveLink/Public/LiveLinkInstance.h
--- a/Engine/Plugins/Animation/LiveLink/Source/LiveLink/Public/LiveLinkInstance.h
+++ b/Engine/Plugins/Animation/LiveLink/Source/LiveLink/Public/LiveLinkInstance.h
@@ -52,4 +52,11 @@ class LIVELINK_API ULiveLinkInstance : public UAnimInstance
 
 protected:
 	virtual FAnimInstanceProxy* CreateAnimInstanceProxy() override;
+	virtual void DestroyAnimInstanceProxy(FAnimInstanceProxy* InProxy) override;
+
+	// Cache for GC
+	UPROPERTY(transient)
+	ULiveLinkRetargetAsset* CurrentRetargetAsset;
+
+	friend FLiveLinkInstanceProxy;
 };
