Copyright notice updates

diff --git a/Engine/Plugins/Animation/LiveLink/Source/LiveLink/Private/AnimNode_LiveLinkArchiveComponentPose.cpp b/Engine/Plugins/Animation/LiveLink/Source/LiveLink/Private/AnimNode_LiveLinkArchiveComponentPose.cpp
--- a/Engine/Plugins/Animation/LiveLink/Source/LiveLink/Private/AnimNode_LiveLinkArchiveComponentPose.cpp
+++ b/Engine/Plugins/Animation/LiveLink/Source/LiveLink/Private/AnimNode_LiveLinkArchiveComponentPose.cpp
@@ -1,4 +1,4 @@
-// Copyright 1998-2018 Epic Games, Inc. All Rights Reserved.
+// Copyright 1998-2019 Epic Games, Inc. All Rights Reserved.
 
 #include "AnimNode_LiveLinkArchiveComponentPose.h"
 
diff --git a/Engine/Plugins/Animation/LiveLink/Source/LiveLink/Private/LiveLinkArchiveComponent.cpp b/Engine/Plugins/Animation/LiveLink/Source/LiveLink/Private/LiveLinkArchiveComponent.cpp
--- a/Engine/Plugins/Animation/LiveLink/Source/LiveLink/Private/LiveLinkArchiveComponent.cpp
+++ b/Engine/Plugins/Animation/LiveLink/Source/LiveLink/Private/LiveLinkArchiveComponent.cpp
@@ -1,4 +1,4 @@
-// Copyright 1998-2018 Epic Games, Inc. All Rights Reserved.
+// Copyright 1998-2019 Epic Games, Inc. All Rights Reserved.
 
 #include "LiveLinkArchiveComponent.h"
 #include "Features/IModularFeatures.h"
diff --git a/Engine/Plugins/Animation/LiveLink/Source/LiveLink/Public/AnimNode_LiveLinkArchiveComponentPose.h b/Engine/Plugins/Animation/LiveLink/Source/LiveLink/Public/AnimNode_LiveLinkArchiveComponentPose.h
--- a/Engine/Plugins/Animation/LiveLink/Source/LiveLink/Public/AnimNode_LiveLinkArchiveComponentPose.h
+++ b/Engine/Plugins/Animation/LiveLink/Source/LiveLink/Public/AnimNode_LiveLinkArchiveComponentPose.h
@@ -1,4 +1,4 @@
-// Copyright 1998-2018 Epic Games, Inc. All Rights Reserved.
+// Copyright 1998-2019 Epic Games, Inc. All Rights Reserved.
 
 #pragma once
 
diff --git a/Engine/Plugins/Animation/LiveLink/Source/LiveLink/Public/LiveLinkArchiveComponent.h b/Engine/Plugins/Animation/LiveLink/Source/LiveLink/Public/LiveLinkArchiveComponent.h
--- a/Engine/Plugins/Animation/LiveLink/Source/LiveLink/Public/LiveLinkArchiveComponent.h
+++ b/Engine/Plugins/Animation/LiveLink/Source/LiveLink/Public/LiveLinkArchiveComponent.h
@@ -1,4 +1,4 @@
-// Copyright 1998-2018 Epic Games, Inc. All Rights Reserved.
+// Copyright 1998-2019 Epic Games, Inc. All Rights Reserved.
 
 #pragma once
 
diff --git a/Engine/Plugins/Animation/LiveLink/Source/LiveLinkEditor/Private/AnimGraphNode_LiveLinkArchiveComponentPose.cpp b/Engine/Plugins/Animation/LiveLink/Source/LiveLinkEditor/Private/AnimGraphNode_LiveLinkArchiveComponentPose.cpp
--- a/Engine/Plugins/Animation/LiveLink/Source/LiveLinkEditor/Private/AnimGraphNode_LiveLinkArchiveComponentPose.cpp
+++ b/Engine/Plugins/Animation/LiveLink/Source/LiveLinkEditor/Private/AnimGraphNode_LiveLinkArchiveComponentPose.cpp
@@ -1,4 +1,4 @@
-// Copyright 1998-2018 Epic Games, Inc. All Rights Reserved.
+// Copyright 1998-2019 Epic Games, Inc. All Rights Reserved.
 
 #include "AnimGraphNode_LiveLinkArchiveComponentPose.h"
 
diff --git a/Engine/Plugins/Animation/LiveLink/Source/LiveLinkEditor/Private/AnimGraphNode_LiveLinkArchiveComponentPose.h b/Engine/Plugins/Animation/LiveLink/Source/LiveLinkEditor/Private/AnimGraphNode_LiveLinkArchiveComponentPose.h
--- a/Engine/Plugins/Animation/LiveLink/Source/LiveLinkEditor/Private/AnimGraphNode_LiveLinkArchiveComponentPose.h
+++ b/Engine/Plugins/Animation/LiveLink/Source/LiveLinkEditor/Private/AnimGraphNode_LiveLinkArchiveComponentPose.h
@@ -1,4 +1,4 @@
-// Copyright 1998-2018 Epic Games, Inc. All Rights Reserved.
+// Copyright 1998-2019 Epic Games, Inc. All Rights Reserved.
 
 #pragma once
 
diff --git a/Engine/Source/Runtime/LiveLinkInterface/Private/LiveLinkCurveRemapSettings.cpp b/Engine/Source/Runtime/LiveLinkInterface/Private/LiveLinkCurveRemapSettings.cpp
--- a/Engine/Source/Runtime/LiveLinkInterface/Private/LiveLinkCurveRemapSettings.cpp
+++ b/Engine/Source/Runtime/LiveLinkInterface/Private/LiveLinkCurveRemapSettings.cpp
@@ -1,4 +1,4 @@
-// Copyright 1998-2018 Epic Games, Inc. All Rights Reserved.
+// Copyright 1998-2019 Epic Games, Inc. All Rights Reserved.
 
 #include "LiveLinkCurveRemapSettings.h" 
 #include "Misc/ConfigCacheIni.h"
diff --git a/Engine/Source/Runtime/LiveLinkInterface/Public/LiveLinkCurveRemapSettings.h b/Engine/Source/Runtime/LiveLinkInterface/Public/LiveLinkCurveRemapSettings.h
--- a/Engine/Source/Runtime/LiveLinkInterface/Public/LiveLinkCurveRemapSettings.h
+++ b/Engine/Source/Runtime/LiveLinkInterface/Public/LiveLinkCurveRemapSettings.h
@@ -1,4 +1,4 @@
-// Copyright 1998-2018 Epic Games, Inc. All Rights Reserved.
+// Copyright 1998-2019 Epic Games, Inc. All Rights Reserved.
 
 #pragma once
 
