Fix for non-interpolated Link Link Curves. Just need to send raw data like we do for transforms.

diff --git a/Engine/Plugins/Animation/LiveLink/Source/LiveLink/Private/MovieScene/MovieSceneLiveLinkSectionTemplate.cpp b/Engine/Plugins/Animation/LiveLink/Source/LiveLink/Private/MovieScene/MovieSceneLiveLinkSectionTemplate.cpp
--- a/Engine/Plugins/Animation/LiveLink/Source/LiveLink/Private/MovieScene/MovieSceneLiveLinkSectionTemplate.cpp
+++ b/Engine/Plugins/Animation/LiveLink/Source/LiveLink/Private/MovieScene/MovieSceneLiveLinkSectionTemplate.cpp
@@ -75,7 +75,7 @@ bool FMovieSceneLiveLinkTemplateData::GetLiveLinkFrameArray(const FFrameTime& Fr
 	{
 		FFrameTime FrameRangeEnd = LowerBound > UpperBound ? LowerBound : UpperBound;
 		FFrameTime FrameRangeStart = LowerBound > UpperBound ? UpperBound : LowerBound;
-		if (TemplateToPush.Transforms.Num() > 0)
+		if (TemplateToPush.Transforms.Num() > 0 || TemplateToPush.CurveElements.Num() > 0)
 		{
 			//assume times will be the same for all channels...
 			TArrayView<const FFrameNumber> Times = FloatChannels[0].GetTimes();
@@ -123,6 +123,7 @@ bool FMovieSceneLiveLinkTemplateData::GetLiveLinkFrameArray(const FFrameTime& Fr
 						bSendInterpolated = false;
 						FLiveLinkFrameData FrameData;
 						FrameData.Transforms.Reset(TemplateToPush.Transforms.Num());
+						FrameData.CurveElements.Reset(TemplateToPush.CurveElements.Num());
 						FrameData.WorldTime = ConvertFrameTimeToLiveLinkWorldTime(Times[Index], FrameRate, FrameTime, WorldTime);
 						if (TimeCodeFrameTime.IsSet())
 						{
@@ -149,17 +150,16 @@ bool FMovieSceneLiveLinkTemplateData::GetLiveLinkFrameArray(const FFrameTime& Fr
 							FrameData.Transforms.Add(Transform);
 						}
 
-						//handle curves also...
+						for (FLiveLinkCurveElement CurveElement : TemplateToPush.CurveElements)
+						{
+							CurveElement.CurveValue = FloatChannels[ChannelIndex++].GetValues()[Index].Value;
+							FrameData.CurveElements.Add(CurveElement);
+						}
 						LiveLinkFrameDataArray.Add(FrameData);
-
 					}
 				}
 			}
 		}
-		else
-		{
-			int CurveStart = 9 * TemplateToPush.Transforms.Num();
-		}
 	}
 	if (bSendInterpolated)
 	{
