Fix LiveLink style path for one image.

diff --git a/Engine/Plugins/Animation/LiveLink/Source/LiveLinkEditor/Private/LiveLinkEditorModule.cpp b/Engine/Plugins/Animation/LiveLink/Source/LiveLinkEditor/Private/LiveLinkEditorModule.cpp
--- a/Engine/Plugins/Animation/LiveLink/Source/LiveLinkEditor/Private/LiveLinkEditorModule.cpp
+++ b/Engine/Plugins/Animation/LiveLink/Source/LiveLinkEditor/Private/LiveLinkEditorModule.cpp
@@ -97,9 +97,9 @@ public:
 		StyleSet->Set("LiveLinkClient.Common.RemoveAllSources", new IMAGE_PLUGIN_BRUSH(TEXT("icon_RemoveSource_40x"), Icon40x40));
 
 		FButtonStyle Button = FButtonStyle()
-			.SetNormal(FSlateBoxBrush(StyleSet->RootToContentDir("ButtonHoverHint.png"), FMargin(4 / 16.0f), FLinearColor(1, 1, 1, 0.15f)))
-			.SetHovered(FSlateBoxBrush(StyleSet->RootToContentDir("ButtonHoverHint.png"), FMargin(4 / 16.0f), FLinearColor(1, 1, 1, 0.25f)))
-			.SetPressed(FSlateBoxBrush(StyleSet->RootToContentDir("ButtonHoverHint.png"), FMargin(4 / 16.0f), FLinearColor(1, 1, 1, 0.30f)))
+			.SetNormal(FSlateBoxBrush(StyleSet->RootToContentDir("Common/ButtonHoverHint.png"), FMargin(4 / 16.0f), FLinearColor(1, 1, 1, 0.15f)))
+			.SetHovered(FSlateBoxBrush(StyleSet->RootToContentDir("Common/ButtonHoverHint.png"), FMargin(4 / 16.0f), FLinearColor(1, 1, 1, 0.25f)))
+			.SetPressed(FSlateBoxBrush(StyleSet->RootToContentDir("Common/ButtonHoverHint.png"), FMargin(4 / 16.0f), FLinearColor(1, 1, 1, 0.30f)))
 			.SetNormalPadding(FMargin(0, 0, 0, 1))
 			.SetPressedPadding(FMargin(0, 1, 0, 0));
 		FComboButtonStyle ComboButton = FComboButtonStyle()
