Shutdown LiveLink sources when the engine close.

fix CIS missing include introduce with CL 6907128
#jra UE-75946

Fix when converting old live link sequence data. Was always creating an animation role data but sequence with no bones should be made with the basic role only.

Added AxisSwitch PreProcessor for LiveLink

Fix for empty curves in old live link sequence data

Fix LiveLink style path for one image.

Fix for crash when opening Mobu. TargetPlatformManager must be initialized early on in game thread when using UDPMessaging

Add deprecation path for LiveLinkBlueprintStructs.

Update EvaluateLiveLink blueprint node. Fix output and use the role as a separate output. Remove LiveLink evaluate by world time and use World Time Offset, since we are using double for time and float is not precise enough.
Add evaluate LiveLink at sene and world time to Blueprint.
Expose LiveLink "enabled" functions

Missing include introduce with last integration CL 6939912

- Fix for crash when evaluating a basic object from mobu with the skeleton hierarchy (property names were not matching property values)

Add default value for LiveLinkFrameInterpolator

- Changed behavior when static data received with different role. Previous subject is removed and a new one is created
- Fix for crash when pushing subject frame without static data received.

Fix for LiveLinkComponent posteditchange not listening for inner struct member changed

Add comments and deprecation path to LiveLinkSourceFactor

- Updating LiveLink AnimPose and PreviewController to use new picker

Fixing CIS : Missing include and deprecation warning disable

diff --git a/Engine/Plugins/Animation/LiveLink/Source/LiveLink/Private/AnimNode_LiveLinkPose.cpp b/Engine/Plugins/Animation/LiveLink/Source/LiveLink/Private/AnimNode_LiveLinkPose.cpp
--- a/Engine/Plugins/Animation/LiveLink/Source/LiveLink/Private/AnimNode_LiveLinkPose.cpp
+++ b/Engine/Plugins/Animation/LiveLink/Source/LiveLink/Private/AnimNode_LiveLinkPose.cpp
@@ -6,6 +6,7 @@
 #include "Features/IModularFeatures.h"
 
 #include "Animation/AnimInstanceProxy.h"
+#include "LiveLinkCustomVersion.h"
 #include "LiveLinkRemapAsset.h"
 #include "Roles/LiveLinkAnimationRole.h"
 #include "Roles/LiveLinkAnimationTypes.h"
@@ -14,6 +15,7 @@ FAnimNode_LiveLinkPose::FAnimNode_LiveLinkPose()
 	: RetargetAsset(ULiveLinkRemapAsset::StaticClass())
 	, CurrentRetargetAsset(nullptr)
 	, LiveLinkClient(nullptr)
+	, CachedDeltaTime(0.0f)
 {
 }
 
@@ -63,7 +65,7 @@ void FAnimNode_LiveLinkPose::Evaluate_AnyThread(FPoseContext& Output)
 	}
 
 	FLiveLinkSubjectFrameData AnimationFrame;
-	if(LiveLinkClient->EvaluateFrame_AnyThread(SubjectName, ULiveLinkAnimationRole::StaticClass(), AnimationFrame))
+	if(LiveLinkClient->EvaluateFrame_AnyThread(LiveLinkSubjectName, ULiveLinkAnimationRole::StaticClass(), AnimationFrame))
 	{
 		FLiveLinkSkeletonStaticData* SkeletonData = AnimationFrame.StaticData.Cast<FLiveLinkSkeletonStaticData>();
 		FLiveLinkAnimationFrameData* FrameData = AnimationFrame.FrameData.Cast<FLiveLinkAnimationFrameData>();
@@ -75,7 +77,6 @@ void FAnimNode_LiveLinkPose::Evaluate_AnyThread(FPoseContext& Output)
 	}
 }
 
-
 void FAnimNode_LiveLinkPose::OnLiveLinkClientRegistered(const FName& Type, class IModularFeature* ModularFeature)
 {
 	if (Type == ILiveLinkClient::ModularFeatureName && !LiveLinkClient)
@@ -100,8 +101,36 @@ void FAnimNode_LiveLinkPose::CacheBones_AnyThread(const FAnimationCacheBonesCont
 
 void FAnimNode_LiveLinkPose::GatherDebugData(FNodeDebugData& DebugData)
 {
-	FString DebugLine = FString::Printf(TEXT("LiveLink - SubjectName: %s"), *SubjectName.ToString());
+	FString DebugLine = FString::Printf(TEXT("LiveLink - SubjectName: %s"), *LiveLinkSubjectName.ToString());
 
 	DebugData.AddDebugItem(DebugLine);
 	InputPose.GatherDebugData(DebugData);
-}
\ No newline at end of file
+}
+
+bool FAnimNode_LiveLinkPose::Serialize(FArchive& Ar)
+{
+	Ar.UsingCustomVersion(FLiveLinkCustomVersion::GUID);
+	
+	UScriptStruct* Struct = FAnimNode_LiveLinkPose::StaticStruct();
+	
+	if (Ar.IsLoading() || Ar.IsSaving())
+	{
+		Struct->SerializeTaggedProperties(Ar, (uint8*)this, Struct, nullptr);
+	}
+
+#if WITH_EDITORONLY_DATA
+	//Take old data and put it in new data structure
+	if (Ar.IsLoading())
+	{
+		if (Ar.CustomVer(FLiveLinkCustomVersion::GUID) < FLiveLinkCustomVersion::NewLiveLinkRoleSystem)
+		{
+			PRAGMA_DISABLE_DEPRECATION_WARNINGS
+			LiveLinkSubjectName.Name = SubjectName_DEPRECATED;
+			PRAGMA_ENABLE_DEPRECATION_WARNINGS
+		}
+	}
+#endif
+
+	return true;
+}
+
diff --git a/Engine/Plugins/Animation/LiveLink/Source/LiveLink/Private/LiveLinkBlueprintLibrary.cpp b/Engine/Plugins/Animation/LiveLink/Source/LiveLink/Private/LiveLinkBlueprintLibrary.cpp
--- a/Engine/Plugins/Animation/LiveLink/Source/LiveLink/Private/LiveLinkBlueprintLibrary.cpp
+++ b/Engine/Plugins/Animation/LiveLink/Source/LiveLink/Private/LiveLinkBlueprintLibrary.cpp
@@ -3,6 +3,12 @@
 #include "LiveLinkBlueprintLibrary.h"
 #include "Features/IModularFeatures.h"
 #include "ILiveLinkClient.h"
+#include "Misc/App.h"
+
+bool ULiveLinkBlueprintLibrary::GetPropertyValue(UPARAM(ref) FLiveLinkBasicBlueprintData& BasicData, FName PropertyName, float& Value)
+{
+	return BasicData.StaticData.FindPropertyValue(BasicData.FrameData, PropertyName, Value);
+};
 
 void ULiveLinkBlueprintLibrary::GetCurves(UPARAM(ref) FSubjectFrameHandle& SubjectFrameHandle, TMap<FName, float>& Curves)
 {
@@ -39,6 +45,18 @@ void ULiveLinkBlueprintLibrary::GetMetadata(UPARAM(ref) FSubjectFrameHandle& Sub
 	SubjectFrameHandle.GetSubjectMetadata(Metadata);
 };
 
+void ULiveLinkBlueprintLibrary::GetBasicData(UPARAM(ref) FSubjectFrameHandle& SubjectFrameHandle, FLiveLinkBasicBlueprintData& BasicBlueprintData)
+{
+	if (FLiveLinkSkeletonStaticData* StaticData = SubjectFrameHandle.GetSourceSkeletonStaticData())
+	{
+		BasicBlueprintData.StaticData = *StaticData;
+	}
+	if (FLiveLinkAnimationFrameData* FrameData = SubjectFrameHandle.GetSourceAnimationFrameData())
+	{
+		BasicBlueprintData.FrameData = *FrameData;
+	}
+};
+
 void ULiveLinkBlueprintLibrary::TransformName(UPARAM(ref) FLiveLinkTransform& LiveLinkTransform, FName& Name)
 {
 	return LiveLinkTransform.GetName(Name);
@@ -131,6 +149,89 @@ FText ULiveLinkBlueprintLibrary::GetSourceMachineName(UPARAM(ref) FLiveLinkSourc
 	}
 }
 
+TArray<FLiveLinkSubjectName> ULiveLinkBlueprintLibrary::GetLiveLinkEnabledSubjectNames(bool bIncludeVirtualSubject)
+{
+	TArray<FLiveLinkSubjectName> Result;
+	IModularFeatures& ModularFeatures = IModularFeatures::Get();
+	if (ModularFeatures.IsModularFeatureAvailable(ILiveLinkClient::ModularFeatureName))
+	{
+		ILiveLinkClient& LiveLinkClient = ModularFeatures.GetModularFeature<ILiveLinkClient>(ILiveLinkClient::ModularFeatureName);
+
+		TArray<FLiveLinkSubjectKey> SubjectKeys = LiveLinkClient.GetSubjects(false, bIncludeVirtualSubject);
+		Result.Reset(SubjectKeys.Num());
+		for (const FLiveLinkSubjectKey& SubjectKey : SubjectKeys)
+		{
+			Result.Add(SubjectKey.SubjectName);
+		}
+	}
+	return Result;
+}
+
+TArray<FLiveLinkSubjectKey> ULiveLinkBlueprintLibrary::GetLiveLinkSubjects(bool bIncludeDisabledSubject, bool bIncludeVirtualSubject)
+{
+	IModularFeatures& ModularFeatures = IModularFeatures::Get();
+	if (ModularFeatures.IsModularFeatureAvailable(ILiveLinkClient::ModularFeatureName))
+	{
+		ILiveLinkClient& LiveLinkClient = ModularFeatures.GetModularFeature<ILiveLinkClient>(ILiveLinkClient::ModularFeatureName);
+
+		return LiveLinkClient.GetSubjects(bIncludeDisabledSubject, bIncludeVirtualSubject);
+	}
+	return TArray<FLiveLinkSubjectKey>();
+}
+
+bool ULiveLinkBlueprintLibrary::IsSpecificLiveLinkSubjectEnabled(const FLiveLinkSubjectKey SubjectKey)
+{
+	IModularFeatures& ModularFeatures = IModularFeatures::Get();
+	if (ModularFeatures.IsModularFeatureAvailable(ILiveLinkClient::ModularFeatureName))
+	{
+		ILiveLinkClient& LiveLinkClient = ModularFeatures.GetModularFeature<ILiveLinkClient>(ILiveLinkClient::ModularFeatureName);
+		return LiveLinkClient.IsSubjectEnabled(SubjectKey);
+	}
+	return false;
+}
+
+bool ULiveLinkBlueprintLibrary::IsLiveLinkSubjectEnabled(const FLiveLinkSubjectName SubjectName)
+{
+	IModularFeatures& ModularFeatures = IModularFeatures::Get();
+	if (ModularFeatures.IsModularFeatureAvailable(ILiveLinkClient::ModularFeatureName))
+	{
+		ILiveLinkClient& LiveLinkClient = ModularFeatures.GetModularFeature<ILiveLinkClient>(ILiveLinkClient::ModularFeatureName);
+		return LiveLinkClient.IsSubjectEnabled(SubjectName);
+	}
+	return false;
+}
+void ULiveLinkBlueprintLibrary::SetLiveLinkSubjectEnabled(const FLiveLinkSubjectKey SubjectKey, bool bEnabled)
+{
+	IModularFeatures& ModularFeatures = IModularFeatures::Get();
+	if (ModularFeatures.IsModularFeatureAvailable(ILiveLinkClient::ModularFeatureName))
+	{
+		ILiveLinkClient& LiveLinkClient = ModularFeatures.GetModularFeature<ILiveLinkClient>(ILiveLinkClient::ModularFeatureName);
+		LiveLinkClient.SetSubjectEnabled(SubjectKey, bEnabled);
+	}
+}
+
+TSubclassOf<ULiveLinkRole> ULiveLinkBlueprintLibrary::GetSpecificLiveLinkSubjectRole(const FLiveLinkSubjectKey SubjectKey)
+{
+	IModularFeatures& ModularFeatures = IModularFeatures::Get();
+	if (ModularFeatures.IsModularFeatureAvailable(ILiveLinkClient::ModularFeatureName))
+	{
+		ILiveLinkClient& LiveLinkClient = ModularFeatures.GetModularFeature<ILiveLinkClient>(ILiveLinkClient::ModularFeatureName);
+		return LiveLinkClient.GetSubjectRole(SubjectKey);
+	}
+	return TSubclassOf<ULiveLinkRole>();
+}
+
+TSubclassOf<ULiveLinkRole> ULiveLinkBlueprintLibrary::GetLiveLinkSubjectRole(const FLiveLinkSubjectName SubjectName)
+{
+	IModularFeatures& ModularFeatures = IModularFeatures::Get();
+	if (ModularFeatures.IsModularFeatureAvailable(ILiveLinkClient::ModularFeatureName))
+	{
+		ILiveLinkClient& LiveLinkClient = ModularFeatures.GetModularFeature<ILiveLinkClient>(ILiveLinkClient::ModularFeatureName);
+		return LiveLinkClient.GetSubjectRole(SubjectName);
+	}
+	return TSubclassOf<ULiveLinkRole>();
+}
+
 bool ULiveLinkBlueprintLibrary::EvaluateLiveLinkFrame(FLiveLinkSubjectRepresentation InSubjectRepresentation, FLiveLinkBaseBlueprintData& OutBlueprintData)
 {
 	// We should never hit this!  stubs to avoid NoExport on the class.
@@ -138,21 +239,183 @@ bool ULiveLinkBlueprintLibrary::EvaluateLiveLinkFrame(FLiveLinkSubjectRepresenta
 	return false;
 }
 
-bool ULiveLinkBlueprintLibrary::Generic_EvaluateLiveLinkFrame(FLiveLinkSubjectRepresentation InSubjectRepresentation, FLiveLinkBlueprintDataStruct& OutBlueprintData)
+bool ULiveLinkBlueprintLibrary::EvaluateLiveLinkFrameWithSpecificRole(FLiveLinkSubjectName SubjectName, TSubclassOf<ULiveLinkRole> Role, FLiveLinkBaseBlueprintData& OutBlueprintData)
 {
-	if (InSubjectRepresentation.Role != nullptr)
+	// We should never hit this!  stubs to avoid NoExport on the class.
+	check(0);
+	return false;
+}
+
+bool ULiveLinkBlueprintLibrary::EvaluateLiveLinkFrameAtWorldTimeOffset(FLiveLinkSubjectName SubjectName, TSubclassOf<ULiveLinkRole> Role, float WorldTimeOffset, FLiveLinkBaseBlueprintData& OutBlueprintData)
+{
+	// We should never hit this!  stubs to avoid NoExport on the class.
+	check(0);
+	return false;
+}
+
+bool ULiveLinkBlueprintLibrary::EvaluateLiveLinkFrameAtSceneTime(FLiveLinkSubjectName SubjectName, TSubclassOf<ULiveLinkRole> Role, FTimecode SceneTime, FLiveLinkBaseBlueprintData& OutBlueprintData)
+{
+	// We should never hit this!  stubs to avoid NoExport on the class.
+	check(0);
+	return false;
+}
+
+namespace EvaluateLiveLinkFrame
+{
+	enum class EEvaluateType
 	{
-		IModularFeatures& ModularFeatures = IModularFeatures::Get();
-		if (ModularFeatures.IsModularFeatureAvailable(ILiveLinkClient::ModularFeatureName))
+		Snapshot,
+		WorldTime,
+		SceneTime,
+	};
+
+	bool Generic_Evaluate(const ULiveLinkBlueprintLibrary* InSelf, const struct FFrame& InStack, UStructProperty* InBlueprintDataStructProp
+		, FLiveLinkSubjectRepresentation InSubjectRepresentation, EEvaluateType InEvaluationType, double InWorldTime, FTimecode& InSceneTime, void* OutBlueprintDataPtr)
+	{
+		bool bSuccess = false;
+
+		if (!InSubjectRepresentation.Role || InSubjectRepresentation.Subject.IsNone())
 		{
-			ILiveLinkClient* LiveLinkClient = &ModularFeatures.GetModularFeature<ILiveLinkClient>(ILiveLinkClient::ModularFeatureName);
-			
-			FLiveLinkSubjectFrameData FrameData;
-			if (LiveLinkClient->EvaluateFrame_AnyThread(InSubjectRepresentation.Subject, InSubjectRepresentation.Role, FrameData))
+			FBlueprintExceptionInfo ExceptionInfo(
+				EBlueprintExceptionType::AccessViolation,
+				NSLOCTEXT("EvaluateLiveLinkFrame", "MissingRoleInput", "Failed to resolve the subject. Be sure the subject name and role are valid.")
+			);
+			FBlueprintCoreDelegates::ThrowScriptException(InSelf, InStack, ExceptionInfo);
+		}
+		else if (InBlueprintDataStructProp && OutBlueprintDataPtr)
+		{
+			if (ULiveLinkRole* LiveLinkRole = Cast<ULiveLinkRole>(InSubjectRepresentation.Role->GetDefaultObject()))
 			{
-				return InSubjectRepresentation.Role.GetDefaultObject()->InitializeBlueprintData(FrameData, OutBlueprintData);
+				UScriptStruct* BlueprintDataType = InBlueprintDataStructProp->Struct;
+				const UScriptStruct* RoleBlueprintDataType = LiveLinkRole->GetBlueprintDataStruct();
+
+				const bool bBlueprintDataCompatible = (BlueprintDataType == RoleBlueprintDataType) ||
+					(BlueprintDataType->IsChildOf(RoleBlueprintDataType) && FStructUtils::TheSameLayout(BlueprintDataType, RoleBlueprintDataType));
+				if (bBlueprintDataCompatible)
+				{
+					P_NATIVE_BEGIN;
+					//Create the struct holder and make it point to the output data
+					FLiveLinkBlueprintDataStruct BlueprintDataWrapper(BlueprintDataType, StaticCast<FLiveLinkBaseBlueprintData*>(OutBlueprintDataPtr));
+					IModularFeatures& ModularFeatures = IModularFeatures::Get();
+					if (ModularFeatures.IsModularFeatureAvailable(ILiveLinkClient::ModularFeatureName))
+					{
+						ILiveLinkClient& LiveLinkClient = ModularFeatures.GetModularFeature<ILiveLinkClient>(ILiveLinkClient::ModularFeatureName);
+						FLiveLinkSubjectFrameData FrameData;
+						switch(InEvaluationType)
+						{
+						case EEvaluateType::WorldTime:
+							bSuccess = LiveLinkClient.EvaluateFrameAtWorldTime_AnyThread(InSubjectRepresentation.Subject, InWorldTime, InSubjectRepresentation.Role, FrameData);
+							break;
+						case EEvaluateType::SceneTime:
+							bSuccess = LiveLinkClient.EvaluateFrameAtSceneTime_AnyThread(InSubjectRepresentation.Subject, InSceneTime, InSubjectRepresentation.Role, FrameData);
+							break;
+						case EEvaluateType::Snapshot:
+						default:
+							bSuccess = LiveLinkClient.EvaluateFrame_AnyThread(InSubjectRepresentation.Subject, InSubjectRepresentation.Role, FrameData);
+							break;
+						}
+
+						if (bSuccess)
+						{
+
+							bSuccess = InSubjectRepresentation.Role.GetDefaultObject()->InitializeBlueprintData(FrameData, BlueprintDataWrapper);
+						}
+					}
+					P_NATIVE_END;
+				}
+				else
+				{
+					FBlueprintExceptionInfo ExceptionInfo(
+						EBlueprintExceptionType::AccessViolation,
+						NSLOCTEXT("EvaluateLiveLinkFrame", "IncompatibleProperty", "Incompatible output blueprint data; the role blueprint's data type is not the same as the return type.")
+					);
+					FBlueprintCoreDelegates::ThrowScriptException(InSelf, InStack, ExceptionInfo);
+				}
 			}
 		}
+		else
+		{
+			FBlueprintExceptionInfo ExceptionInfo(
+				EBlueprintExceptionType::AccessViolation,
+				NSLOCTEXT("EvaluateLiveLinkFrame", "MissingOutputProperty", "Failed to resolve the output parameter for EvaluateLiveLinkFrame.")
+			);
+			FBlueprintCoreDelegates::ThrowScriptException(InSelf, InStack, ExceptionInfo);
+		}
+
+		return bSuccess;
 	}
-	return false;
+}
+
+DEFINE_FUNCTION(ULiveLinkBlueprintLibrary::execEvaluateLiveLinkFrame)
+{
+	P_GET_STRUCT(FLiveLinkSubjectRepresentation, SubjectRepresentation);
+
+	Stack.MostRecentPropertyAddress = nullptr;
+	Stack.StepCompiledIn<UStructProperty>(NULL);
+	void* OutBlueprintDataPtr = Stack.MostRecentPropertyAddress;
+	UStructProperty* BlueprintDataStructProp = Cast<UStructProperty>(Stack.MostRecentProperty);
+	P_FINISH;
+
+	double WorldTime = 0.0;
+	FTimecode SceneTime;
+	bool bSuccess = EvaluateLiveLinkFrame::Generic_Evaluate(P_THIS, Stack, BlueprintDataStructProp, SubjectRepresentation, EvaluateLiveLinkFrame::EEvaluateType::Snapshot, WorldTime, SceneTime, OutBlueprintDataPtr);
+	*(bool*)RESULT_PARAM = bSuccess;
+}
+
+DEFINE_FUNCTION(ULiveLinkBlueprintLibrary::execEvaluateLiveLinkFrameWithSpecificRole)
+{
+	P_GET_STRUCT(FLiveLinkSubjectName, SubjectName);
+	UClass* RoleClass;
+	Stack.StepCompiledIn<UClassProperty>(&RoleClass);
+
+	Stack.MostRecentPropertyAddress = nullptr;
+	Stack.StepCompiledIn<UStructProperty>(NULL);
+	void* OutBlueprintDataPtr = Stack.MostRecentPropertyAddress;
+	UStructProperty* BlueprintDataStructProp = Cast<UStructProperty>(Stack.MostRecentProperty);
+	P_FINISH;
+
+
+	double WorldTime = 0.0;
+	FTimecode SceneTime;
+	bool bSuccess = EvaluateLiveLinkFrame::Generic_Evaluate(P_THIS, Stack, BlueprintDataStructProp, {SubjectName, RoleClass}, EvaluateLiveLinkFrame::EEvaluateType::Snapshot, WorldTime, SceneTime, OutBlueprintDataPtr);
+	*(bool*)RESULT_PARAM = bSuccess;
+}
+
+DEFINE_FUNCTION(ULiveLinkBlueprintLibrary::execEvaluateLiveLinkFrameAtWorldTimeOffset)
+{
+
+	P_GET_STRUCT(FLiveLinkSubjectName, SubjectName);
+	UClass* RoleClass;
+	Stack.StepCompiledIn<UClassProperty>(&RoleClass);
+	float WorldTimeOffset = 0.f;
+	Stack.StepCompiledIn<UFloatProperty>(&WorldTimeOffset);
+
+	Stack.MostRecentPropertyAddress = nullptr;
+	Stack.StepCompiledIn<UStructProperty>(NULL);
+	void* OutBlueprintDataPtr = Stack.MostRecentPropertyAddress;
+	UStructProperty* BlueprintDataStructProp = Cast<UStructProperty>(Stack.MostRecentProperty);
+	P_FINISH;
+
+	double WorldTime = FApp::GetCurrentTime() + WorldTimeOffset;
+	FTimecode SceneTime;
+	bool bSuccess = EvaluateLiveLinkFrame::Generic_Evaluate(P_THIS, Stack, BlueprintDataStructProp, {SubjectName, RoleClass}, EvaluateLiveLinkFrame::EEvaluateType::WorldTime, WorldTime, SceneTime, OutBlueprintDataPtr);
+	*(bool*)RESULT_PARAM = bSuccess;
+}
+
+DEFINE_FUNCTION(ULiveLinkBlueprintLibrary::execEvaluateLiveLinkFrameAtSceneTime)
+{
+	P_GET_STRUCT(FLiveLinkSubjectName, SubjectName);
+	UClass* RoleClass;
+	Stack.StepCompiledIn<UClassProperty>(&RoleClass);
+	P_GET_STRUCT(FTimecode, SceneTime);
+
+	Stack.MostRecentPropertyAddress = nullptr;
+	Stack.StepCompiledIn<UStructProperty>(NULL);
+	void* OutBlueprintDataPtr = Stack.MostRecentPropertyAddress;
+	UStructProperty* BlueprintDataStructProp = Cast<UStructProperty>(Stack.MostRecentProperty);
+	P_FINISH;
+
+	double WorldTime = 0.0;
+	bool bSuccess = EvaluateLiveLinkFrame::Generic_Evaluate(P_THIS, Stack, BlueprintDataStructProp, {SubjectName, RoleClass}, EvaluateLiveLinkFrame::EEvaluateType::SceneTime, WorldTime, SceneTime, OutBlueprintDataPtr);
+	*(bool*)RESULT_PARAM = bSuccess;
 }
\ No newline at end of file
diff --git a/Engine/Plugins/Animation/LiveLink/Source/LiveLink/Private/LiveLinkClient.cpp b/Engine/Plugins/Animation/LiveLink/Source/LiveLink/Private/LiveLinkClient.cpp
--- a/Engine/Plugins/Animation/LiveLink/Source/LiveLink/Private/LiveLinkClient.cpp
+++ b/Engine/Plugins/Animation/LiveLink/Source/LiveLink/Private/LiveLinkClient.cpp
@@ -15,10 +15,11 @@
 #include "LiveLinkSourceFactory.h"
 #include "LiveLinkSourceSettings.h"
 #include "LiveLinkSubject.h"
+#include "Misc/CoreDelegates.h"
 #include "Roles/LiveLinkAnimationRole.h"
 #include "Roles/LiveLinkAnimationTypes.h"
-#include "Roles/LiveLinkCameraTypes.h"
-#include "Roles/LiveLinkCameraRole.h"
+#include "Roles/LiveLinkBasicTypes.h"
+#include "Roles/LiveLinkBasicTypes.h"
 #include "Stats/Stats.h"
 #include "Stats/Stats2.h"
 #include "TimeSynchronizationSource.h"
@@ -40,15 +41,13 @@ DEFINE_LOG_CATEGORY(LogLiveLink);
 FLiveLinkClient::FLiveLinkClient()
 {
 	Collection = MakeUnique<FLiveLinkSourceCollection>();
+	FCoreDelegates::OnPreExit.AddRaw(this, &FLiveLinkClient::Shutdown);
 }
 
 FLiveLinkClient::~FLiveLinkClient()
 {
-	if (Collection)
-	{
-		FScopeLock Lock(&CollectionAccessCriticalSection);
-		Collection->Shutdown();
-	}
+	FCoreDelegates::OnPreExit.RemoveAll(this);
+	Shutdown();
 }
 
 void FLiveLinkClient::Tick(float DeltaTime)
@@ -134,6 +133,15 @@ void FLiveLinkClient::BuildThisTicksSubjectSnapshot()
 	}
 }
 
+void FLiveLinkClient::Shutdown()
+{
+	if (Collection)
+	{
+		FScopeLock Lock(&CollectionAccessCriticalSection);
+		Collection->Shutdown();
+	}
+}
+
 FGuid FLiveLinkClient::AddSource(TSharedPtr<ILiveLinkSource> InSource)
 {
 	check(Collection);
@@ -335,32 +343,57 @@ void FLiveLinkClient::PushSubjectStaticData_Internal(FPendingSubjectStatic&& Sub
 
 		if (LiveLinkSubject->GetRole() != SubjectStaticData.Role)
 		{
-			UE_LOG(LogLiveLink, Warning, TEXT("Subject '%s' of role '%s' is trying to change its role to '%s'."), *SubjectStaticData.SubjectKey.SubjectName.ToString(), *LiveLinkSubject->GetRole().GetDefaultObject()->GetDisplayName().ToString(), *SubjectStaticData.Role.GetDefaultObject()->GetDisplayName().ToString());
-			return;
+			UE_LOG(LogLiveLink, Warning, TEXT("Subject '%s' of role '%s' is changing its role to '%s'. Current subject will be removed and a new one will be created"), *SubjectStaticData.SubjectKey.SubjectName.ToString(), *LiveLinkSubject->GetRole().GetDefaultObject()->GetDisplayName().ToString(), *SubjectStaticData.Role.GetDefaultObject()->GetDisplayName().ToString());
+			
+			Collection->RemoveSubject(SubjectStaticData.SubjectKey);
+			LiveLinkSubject = nullptr;
 		}
 	}
-	else
+	
+	if(LiveLinkSubject == nullptr)
 	{
 		const FLiveLinkRoleProjectSetting DefaultSetting = GetDefault<ULiveLinkSettings>()->GetDefaultSettingForRole(SubjectStaticData.Role.Get());
 
-		UClass* SettingClass = DefaultSetting.SettingClass.Get();
-		if (SettingClass == nullptr)
+		// Setting class should always be valid
+		ULiveLinkSubjectSettings* SubjectSettings = nullptr;
 		{
-			SettingClass = ULiveLinkSubjectSettings::StaticClass();
-		}
+			UClass* SettingClass = DefaultSetting.SettingClass.Get();
+			if (SettingClass == nullptr)
+			{
+				SettingClass = ULiveLinkSubjectSettings::StaticClass();
+			}
 
-		ULiveLinkSubjectSettings* SubjectSettings = NewObject<ULiveLinkSubjectSettings>(GetTransientPackage(), SettingClass);
+			SubjectSettings = NewObject<ULiveLinkSubjectSettings>(GetTransientPackage(), SettingClass);
+			SubjectSettings->Role = SubjectStaticData.Role;
 
-		SubjectSettings->Role = SubjectStaticData.Role;
-		if (DefaultSetting.FrameInterpolationProcessor != nullptr)
-		{
-			SubjectSettings->InterpolationProcessor = NewObject<ULiveLinkFrameInterpolationProcessor>(GetTransientPackage(), DefaultSetting.FrameInterpolationProcessor.Get());
-		}
-		for (TSubclassOf<ULiveLinkFramePreProcessor> PreProcessor : DefaultSetting.FramePreProcessors)
-		{
-			if (PreProcessor != nullptr)
+			UClass* FrameInterpolationProcessorClass = DefaultSetting.FrameInterpolationProcessor.Get();
+			if (FrameInterpolationProcessorClass != nullptr)
+			{
+				UClass* InterpolationRole = FrameInterpolationProcessorClass->GetDefaultObject<ULiveLinkFrameInterpolationProcessor>()->GetRole();
+				if (SubjectStaticData.Role->IsChildOf(InterpolationRole))
+				{
+					SubjectSettings->InterpolationProcessor = NewObject<ULiveLinkFrameInterpolationProcessor>(GetTransientPackage(), FrameInterpolationProcessorClass);
+				}
+				else
+				{
+					UE_LOG(LogLiveLink, Warning, TEXT("The interpolator '%s' is not valid for the Role '%s'"), *FrameInterpolationProcessorClass->GetName(), *SubjectStaticData.Role->GetName());
+				}
+			}
+
+			for (TSubclassOf<ULiveLinkFramePreProcessor> PreProcessor : DefaultSetting.FramePreProcessors)
 			{
-				SubjectSettings->PreProcessors.Add(NewObject<ULiveLinkFramePreProcessor>(GetTransientPackage(), PreProcessor.Get()));
+				if (PreProcessor != nullptr)
+				{
+					UClass* PreProcessorRole = PreProcessor->GetDefaultObject<ULiveLinkFramePreProcessor>()->GetRole();
+					if (SubjectStaticData.Role->IsChildOf(PreProcessorRole))
+					{
+						SubjectSettings->PreProcessors.Add(NewObject<ULiveLinkFramePreProcessor>(GetTransientPackage(), PreProcessor.Get()));
+					}
+					else
+					{
+						UE_LOG(LogLiveLink, Warning, TEXT("The pre processor '%s' is not valid for the Role '%s'"), *PreProcessor->GetName(), *SubjectStaticData.Role->GetName());
+					}
+				}
 			}
 		}
 
@@ -404,18 +437,23 @@ void FLiveLinkClient::PushSubjectFrameData_Internal(FPendingSubjectFrame&& Subje
 		return;
 	}
 
+	//To add a frame data, we need to find our subject but also have a static data associated to it. 
+	//With presets, the subject could exist but no static data received yet.
 	if (FLiveLinkCollectionSubjectItem* SubjectItem = Collection->FindSubject(SubjectFrameData.SubjectKey))
 	{
 		if (SubjectItem->bEnabled && !SubjectItem->bPendingKill)
 		{
 			if (FLiveLinkSubject* LinkSubject = SubjectItem->GetLiveSubject())
 			{
-				if (const FSubjectFramesReceivedHandles* Handles = SubjectFrameReceivedHandles.Find(SubjectFrameData.SubjectKey.SubjectName))
+				if (LinkSubject->HasStaticData())
 				{
-					Handles->OnFrameDataReceived.Broadcast(SubjectItem->Key, SubjectItem->GetSubject()->GetRole(), SubjectFrameData.FrameData);
-				}
+					if (const FSubjectFramesReceivedHandles* Handles = SubjectFrameReceivedHandles.Find(SubjectFrameData.SubjectKey.SubjectName))
+					{
+						Handles->OnFrameDataReceived.Broadcast(SubjectItem->Key, SubjectItem->GetSubject()->GetRole(), SubjectFrameData.FrameData);
+					}
 
-				LinkSubject->AddFrameData(MoveTemp(SubjectFrameData.FrameData));
+					LinkSubject->AddFrameData(MoveTemp(SubjectFrameData.FrameData));
+				}
 			}
 		}
 	}
diff --git a/Engine/Plugins/Animation/LiveLink/Source/LiveLink/Private/LiveLinkSettings.cpp b/Engine/Plugins/Animation/LiveLink/Source/LiveLink/Private/LiveLinkSettings.cpp
new file mode 100644
--- /dev/null
+++ b/Engine/Plugins/Animation/LiveLink/Source/LiveLink/Private/LiveLinkSettings.cpp
@@ -0,0 +1,39 @@
+// Copyright 1998-2019 Epic Games, Inc. All Rights Reserved.
+
+#include "LiveLinkSettings.h"
+#include "LiveLinkSubjectSettings.h"
+#include "Roles/LiveLinkAnimationRole.h"
+#include "Roles/LiveLinkBasicRole.h"
+
+FLiveLinkRoleProjectSetting::FLiveLinkRoleProjectSetting()
+	: SettingClass(ULiveLinkSubjectSettings::StaticClass())
+	, FrameInterpolationProcessor(ULiveLinkBasicFrameInterpolateProcessor::StaticClass())
+{}
+
+
+ULiveLinkSettings::ULiveLinkSettings()
+	: MessageBusPingRequestFrequency(1.0)
+	, MessageBusHeartbeatFrequency(1.0)
+	, MessageBusHeartbeatTimeout(2.0)
+	, TimeWithoutFrameToBeConsiderAsInvalid(0.5)
+	, ValidColor(FLinearColor::Green)
+	, InvalidColor(FLinearColor::Green)
+{
+	// Add the role animation default settings
+	FLiveLinkRoleProjectSetting BaseAnimationSetting;
+	BaseAnimationSetting.Role = ULiveLinkAnimationRole::StaticClass();
+	BaseAnimationSetting.FrameInterpolationProcessor = ULiveLinkAnimationFrameInterpolateProcessor::StaticClass();
+	DefaultRoleSettings.Add(BaseAnimationSetting);
+}
+
+FLiveLinkRoleProjectSetting ULiveLinkSettings::GetDefaultSettingForRole(TSubclassOf<ULiveLinkRole> Role) const
+{
+	int32 IndexOf = DefaultRoleSettings.IndexOfByPredicate([Role](const FLiveLinkRoleProjectSetting& Other) {return Other.Role == Role; });
+	if (IndexOf != INDEX_NONE)
+	{
+		return DefaultRoleSettings[IndexOf];
+	}
+	FLiveLinkRoleProjectSetting Result;
+	Result.Role = Role;
+	return Result;
+}
diff --git a/Engine/Plugins/Animation/LiveLink/Source/LiveLink/Private/LiveLinkSubject.cpp b/Engine/Plugins/Animation/LiveLink/Source/LiveLink/Private/LiveLinkSubject.cpp
--- a/Engine/Plugins/Animation/LiveLink/Source/LiveLink/Private/LiveLinkSubject.cpp
+++ b/Engine/Plugins/Animation/LiveLink/Source/LiveLink/Private/LiveLinkSubject.cpp
@@ -147,6 +147,11 @@ bool FLiveLinkSubject::EvaluateFrameAtSceneTime(const FTimecode& InSceneTime, TS
 	return bSuccess;
 }
 
+bool FLiveLinkSubject::HasStaticData() const
+{
+	return StaticData.IsValid();
+}
+
 void FLiveLinkSubject::AddFrameData(FLiveLinkFrameDataStruct&& InFrameData)
 {
 	check(IsInGameThread());
diff --git a/Engine/Plugins/Animation/LiveLink/Source/LiveLink/Private/LiveLinkSubject.h b/Engine/Plugins/Animation/LiveLink/Source/LiveLink/Private/LiveLinkSubject.h
--- a/Engine/Plugins/Animation/LiveLink/Source/LiveLink/Private/LiveLinkSubject.h
+++ b/Engine/Plugins/Animation/LiveLink/Source/LiveLink/Private/LiveLinkSubject.h
@@ -65,6 +65,8 @@ public:
 	bool EvaluateFrameAtWorldTime(double InWorldTime, TSubclassOf<ULiveLinkRole> InDesiredRole, FLiveLinkSubjectFrameData& OutFrame);
 	bool EvaluateFrameAtSceneTime(const FTimecode& InSceneTime, TSubclassOf<ULiveLinkRole> InDesiredRole, FLiveLinkSubjectFrameData& OutFrame);
 
+	bool HasStaticData() const;
+
 	// Handling setting a new static data. Create role data if not found in map.
 	void SetStaticData(TSubclassOf<ULiveLinkRole> InRole, FLiveLinkStaticDataStruct&& InStaticData);
 
diff --git a/Engine/Plugins/Animation/LiveLink/Source/LiveLink/Public/AnimNode_LiveLinkPose.h b/Engine/Plugins/Animation/LiveLink/Source/LiveLink/Public/AnimNode_LiveLinkPose.h
--- a/Engine/Plugins/Animation/LiveLink/Source/LiveLink/Public/AnimNode_LiveLinkPose.h
+++ b/Engine/Plugins/Animation/LiveLink/Source/LiveLink/Public/AnimNode_LiveLinkPose.h
@@ -2,24 +2,37 @@
 
 #pragma once
 
-#include "CoreMinimal.h"
 #include "Animation/AnimNodeBase.h"
+
+#include "CoreMinimal.h"
 #include "LiveLinkRetargetAsset.h"
+#include "LiveLinkTypes.h"
+
+#include "AnimNode_LiveLinkPose.generated.h"
+
 
 class ILiveLinkClient;
 
-#include "AnimNode_LiveLinkPose.generated.h"
+
+PRAGMA_DISABLE_DEPRECATION_WARNINGS
 
 USTRUCT(BlueprintInternalUseOnly)
 struct LIVELINK_API FAnimNode_LiveLinkPose : public FAnimNode_Base
 {
-	GENERATED_USTRUCT_BODY()
+	GENERATED_BODY()
 
+public:
 	UPROPERTY(EditAnywhere, BlueprintReadWrite, Category = Input)
 	FPoseLink InputPose;
 
 	UPROPERTY(EditAnywhere, BlueprintReadWrite, Category = SourceData, meta = (PinShownByDefault))
-	FName SubjectName;
+	FLiveLinkSubjectName LiveLinkSubjectName;
+
+#if WITH_EDITORONLY_DATA
+	UE_DEPRECATED(4.23, "FName SubjectName is deprecated. Use the SubjectName of type FLiveLinkSubjectName instead.")
+	UPROPERTY()
+	FName SubjectName_DEPRECATED;
+#endif
 
 	UPROPERTY(EditAnywhere, BlueprintReadWrite, NoClear, Category = Retarget, meta = (NeverAsPin))
 	TSubclassOf<ULiveLinkRetargetAsset> RetargetAsset;
@@ -27,23 +40,37 @@ struct LIVELINK_API FAnimNode_LiveLinkPose : public FAnimNode_Base
 	UPROPERTY(transient)
 	ULiveLinkRetargetAsset* CurrentRetargetAsset;
 
+public:
 	FAnimNode_LiveLinkPose();
 
-	// FAnimNode_Base interface
+	//~ FAnimNode_Base interface
 	virtual void Initialize_AnyThread(const FAnimationInitializeContext& Context) override;
 	virtual void CacheBones_AnyThread(const FAnimationCacheBonesContext & Context) override;
 	virtual void Update_AnyThread(const FAnimationUpdateContext & Context) override;
 	virtual void Evaluate_AnyThread(FPoseContext& Output) override;
 	virtual void GatherDebugData(FNodeDebugData& DebugData) override;
-	// End of FAnimNode_Base interface
+	//~ End of FAnimNode_Base interface
 
 	void OnLiveLinkClientRegistered(const FName& Type, class IModularFeature* ModularFeature);
 	void OnLiveLinkClientUnregistered(const FName& Type, class IModularFeature* ModularFeature);
 
+	bool Serialize(FArchive& Ar);
+
 private:
 
 	ILiveLinkClient* LiveLinkClient;
 
 	// Delta time from update so that it can be passed to retargeter
 	float CachedDeltaTime;
-};
\ No newline at end of file
+};
+
+PRAGMA_ENABLE_DEPRECATION_WARNINGS
+
+
+template<> struct TStructOpsTypeTraits<FAnimNode_LiveLinkPose> : public TStructOpsTypeTraitsBase2<FAnimNode_LiveLinkPose>
+{
+	enum 
+	{ 
+		WithSerializer = true
+	};
+};
diff --git a/Engine/Plugins/Animation/LiveLink/Source/LiveLink/Public/LiveLinkBlueprintLibrary.h b/Engine/Plugins/Animation/LiveLink/Source/LiveLink/Public/LiveLinkBlueprintLibrary.h
--- a/Engine/Plugins/Animation/LiveLink/Source/LiveLink/Public/LiveLinkBlueprintLibrary.h
+++ b/Engine/Plugins/Animation/LiveLink/Source/LiveLink/Public/LiveLinkBlueprintLibrary.h
@@ -7,6 +7,7 @@
 #include "Kismet/BlueprintFunctionLibrary.h"
 #include "LiveLinkRole.h"
 #include "Roles/LiveLinkAnimationBlueprintStructs.h"
+#include "Roles/LiveLinkBasicTypes.h"
 
 #include "LiveLinkBlueprintLibrary.generated.h"
 
@@ -15,6 +16,12 @@ class LIVELINK_API ULiveLinkBlueprintLibrary : public UBlueprintFunctionLibrary
 {
 	GENERATED_BODY()
 
+// FLiveLinkBasicBlueprintData
+
+	// Returns the value of a property stored in the Subject Frame
+	UFUNCTION(BlueprintPure, Category = "LiveLink")
+	static bool GetPropertyValue(UPARAM(ref) FLiveLinkBasicBlueprintData& BasicData, FName PropertyName, float& Value);
+
 // FSubjectFrameHandle 
 
 	// Returns the float curves stored in the Subject Frame as a map
@@ -45,6 +52,10 @@ class LIVELINK_API ULiveLinkBlueprintLibrary : public UBlueprintFunctionLibrary
 	UFUNCTION(BlueprintPure, Category = "LiveLink|Animation")
 	static void GetMetadata(UPARAM(ref) FSubjectFrameHandle& SubjectFrameHandle, FSubjectMetadata& Metadata);
 
+	// Returns the Subject base structure stored in the Subject Frame
+	UFUNCTION(BlueprintPure, Category = "LiveLink|Animation")
+	static void GetBasicData(UPARAM(ref) FSubjectFrameHandle& SubjectFrameHandle, FLiveLinkBasicBlueprintData& BasicBlueprintData);
+
 // FLiveLinkTransform
 
 	// Returns the Name of a given LiveLink Transform
@@ -98,69 +109,57 @@ class LIVELINK_API ULiveLinkBlueprintLibrary : public UBlueprintFunctionLibrary
 	static FText GetSourceMachineName(UPARAM(ref) FLiveLinkSourceHandle& SourceHandle);
 
 public:
+
+	/** Get a list of all enabled subject names */
+	UFUNCTION(BlueprintCallable, Category = "LiveLink")
+	static TArray<FLiveLinkSubjectName> GetLiveLinkEnabledSubjectNames(bool bIncludeVirtualSubject);
+
+	/** Get a list of all subjects */
+	UFUNCTION(BlueprintCallable, Category = "LiveLink")
+	static TArray<FLiveLinkSubjectKey> GetLiveLinkSubjects(bool bIncludeDisabledSubject, bool bIncludeDisal);
+
+	/**
+	 * Whether or not a subject from the specific source is the enabled subject
+	 * Only 1 subject with the same name can be enabled
+	 */
+	UFUNCTION(BlueprintCallable, Category = "LiveLink")
+	bool IsSpecificLiveLinkSubjectEnabled(const FLiveLinkSubjectKey SubjectKey);
+
+	/** Whether or not the client has a subject with this name enabled */
+	UFUNCTION(BlueprintCallable, Category = "LiveLink")
+	bool IsLiveLinkSubjectEnabled(const FLiveLinkSubjectName SubjectName);
+
+	/** Set the subject's from a specific source to enabled, disabling the other in the process, or disabled */
+	UFUNCTION(BlueprintCallable, Category = "LiveLink")
+	void SetLiveLinkSubjectEnabled(const FLiveLinkSubjectKey SubjectKey, bool bEnabled);
+
+	/** Get the role of a subject from a specific source */
+	UFUNCTION(BlueprintCallable, Category = "LiveLink")
+	static TSubclassOf<ULiveLinkRole> GetSpecificLiveLinkSubjectRole(const FLiveLinkSubjectKey SubjectKey);
+
+	/** Get the role of the subject with this name */
+	UFUNCTION(BlueprintCallable, Category = "LiveLink")
+	static TSubclassOf<ULiveLinkRole> GetLiveLinkSubjectRole(const FLiveLinkSubjectName SubjectName);
+
+	UE_DEPRECATED(4.23, "EvaluateLiveLinkFrame with Subject Represention is deprecated, recreate the node.")
+	UFUNCTION(BlueprintCallable, CustomThunk, Category = LiveLink, meta = (DeprecatedFunction, CustomStructureParam = "OutBlueprintData", BlueprintInternalUseOnly = "true", AllowAbstract = "false"))
+	static bool EvaluateLiveLinkFrame(FLiveLinkSubjectRepresentation SubjectRepresentation, FLiveLinkBaseBlueprintData& OutBlueprintData);
+
 	/** Fetches a frame on a subject for a specific role. Output is evaluated based on the role */
+	UFUNCTION(BlueprintCallable, CustomThunk, Category = LiveLink, meta = (DisplayName="EvaluateLiveLinkFrame", CustomStructureParam = "OutBlueprintData", BlueprintInternalUseOnly = "true", AllowAbstract = "false"))
+	static bool EvaluateLiveLinkFrameWithSpecificRole(FLiveLinkSubjectName SubjectName, TSubclassOf<ULiveLinkRole> Role, FLiveLinkBaseBlueprintData& OutBlueprintData);
+
+	/** Fetches a frame on a subject for a specific role at an offset from the application current time. Output is evaluated based on the role */
 	UFUNCTION(BlueprintCallable, CustomThunk, Category = LiveLink, meta = (CustomStructureParam = "OutBlueprintData", BlueprintInternalUseOnly = "true", AllowAbstract = "false"))
-	static bool EvaluateLiveLinkFrame(FLiveLinkSubjectRepresentation SubjectRepresentation, FLiveLinkBaseBlueprintData& OutBlueprintData);
+	static bool EvaluateLiveLinkFrameAtWorldTimeOffset(FLiveLinkSubjectName SubjectName, TSubclassOf<ULiveLinkRole> Role, float WorldTimeOffset, FLiveLinkBaseBlueprintData& OutBlueprintData);
 
-	static bool Generic_EvaluateLiveLinkFrame(FLiveLinkSubjectRepresentation SubjectRepresentation, FLiveLinkBlueprintDataStruct& OutBlueprintData);
-	
-	DECLARE_FUNCTION(execEvaluateLiveLinkFrame)
-	{
-        P_GET_STRUCT(FLiveLinkSubjectRepresentation, SubjectRepresentation);
-
-		Stack.MostRecentPropertyAddress = nullptr;
-		Stack.StepCompiledIn<UStructProperty>(NULL);
-		void* OutBlueprintDataPtr = Stack.MostRecentPropertyAddress;
-		UStructProperty* BlueprintDataStructProp = Cast<UStructProperty>(Stack.MostRecentProperty);
-		
-		P_FINISH;
-		bool bSuccess = false;
-
-		if (!SubjectRepresentation.Role || SubjectRepresentation.Subject.IsNone())
-		{
-			FBlueprintExceptionInfo ExceptionInfo(
-				EBlueprintExceptionType::AccessViolation,
-				NSLOCTEXT("EvaluateLiveLinkFrame", "MissingRoleInput", "Failed to resolve the subject. Be sure the subject name and role are valid.")
-			);
-			FBlueprintCoreDelegates::ThrowScriptException(P_THIS, Stack, ExceptionInfo);
-		}
-		else if (BlueprintDataStructProp && OutBlueprintDataPtr)
-		{
-			if (ULiveLinkRole* LiveLinkRole = Cast<ULiveLinkRole>(SubjectRepresentation.Role->GetDefaultObject()))
-			{
-				UScriptStruct* BlueprintDataType = BlueprintDataStructProp->Struct;
-				const UScriptStruct* RoleBlueprintDataType = LiveLinkRole->GetBlueprintDataStruct();
-
-				const bool bBlueprintDataCompatible = (BlueprintDataType == RoleBlueprintDataType) ||
-					(BlueprintDataType->IsChildOf(RoleBlueprintDataType) && FStructUtils::TheSameLayout(BlueprintDataType, RoleBlueprintDataType));
-				if (bBlueprintDataCompatible)
-				{
-					P_NATIVE_BEGIN;
-					
-					//Create the struct holder and make it point to the output data
-					FLiveLinkBlueprintDataStruct BlueprintDataWrapper(BlueprintDataType, StaticCast<FLiveLinkBaseBlueprintData*>(OutBlueprintDataPtr));
-					bSuccess = Generic_EvaluateLiveLinkFrame(SubjectRepresentation, BlueprintDataWrapper);
-					P_NATIVE_END;
-				}
-				else
-				{
-					FBlueprintExceptionInfo ExceptionInfo(
-						EBlueprintExceptionType::AccessViolation,
-						NSLOCTEXT("EvaluateLiveLinkFrame", "IncompatibleProperty", "Incompatible output blueprint data; the role blueprint's data type is not the same as the return type.")
-					);
-					FBlueprintCoreDelegates::ThrowScriptException(P_THIS, Stack, ExceptionInfo);
-				}
-			}
-		}
-		else
-		{
-			FBlueprintExceptionInfo ExceptionInfo(
-				EBlueprintExceptionType::AccessViolation,
-				NSLOCTEXT("EvaluateLiveLinkFrame", "MissingOutputProperty", "Failed to resolve the output parameter for EvaluateLiveLinkFrame.")
-			);
-			FBlueprintCoreDelegates::ThrowScriptException(P_THIS, Stack, ExceptionInfo);
-		}
-		*(bool*)RESULT_PARAM = bSuccess;
-	}
+	/** Fetches a frame on a subject for a specific role at a specified scene time (timecode). Output is evaluated based on the role */
+	UFUNCTION(BlueprintCallable, CustomThunk, Category = LiveLink, meta = (CustomStructureParam = "OutBlueprintData", BlueprintInternalUseOnly = "true", AllowAbstract = "false"))
+	static bool EvaluateLiveLinkFrameAtSceneTime(FLiveLinkSubjectName SubjectName, TSubclassOf<ULiveLinkRole> Role, FTimecode SceneTime, FLiveLinkBaseBlueprintData& OutBlueprintData);
 
+private:
+	DECLARE_FUNCTION(execEvaluateLiveLinkFrame);
+	DECLARE_FUNCTION(execEvaluateLiveLinkFrameWithSpecificRole);
+	DECLARE_FUNCTION(execEvaluateLiveLinkFrameAtWorldTimeOffset);
+	DECLARE_FUNCTION(execEvaluateLiveLinkFrameAtSceneTime);
 };
\ No newline at end of file
diff --git a/Engine/Plugins/Animation/LiveLink/Source/LiveLink/Public/LiveLinkBlueprintStructs.h b/Engine/Plugins/Animation/LiveLink/Source/LiveLink/Public/LiveLinkBlueprintStructs.h
new file mode 100644
--- /dev/null
+++ b/Engine/Plugins/Animation/LiveLink/Source/LiveLink/Public/LiveLinkBlueprintStructs.h
@@ -0,0 +1,15 @@
+// Copyright 1998-2019 Epic Games, Inc. All Rights Reserved.
+
+#pragma once
+
+#include "Roles/LiveLinkAnimationBlueprintStructs.h"
+
+UE_DEPRECATED(4.23, "Including LiveLinkBlueprintStructs.h has been deprecated. Please include Roles/LiveLinkAnimationBlueprintStructs.h instead.")
+inline void LiveLinkBlueprintStructsHeaderDeprecatedWarning()
+{
+}
+
+inline void TriggerLiveLinkBlueprintStructsHeaderDeprecatedWarning()
+{
+	LiveLinkBlueprintStructsHeaderDeprecatedWarning();
+}
diff --git a/Engine/Plugins/Animation/LiveLink/Source/LiveLink/Public/LiveLinkClient.h b/Engine/Plugins/Animation/LiveLink/Source/LiveLink/Public/LiveLinkClient.h
--- a/Engine/Plugins/Animation/LiveLink/Source/LiveLink/Public/LiveLinkClient.h
+++ b/Engine/Plugins/Animation/LiveLink/Source/LiveLink/Public/LiveLinkClient.h
@@ -218,6 +218,9 @@ private:
 	void PushSubjectStaticData_Internal(FPendingSubjectStatic&& SubjectStaticData);
 	void PushSubjectFrameData_Internal(FPendingSubjectFrame&& SubjectFrameData);
 
+	/** Remove all sources. */
+	void Shutdown();
+
 private:
 	/** The current collection used. */
 	TUniquePtr<FLiveLinkSourceCollection> Collection;
diff --git a/Engine/Plugins/Animation/LiveLink/Source/LiveLink/Public/LiveLinkComponent.h b/Engine/Plugins/Animation/LiveLink/Source/LiveLink/Public/LiveLinkComponent.h
--- a/Engine/Plugins/Animation/LiveLink/Source/LiveLink/Public/LiveLinkComponent.h
+++ b/Engine/Plugins/Animation/LiveLink/Source/LiveLink/Public/LiveLinkComponent.h
@@ -35,29 +35,29 @@ public:
 	FLiveLinkTickSignature OnLiveLinkUpdated;
 
 	// Returns a list of available Subject Names for LiveLink
-	UE_DEPRECATED(4.23, "LiveLinkComponent is deprecated, use the new LiveLink library.")
-	UFUNCTION(BlueprintCallable, Category = "LiveLink", meta=(DeprecatedFunction, DeprecationMessage="LiveLinkComponent is deprecated, use the new LiveLink library."))
+	UE_DEPRECATED(4.23, "GetAvailableSubjectNames is deprecated, use GetLiveLinkEnabledSubjectNames.")
+	UFUNCTION(BlueprintCallable, Category = "LiveLink", meta=(DeprecatedFunction, DeprecationMessage="GetAvailableSubjectNames is deprecated, use GetLiveLinkEnabledSubjectNames."))
 	void GetAvailableSubjectNames(TArray<FName>& SubjectNames);
 	
 	// Returns a handle to the current frame of data in LiveLink for a given subject along with a boolean for whether a frame was found.
 	// Returns a handle to an empty frame if no frame of data is found.
-	UE_DEPRECATED(4.23, "LiveLinkComponent is deprecated, use the new LiveLink library.")
-	UFUNCTION(BlueprintCallable, Category = "LiveLink", meta=(DeprecatedFunction, DeprecationMessage="LiveLinkComponent is deprecated, use the new LiveLink library."))
+	UE_DEPRECATED(4.23, "GetSubjectData is deprecated, use EvaluateLiveLinkFrame.")
+	UFUNCTION(BlueprintCallable, Category = "LiveLink", meta=(DeprecatedFunction, DeprecationMessage="GetSubjectData is deprecated, EvaluateLiveLinkFrame."))
 	void GetSubjectData(const FName SubjectName, bool& bSuccess, FSubjectFrameHandle& SubjectFrameHandle);
 
 	// Returns a handle to the current frame of data in LiveLink for a given subject along with a boolean for whether a frame was found.
 	// Returns a handle to an empty frame if no frame of data is found.
-	UE_DEPRECATED(4.23, "LiveLinkComponent is deprecated, use the new LiveLink library.")
-	UFUNCTION(BlueprintCallable, Category = "LiveLink", meta=(DeprecatedFunction, DeprecationMessage="LiveLinkComponent is deprecated, use the new LiveLink library."))
+	UE_DEPRECATED(4.23, "GetSubjectDataAtWorldTime is deprecated, use EvaluateLiveLinkFrameAtWorldTime.")
+	UFUNCTION(BlueprintCallable, Category = "LiveLink", meta=(DeprecatedFunction, DeprecationMessage="GetSubjectDataAtWorldTime is deprecated, use EvaluateLiveLinkFrameAtWorldTime."))
 	void GetSubjectDataAtWorldTime(const FName SubjectName, const float WorldTime, bool& bSuccess, FSubjectFrameHandle& SubjectFrameHandle);
 	
-	UE_DEPRECATED(4.23, "LiveLinkComponent is deprecated, use the new LiveLink library.")
+	UE_DEPRECATED(4.23, "GetSubjectDataAtTime is deprecated, use EvaluateLiveLinkFrameAtSceneTime.")
 	void GetSubjectDataAtTime(const FName SubjectName, const double WorldTime, bool& bSuccess, FSubjectFrameHandle& SubjectFrameHandle);
 
 	// Returns a handle to the frame of data in LiveLink for a given subject at the specified time along with a boolean for whether a frame was found.
 	// Returns a handle to an empty frame if no frame of data is found.
-	UE_DEPRECATED(4.23, "LiveLinkComponent is deprecated, use the new LiveLink library.")
-	UFUNCTION(BlueprintCallable, Category = "LiveLink", meta=(DeprecatedFunction, DeprecationMessage="LiveLinkComponent is deprecated, use the new LiveLink library."))
+	UE_DEPRECATED(4.23, "GetSubjectDataAtSceneTime is deprecated, use EvaluateLiveLinkFrameAtSceneTime.")
+	UFUNCTION(BlueprintCallable, Category = "LiveLink", meta=(DeprecatedFunction, DeprecationMessage="GetSubjectDataAtSceneTime is deprecated, use EvaluateLiveLinkFrameAtSceneTime."))
 	void GetSubjectDataAtSceneTime(const FName SubjectName, const FTimecode& SceneTime, bool& bSuccess, FSubjectFrameHandle& SubjectFrameHandle);
 
 private:
diff --git a/Engine/Plugins/Animation/LiveLink/Source/LiveLink/Public/LiveLinkInstance.h b/Engine/Plugins/Animation/LiveLink/Source/LiveLink/Public/LiveLinkInstance.h
--- a/Engine/Plugins/Animation/LiveLink/Source/LiveLink/Public/LiveLinkInstance.h
+++ b/Engine/Plugins/Animation/LiveLink/Source/LiveLink/Public/LiveLinkInstance.h
@@ -32,6 +32,7 @@ public:
 	virtual bool Evaluate(FPoseContext& Output) override;
 	virtual void UpdateAnimationNode(float DeltaSeconds) override;
 
+
 	FAnimNode_LiveLinkPose PoseNode;
 };
 
@@ -40,9 +41,9 @@ class LIVELINK_API ULiveLinkInstance : public UAnimInstance
 {
 	GENERATED_UCLASS_BODY()
 
-	void SetSubject(FName SubjectName)
+	void SetSubject(FLiveLinkSubjectName SubjectName)
 	{
-		GetProxyOnGameThread<FLiveLinkInstanceProxy>().PoseNode.SubjectName = SubjectName;
+		GetProxyOnGameThread<FLiveLinkInstanceProxy>().PoseNode.LiveLinkSubjectName = SubjectName;
 	}
 
 	void SetRetargetAsset(TSubclassOf<ULiveLinkRetargetAsset> RetargetAsset)
diff --git a/Engine/Plugins/Animation/LiveLink/Source/LiveLink/Public/LiveLinkSettings.h b/Engine/Plugins/Animation/LiveLink/Source/LiveLink/Public/LiveLinkSettings.h
--- a/Engine/Plugins/Animation/LiveLink/Source/LiveLink/Public/LiveLinkSettings.h
+++ b/Engine/Plugins/Animation/LiveLink/Source/LiveLink/Public/LiveLinkSettings.h
@@ -25,15 +25,22 @@ struct LIVELINK_API FLiveLinkRoleProjectSetting
 	GENERATED_BODY()
 
 public:
+	FLiveLinkRoleProjectSetting();
+
+public:
+	/** The role of the current setting. */
 	UPROPERTY(config, EditAnywhere, Category = "LiveLink")
 	TSubclassOf<ULiveLinkRole> Role;
 
+	/** The settings class to use for the subject. If null, LiveLinkSubjectSettings will be used by default. */
 	UPROPERTY(config, EditAnywhere, Category = "LiveLink")
 	TSubclassOf<ULiveLinkSubjectSettings> SettingClass;
 
+	/** The interpolation to use for the subject. If null, no interpolation will be performed. */
 	UPROPERTY(config, EditAnywhere, Category = "LiveLink")
 	TSubclassOf<ULiveLinkFrameInterpolationProcessor> FrameInterpolationProcessor;
 
+	/** The pre processors to use for the subject. */
 	UPROPERTY(config, EditAnywhere, Category = "LiveLink")
 	TArray<TSubclassOf<ULiveLinkFramePreProcessor>> FramePreProcessors;
 };
@@ -47,6 +54,9 @@ class LIVELINK_API ULiveLinkSettings : public UObject
 {
 	GENERATED_BODY()
 
+public:
+	ULiveLinkSettings();
+
 protected:
 	UPROPERTY(config, EditAnywhere, Category="LiveLink")
 	TArray<FLiveLinkRoleProjectSetting> DefaultRoleSettings;
@@ -58,15 +68,15 @@ public:
 
 	/** The refresh frequency of the list of message bus provider (when discovery is requested). */
 	UPROPERTY(config, EditAnywhere, AdvancedDisplay, Category="LiveLink", meta=(ConfigRestartRequired=true))
-	double MessageBusPingRequestFrequency = 1.0;
+	double MessageBusPingRequestFrequency;
 
 	/** The refresh frequency of the heartbeat when a provider didn't send us an updated. */
 	UPROPERTY(config, EditAnywhere, AdvancedDisplay, Category="LiveLink", meta=(ConfigRestartRequired = true))
-	double MessageBusHeartbeatFrequency = 1.0;
+	double MessageBusHeartbeatFrequency;
 
 	/** How long we should wait before a provider become unresponsive. */
 	UPROPERTY(config, EditAnywhere, AdvancedDisplay, Category="LiveLink", meta=(ConfigRestartRequired = true))
-	double MessageBusHeartbeatTimeout = 2.0;
+	double MessageBusHeartbeatTimeout;
 
 	/**
 	 * A source may still exist but do not send frames for a subject.
@@ -75,25 +85,16 @@ public:
 	 * An invalid subject is shown as yellow in the LiveLink UI.
 	 */
 	UPROPERTY(config, EditAnywhere, Category = "LiveLink|UI")
-	double TimeWithoutFrameToBeConsiderAsInvalid = 0.5;
+	double TimeWithoutFrameToBeConsiderAsInvalid;
 
 	UPROPERTY(config, EditAnywhere, Category = "LiveLink|UI")
-	FLinearColor ValidColor = FLinearColor::Green;
+	FLinearColor ValidColor;
 
 	UPROPERTY(config, EditAnywhere, Category = "LiveLink|UI")
-	FLinearColor InvalidColor = FLinearColor::Yellow;
+	FLinearColor InvalidColor;
 
 public:
-	FLiveLinkRoleProjectSetting GetDefaultSettingForRole(TSubclassOf<ULiveLinkRole> Role) const
-	{
-		int32 IndexOf = DefaultRoleSettings.IndexOfByPredicate([Role](const FLiveLinkRoleProjectSetting& Other){return Other.Role == Role;});
-		if (IndexOf != INDEX_NONE)
-		{
-			return DefaultRoleSettings[IndexOf];
-		}
-		return FLiveLinkRoleProjectSetting();
-	}
-
+	FLiveLinkRoleProjectSetting GetDefaultSettingForRole(TSubclassOf<ULiveLinkRole> Role) const;
 	const FDirectoryPath& GetPresetSaveDir() const { return PresetSaveDir; }
 	double GetTimeWithoutFrameToBeConsiderAsInvalid() const { return TimeWithoutFrameToBeConsiderAsInvalid; }
 	FLinearColor GetValidColor() const { return ValidColor; }
diff --git a/Engine/Plugins/Animation/LiveLink/Source/LiveLinkComponents/Private/LiveLinkComponentController.cpp b/Engine/Plugins/Animation/LiveLink/Source/LiveLinkComponents/Private/LiveLinkComponentController.cpp
--- a/Engine/Plugins/Animation/LiveLink/Source/LiveLinkComponents/Private/LiveLinkComponentController.cpp
+++ b/Engine/Plugins/Animation/LiveLink/Source/LiveLinkComponents/Private/LiveLinkComponentController.cpp
@@ -56,41 +56,48 @@ void ULiveLinkComponentController::TickComponent(float DeltaTime, ELevelTick Tic
 #if WITH_EDITOR
 void ULiveLinkComponentController::PostEditChangeChainProperty(struct FPropertyChangedChainEvent& PropertyChangedEvent)
 {
-	if (PropertyChangedEvent.Property->GetFName() == GET_MEMBER_NAME_CHECKED(ULiveLinkComponentController, SubjectRepresentation))
+	if (PropertyChangedEvent.PropertyChain.Num() > 0)
 	{
-		bool bCreateAnewController = false;
-		if (SubjectRepresentation.Role.Get() == nullptr)
+		UProperty* MemberProperty = PropertyChangedEvent.PropertyChain.GetActiveMemberNode()->GetValue();
+		if (MemberProperty != NULL)
 		{
-			Controller = nullptr;
-		}
-		else if (Controller)
-		{
-			if (!Controller->IsRoleSupported(SubjectRepresentation.Role))
-			{
-				Controller = nullptr;
-				bCreateAnewController = true;
-			}
-		}
-		else
-		{
-			bCreateAnewController = true;
-		}
-
-		if (bCreateAnewController)
-		{
-			TSubclassOf<ULiveLinkControllerBase> NewControllerClass = ULiveLinkControllerBase::GetControllerForRole(SubjectRepresentation.Role);
-			if (NewControllerClass.Get())
-			{
-				const EObjectFlags ControllerObjectFlags = GetMaskedFlags(RF_Public | RF_Transactional | RF_ArchetypeObject);
-				Controller = NewObject<ULiveLinkControllerBase>(this, NewControllerClass.Get(), NAME_None, ControllerObjectFlags);
-				Controller->InitializeInEditor();
-			}
-			else
+			if (MemberProperty->GetFName() == GET_MEMBER_NAME_CHECKED(ULiveLinkComponentController, SubjectRepresentation))
 			{
-				UE_LOG(LogLiveLinkComponents, Warning, TEXT("No role was found for role '%s'."), *SubjectRepresentation.Role->GetName());
-				FNotificationInfo NotificationInfo(LOCTEXT("NoFoundController", "No controller was found for the role."));
-				NotificationInfo.ExpireDuration = 2.0f;
-				FSlateNotificationManager::Get().AddNotification(NotificationInfo);
+				bool bCreateAnewController = false;
+				if (SubjectRepresentation.Role.Get() == nullptr)
+				{
+					Controller = nullptr;
+				}
+				else if (Controller)
+				{
+					if (!Controller->IsRoleSupported(SubjectRepresentation.Role))
+					{
+						Controller = nullptr;
+						bCreateAnewController = true;
+					}
+				}
+				else
+				{
+					bCreateAnewController = true;
+				}
+
+				if (bCreateAnewController)
+				{
+					TSubclassOf<ULiveLinkControllerBase> NewControllerClass = ULiveLinkControllerBase::GetControllerForRole(SubjectRepresentation.Role);
+					if (NewControllerClass.Get())
+					{
+						const EObjectFlags ControllerObjectFlags = GetMaskedFlags(RF_Public | RF_Transactional | RF_ArchetypeObject);
+						Controller = NewObject<ULiveLinkControllerBase>(this, NewControllerClass.Get(), NAME_None, ControllerObjectFlags);
+						Controller->InitializeInEditor();
+					}
+					else
+					{
+						UE_LOG(LogLiveLinkComponents, Warning, TEXT("No controller was found for role '%s'."), *SubjectRepresentation.Role->GetName());
+						FNotificationInfo NotificationInfo(LOCTEXT("NoFoundController", "No controller was found for the role."));
+						NotificationInfo.ExpireDuration = 2.0f;
+						FSlateNotificationManager::Get().AddNotification(NotificationInfo);
+					}
+				}
 			}
 		}
 	}
diff --git a/Engine/Plugins/Animation/LiveLink/Source/LiveLinkEditor/Private/AnimGraphNode_LiveLinkPose.cpp b/Engine/Plugins/Animation/LiveLink/Source/LiveLinkEditor/Private/AnimGraphNode_LiveLinkPose.cpp
--- a/Engine/Plugins/Animation/LiveLink/Source/LiveLinkEditor/Private/AnimGraphNode_LiveLinkPose.cpp
+++ b/Engine/Plugins/Animation/LiveLink/Source/LiveLinkEditor/Private/AnimGraphNode_LiveLinkPose.cpp
@@ -2,12 +2,9 @@
 
 #include "AnimGraphNode_LiveLinkPose.h"
 
-#define LOCTEXT_NAMESPACE "LiveLinkAnimNode"
+#include "EdGraph/EdGraphSchema.h"
 
-UAnimGraphNode_LiveLinkPose::UAnimGraphNode_LiveLinkPose(const FObjectInitializer& ObjectInitializer)
-	: Super(ObjectInitializer)
-{
-}
+#define LOCTEXT_NAMESPACE "LiveLinkAnimNode"
 
 FText UAnimGraphNode_LiveLinkPose::GetNodeTitle(ENodeTitleType::Type TitleType) const
 {
@@ -24,4 +21,34 @@ FText UAnimGraphNode_LiveLinkPose::GetMenuCategory() const
 	return LOCTEXT("NodeCategory", "Live Link");
 }
 
+void UAnimGraphNode_LiveLinkPose::ConvertDeprecatedNode(UEdGraph* Graph, bool bOnlySafeChanges)
+{
+	//Find deprecated SubjectName pin and set new pin with its value
+	const FName OldPinName = TEXT("SubjectName"); //Variable now has _DEPRECATED appended so can't use GET_MEMBER_NAME_CHECKED
+	const FName NewPinName = GET_MEMBER_NAME_CHECKED(FAnimNode_LiveLinkPose, LiveLinkSubjectName);
+	UEdGraphPin** FoundPinPtr = Pins.FindByPredicate([OldPinName](const UEdGraphPin* Other) { return Other->PinName == OldPinName; });
+	if (FoundPinPtr != nullptr)
+	{
+		UEdGraphPin** FoundNewPinPtr = Pins.FindByPredicate([NewPinName](const UEdGraphPin* Other) { return Other->PinName == NewPinName; });
+		if (FoundNewPinPtr != nullptr)
+		{
+			UScriptStruct* StructType = FLiveLinkSubjectName::StaticStruct();
+			UEdGraphPin* OldPin = *FoundPinPtr;
+			UEdGraphPin* NewPin = *FoundNewPinPtr;
+
+			//Create new structure from old data
+			FLiveLinkSubjectName NewName;
+			NewName.Name = *OldPin->DefaultValue;
+
+			//Apply new name structure 
+			FString StringValue;
+			StructType->ExportText(StringValue, &NewName, nullptr, nullptr, EPropertyPortFlags::PPF_None, nullptr);
+			NewPin->GetSchema()->TrySetDefaultValue(*NewPin, StringValue);
+
+			//Update node data with graph data
+			Node.LiveLinkSubjectName = NewName;
+		}
+	}
+}
+
 #undef LOCTEXT_NAMESPACE
\ No newline at end of file
diff --git a/Engine/Plugins/Animation/LiveLink/Source/LiveLinkEditor/Private/AnimGraphNode_LiveLinkPose.h b/Engine/Plugins/Animation/LiveLink/Source/LiveLinkEditor/Private/AnimGraphNode_LiveLinkPose.h
--- a/Engine/Plugins/Animation/LiveLink/Source/LiveLinkEditor/Private/AnimGraphNode_LiveLinkPose.h
+++ b/Engine/Plugins/Animation/LiveLink/Source/LiveLinkEditor/Private/AnimGraphNode_LiveLinkPose.h
@@ -6,17 +6,29 @@
 #include "AnimNode_LiveLinkPose.h"
 
 #include "AnimGraphNode_LiveLinkPose.generated.h"
+
+
+
 UCLASS()
 class UAnimGraphNode_LiveLinkPose : public UAnimGraphNode_Base
 {
-	GENERATED_UCLASS_BODY()
+	GENERATED_BODY()
 
-	UPROPERTY(EditAnywhere, Category = Settings)
-	FAnimNode_LiveLinkPose Node;
+public:
 
-	// UEdGraphNode interface
+	//~ UEdGraphNode interface
 	virtual FText GetNodeTitle(ENodeTitleType::Type TitleType) const override;
 	virtual FText GetTooltipText() const override;
 	virtual FText GetMenuCategory() const;
-	// End of UEdGraphNode
-};
\ No newline at end of file
+	//~ End of UEdGraphNode
+
+	//~ Begin UK2Node interface
+	virtual void ConvertDeprecatedNode(UEdGraph* Graph, bool bOnlySafeChanges) override;
+	//~ End UK2Node interface
+
+public:
+
+	UPROPERTY(EditAnywhere, Category = Settings)
+	FAnimNode_LiveLinkPose Node;
+
+};
diff --git a/Engine/Plugins/Animation/LiveLink/Source/LiveLinkEditor/Private/K2Node_EvaluateLiveLinkFrame.cpp b/Engine/Plugins/Animation/LiveLink/Source/LiveLinkEditor/Private/K2Node_EvaluateLiveLinkFrame.cpp
--- a/Engine/Plugins/Animation/LiveLink/Source/LiveLinkEditor/Private/K2Node_EvaluateLiveLinkFrame.cpp
+++ b/Engine/Plugins/Animation/LiveLink/Source/LiveLinkEditor/Private/K2Node_EvaluateLiveLinkFrame.cpp
@@ -19,42 +19,78 @@
 #define LOCTEXT_NAMESPACE "K2Node_EvaluateLiveLinkFrame"
 
 
-struct UK2Node_EvaluateLiveLinkFrameHelper
+namespace UK2Node_EvaluateLiveLinkFrameHelper
 {
-	static FName LiveLinkSubjectPinName;
-	static FName LiveLinkDataResultPinName;
-	static FName FrameNotAvailablePinName;
+	static FName LiveLinkSubjectPinName = "Subject";
+	static FName LiveLinkRolePinName = "Role";
+	static FName LiveLinkWorldTimePinName = "WorldTime";
+	static FName LiveLinkSceneTimePinName = "SceneTime";
+	static FName LiveLinkDataResultPinName = "DataResult";
+	static FName FrameNotAvailablePinName = "InvalidFrame";
 };
 
-FName UK2Node_EvaluateLiveLinkFrameHelper::FrameNotAvailablePinName(*LOCTEXT("FrameNotAvailablePinName", "Invalid Frame").ToString());
-FName UK2Node_EvaluateLiveLinkFrameHelper::LiveLinkDataResultPinName(*LOCTEXT("LiveLinkDataResultPinName", "LiveLinkDataResult").ToString());
-FName UK2Node_EvaluateLiveLinkFrameHelper::LiveLinkSubjectPinName(*LOCTEXT("LiveLinkSubjectNamePinName", "LiveLinkSubject").ToString());
-
-UK2Node_EvaluateLiveLinkFrame::UK2Node_EvaluateLiveLinkFrame(const FObjectInitializer& ObjectInitializer)
-	: Super(ObjectInitializer)
-{
-}
-
 void UK2Node_EvaluateLiveLinkFrame::AllocateDefaultPins()
 {
 	const UEdGraphSchema_K2* K2Schema = GetDefault<UEdGraphSchema_K2>();
 
 	// Add execution pins
 	CreatePin(EGPD_Input, UEdGraphSchema_K2::PC_Exec, UEdGraphSchema_K2::PN_Execute);
-	UEdGraphPin* FrameAvailablePin = CreatePin(EGPD_Output, UEdGraphSchema_K2::PC_Exec, UEdGraphSchema_K2::PN_Then);
-	FrameAvailablePin->PinFriendlyName = LOCTEXT("EvaluateLiveLinkFrame Frame available", "Valid Frame");
-	CreatePin(EGPD_Output, UEdGraphSchema_K2::PC_Exec, UK2Node_EvaluateLiveLinkFrameHelper::FrameNotAvailablePinName);
+
+	// Output pin
+	{
+		UEdGraphPin* FrameAvailablePin = CreatePin(EGPD_Output, UEdGraphSchema_K2::PC_Exec, UEdGraphSchema_K2::PN_Then);
+		FrameAvailablePin->PinFriendlyName = LOCTEXT("FrameAvailablePin", "Valid Frame");
+		UEdGraphPin* FrameInvalidePin = CreatePin(EGPD_Output, UEdGraphSchema_K2::PC_Exec, UK2Node_EvaluateLiveLinkFrameHelper::FrameNotAvailablePinName);
+		FrameInvalidePin->PinFriendlyName = LOCTEXT("FrameNotAvailablePinName", "Invalid Frame");
+	}
+	// Output structs pins
+	{
+		UEdGraphPin* DataResultPin = CreatePin(EGPD_Output, UEdGraphSchema_K2::PC_Wildcard, UK2Node_EvaluateLiveLinkFrameHelper::LiveLinkDataResultPinName);
+		DataResultPin->PinFriendlyName = LOCTEXT("LiveLinkDataResultPinName", "DataResult");
+		SetPinToolTip(*DataResultPin, LOCTEXT("DataResultPinDescription", "The data struct, if a frame was present for the given role"));
+	}
 
 	// Subject pin
-	UEdGraphPin* LiveLinkSubjectRepPin = CreatePin(EGPD_Input, UEdGraphSchema_K2::PC_Struct, FLiveLinkSubjectRepresentation::StaticStruct(), UK2Node_EvaluateLiveLinkFrameHelper::LiveLinkSubjectPinName);
-	SetPinToolTip(*LiveLinkSubjectRepPin, LOCTEXT("LiveLinkSubjectNamePinDescription", "The Live Link Subject Reprensation to get a frame from"));
-	
+	{
+		UEdGraphPin* LiveLinkSubjectRepPin = CreatePin(EGPD_Input, UEdGraphSchema_K2::PC_Struct, FLiveLinkSubjectName::StaticStruct(), UK2Node_EvaluateLiveLinkFrameHelper::LiveLinkSubjectPinName);
+		LiveLinkSubjectRepPin->PinFriendlyName = LOCTEXT("LiveLinkSubjectPinName", "Subject");
+		SetPinToolTip(*LiveLinkSubjectRepPin, LOCTEXT("LiveLinkSubjectNamePinDescription", "The Live Link Subject Name to get a frame from"));
+	}
+	// Role pin
+	{
+		UEdGraphPin* LiveLinkRoleRepPin = CreatePin(EGPD_Input, UEdGraphSchema_K2::PC_Class, ULiveLinkRole::StaticClass(), UK2Node_EvaluateLiveLinkFrameHelper::LiveLinkRolePinName);
+		LiveLinkRoleRepPin->PinFriendlyName = LOCTEXT("LiveLinkRolePinName", "Role");
+		LiveLinkRoleRepPin->bNotConnectable = true;
+		SetPinToolTip(*LiveLinkRoleRepPin, LOCTEXT("LiveLinkRolePinNameDescription", "The Live Link Role the data will be converted to."));
+	}
 
-	// Output structs pins
-	UEdGraphPin* DataResultPin = CreatePin(EGPD_Output, UEdGraphSchema_K2::PC_Wildcard, UK2Node_EvaluateLiveLinkFrameHelper::LiveLinkDataResultPinName);
-	SetPinToolTip(*DataResultPin, LOCTEXT("DataResultPinDescription", "The data struct, if a frame was present for the given role"));
+	Super::AllocateDefaultPins();
+}
+
+void UK2Node_EvaluateLiveLinkFrameAtWorldTime::AllocateDefaultPins()
+{
+	Super::AllocateDefaultPins();
+
+	const UEdGraphSchema_K2* K2Schema = GetDefault<UEdGraphSchema_K2>();
+
+	// float pin
+	UEdGraphPin* LiveLinkWorldTimePin = CreatePin(EGPD_Input, UEdGraphSchema_K2::PC_Float, UK2Node_EvaluateLiveLinkFrameHelper::LiveLinkWorldTimePinName);
+	LiveLinkWorldTimePin->PinFriendlyName = LOCTEXT("LiveLinkWorldTimePin", "World Time");
+	SetPinToolTip(*LiveLinkWorldTimePin, LOCTEXT("LiveLinkWorldTimePinDescription", "The World Time the subject will be evaluated to"));
+}
 
+void UK2Node_EvaluateLiveLinkFrameAtSceneTime::AllocateDefaultPins()
+{
 	Super::AllocateDefaultPins();
+
+	const UEdGraphSchema_K2* K2Schema = GetDefault<UEdGraphSchema_K2>();
+
+	// Timecode pin
+	UScriptStruct* TimecodeScriptStruct = FindObject<UScriptStruct>(ANY_PACKAGE, TEXT("Timecode"), true);
+	check(TimecodeScriptStruct);
+	UEdGraphPin* LiveLinkSceneTimePin = CreatePin(EGPD_Input, UEdGraphSchema_K2::PC_Struct, TimecodeScriptStruct, UK2Node_EvaluateLiveLinkFrameHelper::LiveLinkSceneTimePinName);
+	LiveLinkSceneTimePin->PinFriendlyName = LOCTEXT("LiveLinkSceneTimePin", "Scene Time");
+	SetPinToolTip(*LiveLinkSceneTimePin, LOCTEXT("LiveLinkSceneTimePinDescription", "The Scene Time the subject will be evaluated to"));
 }
 
 void UK2Node_EvaluateLiveLinkFrame::SetPinToolTip(UEdGraphPin& InOutMutatablePin, const FText& InPinDescription) const
@@ -96,41 +132,19 @@ void UK2Node_EvaluateLiveLinkFrame::SetReturnTypeForOutputStruct(UScriptStruct*
 	}
 }
 
-UScriptStruct* UK2Node_EvaluateLiveLinkFrame::GetReturnTypeForOutputDataStruct()
+UScriptStruct* UK2Node_EvaluateLiveLinkFrame::GetReturnTypeForOutputDataStruct() const
 {
-	UScriptStruct* ReturnStructType = (UScriptStruct*)(GetResultingDataPin()->PinType.PinSubCategoryObject.Get());
-
-	return ReturnStructType;
+	return Cast<UScriptStruct>(GetResultingDataPin()->PinType.PinSubCategoryObject.Get());
 }
 
 UScriptStruct* UK2Node_EvaluateLiveLinkFrame::GetLiveLinkRoleOutputStructType() const
 {
 	UScriptStruct* DataStructType = nullptr;
 
-	FLiveLinkSubjectRepresentation Representation = GetDefaultSubjectPinValue();
-	if (Representation.Role != nullptr)
+	TSubclassOf<ULiveLinkRole> Role = GetDefaultRolePinValue();
+	if (Role != nullptr)
 	{
-		DataStructType = Representation.Role.GetDefaultObject()->GetBlueprintDataStruct();
-	}
-
-	//No type was deduced from the role, try to deduce it from where it's connected
-	if (DataStructType == nullptr)
-	{
-		UEdGraphPin* ResultPin = GetResultingDataPin();
-		if (ResultPin && ResultPin->LinkedTo.Num() > 0)
-		{
-			DataStructType = Cast<UScriptStruct>(ResultPin->LinkedTo[0]->PinType.PinSubCategoryObject.Get());
-			for (int32 LinkIndex = 1; LinkIndex < ResultPin->LinkedTo.Num(); ++LinkIndex)
-			{
-				UEdGraphPin* Link = ResultPin->LinkedTo[LinkIndex];
-				UScriptStruct* LinkType = Cast<UScriptStruct>(Link->PinType.PinSubCategoryObject.Get());
-
-				if (DataStructType->IsChildOf(LinkType))
-				{
-					DataStructType = LinkType;
-				}
-			}
-		}
+		DataStructType = Role.GetDefaultObject()->GetBlueprintDataStruct();
 	}
 	return DataStructType;
 }
@@ -142,12 +156,12 @@ void UK2Node_EvaluateLiveLinkFrame::ReallocatePinsDuringReconstruction(TArray<UE
 
 void UK2Node_EvaluateLiveLinkFrame::GetMenuActions(FBlueprintActionDatabaseRegistrar& ActionRegistrar) const
 {
-	// actions get registered under specific object-keys; the idea is that 
-	// actions might have to be updated (or deleted) if their object-key is  
-	// mutated (or removed)... here we use the node's class (so if the node 
+	// actions get registered under specific object-keys; the idea is that
+	// actions might have to be updated (or deleted) if their object-key is
+	// mutated (or removed)... here we use the node's class (so if the node
 	// type disappears, then the action should go with it)
 	UClass* ActionKey = GetClass();
-	// to keep from needlessly instantiating a UBlueprintNodeSpawner, first   
+	// to keep from needlessly instantiating a UBlueprintNodeSpawner, first
 	// check to make sure that the registrar is looking for actions of this type
 	// (could be regenerating actions for a specific asset, and therefore the 
 	// registrar would only accept actions corresponding to that asset)
@@ -172,9 +186,10 @@ bool UK2Node_EvaluateLiveLinkFrame::IsConnectionDisallowed(const UEdGraphPin* My
 		bool bDisallowed = true;
 		if (OtherPin->PinType.PinCategory == UEdGraphSchema_K2::PC_Struct)
 		{
+			bDisallowed = false;
 			if (UScriptStruct* ConnectionType = Cast<UScriptStruct>(OtherPin->PinType.PinSubCategoryObject.Get()))
 			{
-				bDisallowed = !ConnectionType->IsChildOf(FLiveLinkBaseBlueprintData::StaticStruct());
+				bDisallowed = GetLiveLinkRoleOutputStructType() != ConnectionType;
 			}
 		}
 		else if (OtherPin->PinType.PinCategory == UEdGraphSchema_K2::PC_Wildcard)
@@ -194,7 +209,7 @@ bool UK2Node_EvaluateLiveLinkFrame::IsConnectionDisallowed(const UEdGraphPin* My
 
 void UK2Node_EvaluateLiveLinkFrame::PinDefaultValueChanged(UEdGraphPin* ChangedPin)
 {
-	if (ChangedPin && ChangedPin->PinName == UK2Node_EvaluateLiveLinkFrameHelper::LiveLinkSubjectPinName)
+	if (ChangedPin && ChangedPin->PinName == UK2Node_EvaluateLiveLinkFrameHelper::LiveLinkRolePinName)
 	{
 		RefreshDataOutputPinType();
 	}
@@ -214,19 +229,16 @@ UEdGraphPin* UK2Node_EvaluateLiveLinkFrame::GetThenPin()const
 	return Pin;
 }
 
-UEdGraphPin* UK2Node_EvaluateLiveLinkFrame::GetLiveLinkSubjectPin(const TArray<UEdGraphPin*>* InPinsToSearch /*= nullptr*/) const
+UEdGraphPin* UK2Node_EvaluateLiveLinkFrame::GetLiveLinkRolePin() const
 {
-	const TArray<UEdGraphPin*>* PinsToSearch = InPinsToSearch ? InPinsToSearch : &Pins;
+	UEdGraphPin* Pin = FindPinChecked(UK2Node_EvaluateLiveLinkFrameHelper::LiveLinkRolePinName);
+	check(Pin == nullptr || Pin->Direction == EGPD_Input);
+	return Pin;
+}
 
-	UEdGraphPin* Pin = nullptr;
-	for (UEdGraphPin* TestPin : *PinsToSearch)
-	{
-		if (TestPin && TestPin->PinName == UK2Node_EvaluateLiveLinkFrameHelper::LiveLinkSubjectPinName)
-		{
-			Pin = TestPin;
-			break;
-		}
-	}
+UEdGraphPin* UK2Node_EvaluateLiveLinkFrame::GetLiveLinkSubjectPin() const
+{
+	UEdGraphPin* Pin = FindPinChecked(UK2Node_EvaluateLiveLinkFrameHelper::LiveLinkSubjectPinName);
 	check(Pin == nullptr || Pin->Direction == EGPD_Input);
 	return Pin;
 }
@@ -249,16 +261,54 @@ UEdGraphPin* UK2Node_EvaluateLiveLinkFrame::GetResultingDataPin() const
 
 FText UK2Node_EvaluateLiveLinkFrame::GetNodeTitle(ENodeTitleType::Type TitleType) const
 {
-	return LOCTEXT("ListViewTitle", "Evaluate Live Link Frame");
+	return LOCTEXT("EvaluateLiveLinkFrameTitle", "Evaluate Live Link Frame");
+}
+
+FText UK2Node_EvaluateLiveLinkFrameAtWorldTime::GetNodeTitle(ENodeTitleType::Type TitleType) const
+{
+	return LOCTEXT("EvaluateLiveLinkFrameAtWorldTimeTitle", "Evaluate Live Link Frame at World Time");
+}
+
+FText UK2Node_EvaluateLiveLinkFrameAtSceneTime::GetNodeTitle(ENodeTitleType::Type TitleType) const
+{
+	return LOCTEXT("EvaluateLiveLinkFrameAtSceneTimeTitle", "Evaluate Live Link Frame at Scene Time");
+}
+
+FName UK2Node_EvaluateLiveLinkFrame::GetEvaluateFunctionName() const
+{
+	return GET_FUNCTION_NAME_CHECKED(ULiveLinkBlueprintLibrary, EvaluateLiveLinkFrameWithSpecificRole);
+}
+
+FName UK2Node_EvaluateLiveLinkFrameAtWorldTime::GetEvaluateFunctionName() const
+{
+	return GET_FUNCTION_NAME_CHECKED(ULiveLinkBlueprintLibrary, EvaluateLiveLinkFrameAtWorldTimeOffset);
+}
+
+FName UK2Node_EvaluateLiveLinkFrameAtSceneTime::GetEvaluateFunctionName() const
+{
+	return GET_FUNCTION_NAME_CHECKED(ULiveLinkBlueprintLibrary, EvaluateLiveLinkFrameAtSceneTime);
+}
+
+void UK2Node_EvaluateLiveLinkFrameAtWorldTime::AddOtherPin(FKismetCompilerContext& CompilerContext, UK2Node_CallFunction* EvaluateLiveLinkFrameFunction)
+{
+	UEdGraphPin* InPinSwitch = FindPinChecked(UK2Node_EvaluateLiveLinkFrameHelper::LiveLinkWorldTimePinName);
+	UEdGraphPin* TimePin = EvaluateLiveLinkFrameFunction->FindPinChecked(TEXT("WorldTime"));
+	CompilerContext.CopyPinLinksToIntermediate(*InPinSwitch, *TimePin);
+}
+
+void UK2Node_EvaluateLiveLinkFrameAtSceneTime::AddOtherPin(FKismetCompilerContext& CompilerContext, UK2Node_CallFunction* EvaluateLiveLinkFrameFunction)
+{
+	UEdGraphPin* InPinSwitch = FindPinChecked(UK2Node_EvaluateLiveLinkFrameHelper::LiveLinkSceneTimePinName);
+	UEdGraphPin* TimePin = EvaluateLiveLinkFrameFunction->FindPinChecked(TEXT("SceneTime"));
+	CompilerContext.CopyPinLinksToIntermediate(*InPinSwitch, *TimePin);
 }
 
 void UK2Node_EvaluateLiveLinkFrame::ExpandNode(class FKismetCompilerContext& CompilerContext, UEdGraph* SourceGraph)
 {
 	Super::ExpandNode(CompilerContext, SourceGraph);
 
-	UEdGraphPin* OriginalLiveLinkSubjectPin = GetLiveLinkSubjectPin();
-	FLiveLinkSubjectRepresentation Representation = GetDefaultSubjectPinValue();
-	if (Representation.Role == nullptr)
+	TSubclassOf<ULiveLinkRole> Role = GetDefaultRolePinValue();
+	if (Role == nullptr)
 	{
 		CompilerContext.MessageLog.Error(*LOCTEXT("EvaluateLiveLinkRoleNoRole_Error", "EvaluateLiveLinkFrame must have a Role specified.").ToString(), this);
 		// we break exec links so this is the only error we get
@@ -267,7 +317,7 @@ void UK2Node_EvaluateLiveLinkFrame::ExpandNode(class FKismetCompilerContext& Com
 	}
 
 	// FUNCTION NODE
-	const FName FunctionName = GET_FUNCTION_NAME_CHECKED(ULiveLinkBlueprintLibrary, EvaluateLiveLinkFrame);
+	const FName FunctionName = GetEvaluateFunctionName();
 	UK2Node_CallFunction* EvaluateLiveLinkFrameFunction = CompilerContext.SpawnIntermediateNode<UK2Node_CallFunction>(this, SourceGraph);
 	EvaluateLiveLinkFrameFunction->FunctionReference.SetExternalMember(FunctionName, ULiveLinkBlueprintLibrary::StaticClass());
 	EvaluateLiveLinkFrameFunction->AllocateDefaultPins();
@@ -275,19 +325,18 @@ void UK2Node_EvaluateLiveLinkFrame::ExpandNode(class FKismetCompilerContext& Com
 
 	// Connect the input of our EvaluateLiveLinkFrame to the Input of our Function pin
 	{
-		UEdGraphPin* LiveLinkSubjectInPin = EvaluateLiveLinkFrameFunction->FindPinChecked(TEXT("SubjectRepresentation"));
-		if (OriginalLiveLinkSubjectPin->LinkedTo.Num() > 0)
-		{
-			// Copy the connection
-			CompilerContext.MovePinLinksToIntermediate(*OriginalLiveLinkSubjectPin, *LiveLinkSubjectInPin);
-		}
-		else
-		{
-			// Copy literal
-			LiveLinkSubjectInPin->DefaultValue = OriginalLiveLinkSubjectPin->DefaultValue;
-		}
+		UEdGraphPin* OriginalLiveLinkSubjectPin = GetLiveLinkSubjectPin();
+		UEdGraphPin* LiveLinkSubjectInPin = EvaluateLiveLinkFrameFunction->FindPinChecked(TEXT("SubjectName"));
+		CompilerContext.CopyPinLinksToIntermediate(*OriginalLiveLinkSubjectPin, *LiveLinkSubjectInPin);
+	}
+	{
+		UEdGraphPin* OriginalLiveLinkRolePin = GetLiveLinkRolePin();
+		UEdGraphPin* LiveLinkRoleInPin = EvaluateLiveLinkFrameFunction->FindPinChecked(TEXT("Role"));
+		CompilerContext.CopyPinLinksToIntermediate(*OriginalLiveLinkRolePin, *LiveLinkRoleInPin);
 	}
 
+	AddOtherPin(CompilerContext, EvaluateLiveLinkFrameFunction);
+
 	// Get some pins to work with
 	UEdGraphPin* OriginalDataOutPin = FindPinChecked(UK2Node_EvaluateLiveLinkFrameHelper::LiveLinkDataResultPinName);
 	UEdGraphPin* FunctionDataOutPin = EvaluateLiveLinkFrameFunction->FindPinChecked(TEXT("OutBlueprintData"));
@@ -338,18 +387,37 @@ void UK2Node_EvaluateLiveLinkFrame::EarlyValidation(class FCompilerResultsLog& M
 		return;
 	}
 
-	FLiveLinkSubjectRepresentation Representation = GetDefaultSubjectPinValue();
-	if (Representation.Role == nullptr || !Representation.Role->IsChildOf(ULiveLinkRole::StaticClass()))
+	TSubclassOf<ULiveLinkRole> Role = GetDefaultRolePinValue();
+	if (Role == nullptr || !Role->IsChildOf(ULiveLinkRole::StaticClass()))
 	{
 		MessageLog.Error(*LOCTEXT("NoLiveLinkRole", "No LiveLinkRole in @@").ToString(), this);
 		return;
 	}
 
-	if (Representation.Subject.IsNone())
+	UScriptStruct* DesiredOutputStruct = GetLiveLinkRoleOutputStructType();
+	UScriptStruct* ActualOutputStruct = GetReturnTypeForOutputDataStruct();
+	if (DesiredOutputStruct != ActualOutputStruct)
 	{
-		MessageLog.Warning(*LOCTEXT("NoLiveLinkSubjectName", "No subject in @@").ToString(), this);
+		MessageLog.Error(*LOCTEXT("OutputPinDoNotMatches", "The output data pin do not maches in @@").ToString(), this);
 		return;
 	}
+
+	if (ActualOutputStruct != nullptr)
+	{
+		UEdGraphPin* ResultPin = GetResultingDataPin();
+		if (ResultPin && ResultPin->LinkedTo.Num() > 0)
+		{
+			for (UEdGraphPin* LinkPin : ResultPin->LinkedTo)
+			{
+				UScriptStruct* LinkType = Cast<UScriptStruct>(LinkPin->PinType.PinSubCategoryObject.Get());
+				if (ActualOutputStruct != LinkType)
+				{
+					MessageLog.Error(*LOCTEXT("OutputLinkPinDoNotMatches", "A linked pin data structure do not match the Live Link Role data structure in @@. Expected: @@.").ToString(), this, ActualOutputStruct);
+					return;
+				}
+			}
+		}
+	}
 }
 
 void UK2Node_EvaluateLiveLinkFrame::PreloadRequiredAssets()
@@ -361,17 +429,17 @@ void UK2Node_EvaluateLiveLinkFrame::NotifyPinConnectionListChanged(UEdGraphPin*
 {
 	Super::NotifyPinConnectionListChanged(Pin);
 
-	UEdGraphPin* LiveLinkSubjectPin = GetLiveLinkSubjectPin();
+	UEdGraphPin* LiveLinkRolePin = GetLiveLinkRolePin();
 	if (Pin == GetResultingDataPin())
 	{
 		// this connection would only change the output type if the role pin is undefined
-		const bool bIsTypeAuthority = (LiveLinkSubjectPin->LinkedTo.Num() <= 0 && LiveLinkSubjectPin->DefaultObject == nullptr);
+		const bool bIsTypeAuthority = (LiveLinkRolePin->LinkedTo.Num() <= 0 && LiveLinkRolePin->DefaultObject == nullptr);
 		if (bIsTypeAuthority)
 		{
 			RefreshDataOutputPinType();
 		}
 	}
-	else if (Pin == LiveLinkSubjectPin)
+	else if (Pin == LiveLinkRolePin)
 	{
 		const bool bConnectionAdded = Pin->LinkedTo.Num() > 0;
 		if (bConnectionAdded)
@@ -381,27 +449,18 @@ void UK2Node_EvaluateLiveLinkFrame::NotifyPinConnectionListChanged(UEdGraphPin*
 	}
 }
 
-FLiveLinkSubjectRepresentation UK2Node_EvaluateLiveLinkFrame::GetDefaultSubjectPinValue() const
+TSubclassOf<ULiveLinkRole> UK2Node_EvaluateLiveLinkFrame::GetDefaultRolePinValue() const
 {
-	FLiveLinkSubjectRepresentation Representation;
-
-	UEdGraphPin* LiveLinkSubjectPin = GetLiveLinkSubjectPin();
-	if (LiveLinkSubjectPin)
+	UEdGraphPin* LiveLinkRolePin = GetLiveLinkRolePin();
+	if (LiveLinkRolePin && LiveLinkRolePin->DefaultObject != nullptr && LiveLinkRolePin->LinkedTo.Num() == 0)
 	{
-		if (LiveLinkSubjectPin->LinkedTo.Num() > 0)
-		{
-			if (LiveLinkSubjectPin->LinkedTo[0]->DefaultValue.Len())
-			{
-				FLiveLinkSubjectRepresentation::StaticStruct()->ImportText(*LiveLinkSubjectPin->LinkedTo[0]->DefaultValue, &Representation, nullptr, EPropertyPortFlags::PPF_None, nullptr, TEXT(""));
-			}
-		}
-		else if (LiveLinkSubjectPin->DefaultValue.Len())
+		UClass* ClassValue = Cast<UClass>(LiveLinkRolePin->DefaultObject);
+		if (ClassValue && ClassValue->IsChildOf(ULiveLinkRole::StaticClass()))
 		{
-			FLiveLinkSubjectRepresentation::StaticStruct()->ImportText(*LiveLinkSubjectPin->DefaultValue, &Representation, nullptr, EPropertyPortFlags::PPF_None, nullptr, TEXT(""));
+			return ClassValue;
 		}
 	}
-
-	return Representation;
+	return nullptr;
 }
 
 
diff --git a/Engine/Plugins/Animation/LiveLink/Source/LiveLinkEditor/Private/LiveLinkEditorModule.cpp b/Engine/Plugins/Animation/LiveLink/Source/LiveLinkEditor/Private/LiveLinkEditorModule.cpp
--- a/Engine/Plugins/Animation/LiveLink/Source/LiveLinkEditor/Private/LiveLinkEditorModule.cpp
+++ b/Engine/Plugins/Animation/LiveLink/Source/LiveLinkEditor/Private/LiveLinkEditorModule.cpp
@@ -33,9 +33,13 @@
 #include "LiveLinkVirtualSubjectDetailCustomization.h"
 
 /**
- * Implements the Messaging module.
+ * Implements the LiveLinkEditor module.
  */
 
+
+DEFINE_LOG_CATEGORY(LogLiveLinkEditor);
+
+
 #define LOCTEXT_NAMESPACE "LiveLinkModule"
 
 static const FName LiveLinkClientTabName(TEXT("LiveLink"));
diff --git a/Engine/Plugins/Animation/LiveLink/Source/LiveLinkEditor/Private/LiveLinkEditorPrivate.h b/Engine/Plugins/Animation/LiveLink/Source/LiveLinkEditor/Private/LiveLinkEditorPrivate.h
--- a/Engine/Plugins/Animation/LiveLink/Source/LiveLinkEditor/Private/LiveLinkEditorPrivate.h
+++ b/Engine/Plugins/Animation/LiveLink/Source/LiveLinkEditor/Private/LiveLinkEditorPrivate.h
@@ -4,6 +4,9 @@
 
 #include "CoreMinimal.h"
 
+DECLARE_LOG_CATEGORY_EXTERN(LogLiveLinkEditor, Log, All);
+
+
 class ISlateStyle;
 
 struct FLiveLinkEditorPrivate
diff --git a/Engine/Plugins/Animation/LiveLink/Source/LiveLinkEditor/Private/LiveLinkGraphPanelPinFactory.h b/Engine/Plugins/Animation/LiveLink/Source/LiveLinkEditor/Private/LiveLinkGraphPanelPinFactory.h
--- a/Engine/Plugins/Animation/LiveLink/Source/LiveLinkEditor/Private/LiveLinkGraphPanelPinFactory.h
+++ b/Engine/Plugins/Animation/LiveLink/Source/LiveLinkEditor/Private/LiveLinkGraphPanelPinFactory.h
@@ -9,6 +9,7 @@
 #include "SGraphPin.h"
 
 #include "LiveLinkRole.h"
+#include "SLiveLinkSubjectNameGraphPin.h"
 #include "SLiveLinkSubjectRepresentationGraphPin.h"
 
 
@@ -20,10 +21,14 @@ class FLiveLinkGraphPanelPinFactory : public FGraphPanelPinFactory
 		{
 			if (UScriptStruct* PinStructType = Cast<UScriptStruct>(InPin->PinType.PinSubCategoryObject.Get()))
 			{
-				if (PinStructType->IsChildOf(FLiveLinkSubjectRepresentation::StaticStruct()))
+				if (PinStructType == FLiveLinkSubjectRepresentation::StaticStruct())
 				{
 					return SNew(SLiveLinkSubjectRepresentationGraphPin, InPin);
 				}
+				else if (PinStructType == FLiveLinkSubjectName::StaticStruct())
+				{
+					return SNew(SLiveLinkSubjectNameGraphPin, InPin);
+				}
 			}
 		}
 
diff --git a/Engine/Plugins/Animation/LiveLink/Source/LiveLinkEditor/Private/LiveLinkPreviewController.cpp b/Engine/Plugins/Animation/LiveLink/Source/LiveLinkEditor/Private/LiveLinkPreviewController.cpp
--- a/Engine/Plugins/Animation/LiveLink/Source/LiveLinkEditor/Private/LiveLinkPreviewController.cpp
+++ b/Engine/Plugins/Animation/LiveLink/Source/LiveLinkEditor/Private/LiveLinkPreviewController.cpp
@@ -7,6 +7,7 @@
 #include "ILiveLinkClient.h"
 #include "IPersonaPreviewScene.h"
 #include "LiveLinkClientReference.h"
+#include "LiveLinkCustomVersion.h"
 #include "LiveLinkInstance.h"
 #include "LiveLinkRemapAsset.h"
 #include "Roles/LiveLinkCameraRole.h"
@@ -55,7 +56,7 @@ void ULiveLinkPreviewController::InitializeView(UPersonaPreviewSceneDescription*
 
 	if (ULiveLinkInstance* LiveLinkInstance = Cast<ULiveLinkInstance>(PreviewScene->GetPreviewMeshComponent()->GetAnimInstance()))
 	{
-		LiveLinkInstance->SetSubject(SubjectName);
+		LiveLinkInstance->SetSubject(LiveLinkSubjectName);
 		LiveLinkInstance->SetRetargetAsset(RetargetAsset);
 	}
 	if (bEnableCameraSync)
@@ -68,4 +69,25 @@ void ULiveLinkPreviewController::UninitializeView(UPersonaPreviewSceneDescriptio
 {
 	PreviewScene->GetPreviewMeshComponent()->SetAnimInstanceClass(nullptr);
 	PreviewScene->SetCameraOverride(nullptr);
-}
\ No newline at end of file
+}
+
+void ULiveLinkPreviewController::Serialize(FArchive& Ar)
+{
+	Super::Serialize(Ar);
+
+#if WITH_EDITORONLY_DATA
+	Ar.UsingCustomVersion(FLiveLinkCustomVersion::GUID);
+
+	if (Ar.IsLoading())
+	{
+		const int32 LiveLinkVersion = Ar.CustomVer(FLiveLinkCustomVersion::GUID);
+
+		if (LiveLinkVersion < FLiveLinkCustomVersion::NewLiveLinkRoleSystem)
+		{
+			PRAGMA_DISABLE_DEPRECATION_WARNINGS
+			LiveLinkSubjectName.Name = SubjectName_DEPRECATED;
+			PRAGMA_ENABLE_DEPRECATION_WARNINGS
+		}
+	}
+#endif //WITH_EDITORONLY_DATA
+}
diff --git a/Engine/Plugins/Animation/LiveLink/Source/LiveLinkEditor/Private/LiveLinkPreviewController.h b/Engine/Plugins/Animation/LiveLink/Source/LiveLinkEditor/Private/LiveLinkPreviewController.h
--- a/Engine/Plugins/Animation/LiveLink/Source/LiveLinkEditor/Private/LiveLinkPreviewController.h
+++ b/Engine/Plugins/Animation/LiveLink/Source/LiveLinkEditor/Private/LiveLinkPreviewController.h
@@ -3,6 +3,8 @@
 #pragma once
 
 #include "PersonaPreviewSceneController.h"
+
+#include "LiveLinkTypes.h"
 #include "Templates/SubclassOf.h"
 
 #include "LiveLinkPreviewController.generated.h"
@@ -15,8 +17,14 @@ class ULiveLinkPreviewController : public UPersonaPreviewSceneController
 public:
 	GENERATED_BODY()
 
+#if WITH_EDITORONLY_DATA
+	UE_DEPRECATED(4.23, "FName SubjectName is deprecated. Use the SubjectName of type FLiveLinkSubjectName instead.")
+	UPROPERTY()
+	FName SubjectName_DEPRECATED;
+#endif //WITH_EDITORONLY_DATA
+
 	UPROPERTY(EditAnywhere, Category = "Live Link")
-	FName SubjectName;
+	FLiveLinkSubjectName LiveLinkSubjectName;
 
 	UPROPERTY(EditAnywhere, Category = "Live Link")
 	bool bEnableCameraSync;
@@ -27,4 +35,8 @@ public:
 	virtual void InitializeView(UPersonaPreviewSceneDescription* SceneDescription, IPersonaPreviewScene* PreviewScene) const;
 	virtual void UninitializeView(UPersonaPreviewSceneDescription* SceneDescription, IPersonaPreviewScene* PreviewScene) const;
 
+	//~ Begin UObject Interface
+	virtual void Serialize(FArchive& Ar);
+	//~ End UObject Interface
+
 };
diff --git a/Engine/Plugins/Animation/LiveLink/Source/LiveLinkEditor/Private/SLiveLinkSubjectNameGraphPin.cpp b/Engine/Plugins/Animation/LiveLink/Source/LiveLinkEditor/Private/SLiveLinkSubjectNameGraphPin.cpp
new file mode 100644
--- /dev/null
+++ b/Engine/Plugins/Animation/LiveLink/Source/LiveLinkEditor/Private/SLiveLinkSubjectNameGraphPin.cpp
@@ -0,0 +1,48 @@
+// Copyright 1998-2019 Epic Games, Inc. All Rights Reserved.
+
+#include "SLiveLinkSubjectNameGraphPin.h"
+#include "SLiveLinkSubjectRepresentationPicker.h"
+#include "K2Node_CallFunction.h"
+#include "K2Node_VariableSet.h"
+#include "K2Node_EvaluateLiveLinkFrame.h"
+
+
+void SLiveLinkSubjectNameGraphPin::Construct(const FArguments& InArgs, UEdGraphPin* InGraphPinObj)
+{
+	SGraphPin::Construct(SGraphPin::FArguments(), InGraphPinObj);
+}
+
+TSharedRef<SWidget> SLiveLinkSubjectNameGraphPin::GetDefaultValueWidget()
+{
+	//Create widget
+	return SNew(SLiveLinkSubjectRepresentationPicker)
+		.ShowRole(false)
+		.HasMultipleValues(false)
+		.Value(this, &SLiveLinkSubjectNameGraphPin::GetValue)
+		.OnValueChanged(this, &SLiveLinkSubjectNameGraphPin::SetValue)
+		.Visibility(this, &SGraphPin::GetDefaultValueVisibility);
+}
+
+FLiveLinkSubjectRepresentation SLiveLinkSubjectNameGraphPin::GetValue() const
+{
+	FLiveLinkSubjectRepresentation SubjectRepresentation;
+
+	FLiveLinkSubjectName::StaticStruct()->ImportText(*GraphPinObj->GetDefaultAsString(), &SubjectRepresentation.Subject, nullptr, EPropertyPortFlags::PPF_None, GLog, FLiveLinkSubjectName::StaticStruct()->GetName());
+	SubjectRepresentation.Role = Cast<UClass>(GraphPinObj->DefaultObject);
+
+	return SubjectRepresentation;
+}
+
+void SLiveLinkSubjectNameGraphPin::SetValue(FLiveLinkSubjectRepresentation NewValue)
+{
+	FString ValueString;
+	FLiveLinkSubjectName::StaticStruct()->ExportText(ValueString, &NewValue.Subject, nullptr, nullptr, EPropertyPortFlags::PPF_None, nullptr);
+	GraphPinObj->GetSchema()->TrySetDefaultValue(*GraphPinObj, ValueString);
+
+	if (UK2Node_EvaluateLiveLinkFrame* OnwingNode = Cast<UK2Node_EvaluateLiveLinkFrame>(GraphPinObj->GetOwningNode()))
+	{
+		UEdGraphPin* RolePin = OnwingNode->GetLiveLinkRolePin();
+		RolePin->GetSchema()->TrySetDefaultObject(*RolePin, NewValue.Role.Get());
+	}
+}
+
diff --git a/Engine/Plugins/Animation/LiveLink/Source/LiveLinkEditor/Private/SLiveLinkSubjectNameGraphPin.h b/Engine/Plugins/Animation/LiveLink/Source/LiveLinkEditor/Private/SLiveLinkSubjectNameGraphPin.h
new file mode 100644
--- /dev/null
+++ b/Engine/Plugins/Animation/LiveLink/Source/LiveLinkEditor/Private/SLiveLinkSubjectNameGraphPin.h
@@ -0,0 +1,26 @@
+// Copyright 1998-2019 Epic Games, Inc. All Rights Reserved.
+
+#pragma once
+
+#include "CoreMinimal.h"
+#include "LiveLinkRole.h"
+#include "SGraphPin.h"
+#include "Widgets/SWidget.h"
+#include "Widgets/DeclarativeSyntaxSupport.h"
+
+class SLiveLinkSubjectNameGraphPin : public SGraphPin
+{
+public:
+	SLATE_BEGIN_ARGS(SLiveLinkSubjectNameGraphPin) {}
+	SLATE_END_ARGS()
+
+	void Construct(const FArguments& InArgs, UEdGraphPin* InGraphPinObj);
+
+	//~ Begin SGraphPin Interface
+	virtual TSharedRef<SWidget>	GetDefaultValueWidget() override;
+	//~ End SGraphPin Interface
+
+private:
+	FLiveLinkSubjectRepresentation GetValue() const;
+	void SetValue(FLiveLinkSubjectRepresentation NewValue);
+};
diff --git a/Engine/Plugins/Animation/LiveLink/Source/LiveLinkEditor/Public/K2Node_EvaluateLiveLinkFrame.h b/Engine/Plugins/Animation/LiveLink/Source/LiveLinkEditor/Public/K2Node_EvaluateLiveLinkFrame.h
--- a/Engine/Plugins/Animation/LiveLink/Source/LiveLinkEditor/Public/K2Node_EvaluateLiveLinkFrame.h
+++ b/Engine/Plugins/Animation/LiveLink/Source/LiveLinkEditor/Public/K2Node_EvaluateLiveLinkFrame.h
@@ -12,21 +12,24 @@
 #include "K2Node_EvaluateLiveLinkFrame.generated.h"
 
 class FBlueprintActionDatabaseRegistrar;
+class FKismetCompilerContext;
 class UDataTable;
 class UEdGraph;
+class UK2Node_CallFunction;
 
 
 UCLASS()
 class LIVELINKEDITOR_API UK2Node_EvaluateLiveLinkFrame : public UK2Node
 {
-	GENERATED_UCLASS_BODY()
+	GENERATED_BODY()
 
+public:
 	//~ Begin UEdGraphNode Interface.
 	virtual void AllocateDefaultPins() override;
 	virtual FText GetNodeTitle(ENodeTitleType::Type TitleType) const override;
 	virtual void PinDefaultValueChanged(UEdGraphPin* Pin) override;
 	virtual FText GetTooltipText() const override;
-	virtual void ExpandNode(class FKismetCompilerContext& CompilerContext, UEdGraph* SourceGraph) override;
+	virtual void ExpandNode(FKismetCompilerContext& CompilerContext, UEdGraph* SourceGraph) override;
 	virtual FSlateIcon GetIconAndTint(FLinearColor& OutColor) const override;
 	virtual void PostReconstructNode() override;
 	//~ End UEdGraphNode Interface.
@@ -43,13 +46,16 @@ class LIVELINKEDITOR_API UK2Node_EvaluateLiveLinkFrame : public UK2Node
 	//~ End UK2Node Interface
 
 	/** Get the return types of our struct */
-	UScriptStruct* GetReturnTypeForOutputDataStruct();
+	UScriptStruct* GetReturnTypeForOutputDataStruct() const;
 
 	/** Get the then output pin */
 	UEdGraphPin* GetThenPin() const;
 
-	/** Get the Live Link Subject Representation input pin */
-	UEdGraphPin* GetLiveLinkSubjectPin(const TArray<UEdGraphPin*>* InPinsToSearch = nullptr) const;
+	/** Get the Live Link Role input pin */
+	UEdGraphPin* GetLiveLinkRolePin() const;
+
+	/** Get the Live Link Subject input pin */
+	UEdGraphPin* GetLiveLinkSubjectPin() const;
 
 	/** Get the exec output pin for when no frame is available for the desired role */
 	UEdGraphPin* GetFrameNotAvailablePin() const;
@@ -61,7 +67,11 @@ class LIVELINKEDITOR_API UK2Node_EvaluateLiveLinkFrame : public UK2Node
 	UScriptStruct* GetLiveLinkRoleOutputStructType() const;
 	UScriptStruct* GetLiveLinkRoleOutputFrameStructType() const;
 
-private:
+
+	virtual FName GetEvaluateFunctionName() const;
+	virtual void AddOtherPin(FKismetCompilerContext& CompilerContext, UK2Node_CallFunction* EvaluateLiveLinkFrameFunction) {}
+
+protected:
 
 	/**
 	 * Takes the specified "MutatablePin" and sets its 'PinToolTip' field (according
@@ -78,5 +88,31 @@ private:
 	/** Queries for the authoritative return type, then modifies the return pin to match */
 	void RefreshDataOutputPinType();
 
-	FLiveLinkSubjectRepresentation GetDefaultSubjectPinValue() const;
+	TSubclassOf<ULiveLinkRole> GetDefaultRolePinValue() const;
+};
+
+UCLASS()
+class LIVELINKEDITOR_API UK2Node_EvaluateLiveLinkFrameAtWorldTime : public UK2Node_EvaluateLiveLinkFrame
+{
+	GENERATED_BODY()
+
+public:
+	virtual void AllocateDefaultPins() override;
+	virtual FText GetNodeTitle(ENodeTitleType::Type TitleType) const override;
+
+	virtual FName GetEvaluateFunctionName() const override;
+	virtual void AddOtherPin(FKismetCompilerContext& CompilerContext, UK2Node_CallFunction* EvaluateLiveLinkFrameFunction) override;
 };
+
+UCLASS()
+class LIVELINKEDITOR_API UK2Node_EvaluateLiveLinkFrameAtSceneTime : public UK2Node_EvaluateLiveLinkFrame
+{
+	GENERATED_BODY()
+
+public:
+	virtual void AllocateDefaultPins() override;
+	virtual FText GetNodeTitle(ENodeTitleType::Type TitleType) const override;
+
+	virtual FName GetEvaluateFunctionName() const override;
+	virtual void AddOtherPin(FKismetCompilerContext& CompilerContext, UK2Node_CallFunction* EvaluateLiveLinkFrameFunction) override;
+};
\ No newline at end of file
diff --git a/Engine/Plugins/Animation/LiveLink/Source/LiveLinkMovieScene/Private/MovieScene/MovieSceneLiveLinkSection.cpp b/Engine/Plugins/Animation/LiveLink/Source/LiveLinkMovieScene/Private/MovieScene/MovieSceneLiveLinkSection.cpp
--- a/Engine/Plugins/Animation/LiveLink/Source/LiveLinkMovieScene/Private/MovieScene/MovieSceneLiveLinkSection.cpp
+++ b/Engine/Plugins/Animation/LiveLink/Source/LiveLinkMovieScene/Private/MovieScene/MovieSceneLiveLinkSection.cpp
@@ -189,34 +189,51 @@ void UMovieSceneLiveLinkSection::ConvertPreRoleData()
 	PRAGMA_DISABLE_DEPRECATION_WARNINGS
 	SubjectPreset.Key.Source.Invalidate();
 	SubjectPreset.Key.SubjectName = SubjectName_DEPRECATED;
-	SubjectPreset.Role = ULiveLinkAnimationRole::StaticClass();
 
-	StaticData->InitializeWith(FLiveLinkSkeletonStaticData::StaticStruct(), nullptr);
-
-	//Take old skeleton data for new static data structure
-
-	FLiveLinkSkeletonStaticData* SkeletonData = StaticData->Cast<FLiveLinkSkeletonStaticData>();
-	SkeletonData->BoneNames = MoveTemp(RefSkeleton_DEPRECATED.BoneNames);
-	SkeletonData->BoneParents = MoveTemp(RefSkeleton_DEPRECATED.BoneParents);
-	SkeletonData->PropertyNames = MoveTemp(CurveNames_DEPRECATED);
+	//Verify if there are any bones for this subject. No bones means only curves were present so create a basic role in this case.
+	if (RefSkeleton_DEPRECATED.BoneNames.Num() == 0)
+	{
+		SubjectPreset.Role = ULiveLinkBasicRole::StaticClass();
+		StaticData->InitializeWith(FLiveLinkBaseStaticData::StaticStruct(), nullptr);
 
-	//Create subsection that will manage this data
-	UMovieSceneLiveLinkSubSectionAnimation* AnimationSubSection = NewObject<UMovieSceneLiveLinkSubSectionAnimation>(this, UMovieSceneLiveLinkSubSectionAnimation::StaticClass(), NAME_None, RF_Transactional);
-	AnimationSubSection->Initialize(SubjectPreset.Role, StaticData);
+		//Data initialization will be done in the common part
+	}
+	else
+	{
+		SubjectPreset.Role = ULiveLinkAnimationRole::StaticClass();
+		StaticData->InitializeWith(FLiveLinkSkeletonStaticData::StaticStruct(), nullptr);
+		
+		//Take old skeleton data for new static data structure
+
+		FLiveLinkSkeletonStaticData* SkeletonData = StaticData->Cast<FLiveLinkSkeletonStaticData>();
+		SkeletonData->BoneNames = MoveTemp(RefSkeleton_DEPRECATED.BoneNames);
+		SkeletonData->BoneParents = MoveTemp(RefSkeleton_DEPRECATED.BoneParents);
+
+		//Create subsection that will manage this data
+		UMovieSceneLiveLinkSubSectionAnimation* AnimationSubSection = NewObject<UMovieSceneLiveLinkSubSectionAnimation>(this, UMovieSceneLiveLinkSubSectionAnimation::StaticClass(), NAME_None, RF_Transactional);
+		AnimationSubSection->Initialize(SubjectPreset.Role, StaticData);
+
+		//Convert transforms to float data channels
+		const int32 TransformCount = SkeletonData->BoneNames.Num() * 9;
+		ensure(AnimationSubSection->SubSectionData.Properties[0].FloatChannel.Num() == TransformCount);
+		AnimationSubSection->SubSectionData.Properties[0].FloatChannel = TArray<FMovieSceneFloatChannel>(PropertyFloatChannels_DEPRECATED.GetData(), TransformCount);
+		PropertyFloatChannels_DEPRECATED.RemoveAtSwap(0, TransformCount);
+		SubSections.Add(AnimationSubSection);
+	}
 
-	//Convert float data channels
-	const int32 TransformCount = SkeletonData->BoneNames.Num() * 9;
-	check(AnimationSubSection->SubSectionData.Properties[0].FloatChannel.Num() == TransformCount);
-	AnimationSubSection->SubSectionData.Properties[0].FloatChannel = TArray<FMovieSceneFloatChannel>(PropertyFloatChannels_DEPRECATED.GetData(), TransformCount);
-	PropertyFloatChannels_DEPRECATED.RemoveAtSwap(0, TransformCount);
-	SubSections.Add(AnimationSubSection);
-	
-	//Basic role sub section to handle properties
-	UMovieSceneLiveLinkSubSectionBasicRole* BasicRoleSubSection = NewObject<UMovieSceneLiveLinkSubSectionBasicRole>(this, UMovieSceneLiveLinkSubSectionBasicRole::StaticClass(), NAME_None, RF_Transactional);
-	BasicRoleSubSection->Initialize(SubjectPreset.Role, StaticData);
-	check(BasicRoleSubSection->SubSectionData.Properties[0].FloatChannel.Num() == PropertyFloatChannels_DEPRECATED.Num());
-	BasicRoleSubSection->SubSectionData.Properties[0].FloatChannel = MoveTemp(PropertyFloatChannels_DEPRECATED);
-	SubSections.Add(BasicRoleSubSection);
+	// If there were curves, add a basic section handler
+	if (CurveNames_DEPRECATED.Num() > 0)
+	{
+		FLiveLinkBaseStaticData* BaseStaticData = StaticData->Cast<FLiveLinkBaseStaticData>();
+		BaseStaticData->PropertyNames = MoveTemp(CurveNames_DEPRECATED);
+		
+		//Basic role sub section to handle properties
+		UMovieSceneLiveLinkSubSectionBasicRole* BasicRoleSubSection = NewObject<UMovieSceneLiveLinkSubSectionBasicRole>(this, UMovieSceneLiveLinkSubSectionBasicRole::StaticClass(), NAME_None, RF_Transactional);
+		BasicRoleSubSection->Initialize(SubjectPreset.Role, StaticData);
+		ensure(BasicRoleSubSection->SubSectionData.Properties[0].FloatChannel.Num() == PropertyFloatChannels_DEPRECATED.Num());
+		BasicRoleSubSection->SubSectionData.Properties[0].FloatChannel = MoveTemp(PropertyFloatChannels_DEPRECATED);
+		SubSections.Add(BasicRoleSubSection);
+	}
 
 	//Generic properties sub section because all sections have it
 	UMovieSceneLiveLinkSubSectionProperties* PropertiesSubSection = NewObject<UMovieSceneLiveLinkSubSectionProperties>(this, UMovieSceneLiveLinkSubSectionProperties::StaticClass(), NAME_None, RF_Transactional);
diff --git a/Engine/Source/Runtime/LiveLinkInterface/Private/ILiveLinkSubject.cpp b/Engine/Source/Runtime/LiveLinkInterface/Private/ILiveLinkSubject.cpp
--- a/Engine/Source/Runtime/LiveLinkInterface/Private/ILiveLinkSubject.cpp
+++ b/Engine/Source/Runtime/LiveLinkInterface/Private/ILiveLinkSubject.cpp
@@ -37,7 +37,13 @@ bool ILiveLinkSubject::EvaluateFrame(TSubclassOf<ULiveLinkRole> InDesiredRole, F
 		return true;
 	}
 
-	return Translate(this, InDesiredRole, GetFrameSnapshot().StaticData, GetFrameSnapshot().FrameData, OutFrame);
+	const bool bSuccess = Translate(this, InDesiredRole, GetFrameSnapshot().StaticData, GetFrameSnapshot().FrameData, OutFrame);
+	if (!bSuccess)
+	{
+		UE_LOG(LogLiveLinkSubject, Warning, TEXT("Can't evaluate frame for subject '%s' for incompatible role '%s. Subject has the role '%s' and no translators could work."), *GetSubjectKey().SubjectName.ToString(), *InDesiredRole->GetName(), *Role->GetName());
+	}
+
+	return bSuccess;
 }
 
 
diff --git a/Engine/Source/Runtime/LiveLinkInterface/Private/LiveLinkSourceFactory.cpp b/Engine/Source/Runtime/LiveLinkInterface/Private/LiveLinkSourceFactory.cpp
--- a/Engine/Source/Runtime/LiveLinkInterface/Private/LiveLinkSourceFactory.cpp
+++ b/Engine/Source/Runtime/LiveLinkInterface/Private/LiveLinkSourceFactory.cpp
@@ -7,3 +7,13 @@ TSharedPtr<SWidget> ULiveLinkSourceFactory::BuildCreationPanel(FOnLiveLinkSource
 {
 	return TSharedPtr<SWidget>();
 }
+
+TSharedPtr<SWidget> ULiveLinkSourceFactory::CreateSourceCreationPanel()
+{
+	return TSharedPtr<SWidget>();
+}
+
+TSharedPtr<ILiveLinkSource> ULiveLinkSourceFactory::OnSourceCreationPanelClosed(bool bMakeSource)
+{
+	return TSharedPtr<ILiveLinkSource>();
+}
diff --git a/Engine/Source/Runtime/LiveLinkInterface/Private/Roles/LiveLinkAnimationBlueprintStructs.cpp b/Engine/Source/Runtime/LiveLinkInterface/Private/Roles/LiveLinkAnimationBlueprintStructs.cpp
--- a/Engine/Source/Runtime/LiveLinkInterface/Private/Roles/LiveLinkAnimationBlueprintStructs.cpp
+++ b/Engine/Source/Runtime/LiveLinkInterface/Private/Roles/LiveLinkAnimationBlueprintStructs.cpp
@@ -348,6 +348,24 @@ void FSubjectFrameHandle::GetTransformByName(FName InTransformName, FLiveLinkTra
 	OutLiveLinkTransform.SetTransformIndex(TransformIndex);
 };
 
+FLiveLinkSkeletonStaticData* FSubjectFrameHandle::GetSourceSkeletonStaticData()
+{
+	if (CachedFrame.IsValid())
+	{
+		return &CachedFrame->GetSourceSkeletonData();
+	}
+	return nullptr;
+}
+
+FLiveLinkAnimationFrameData* FSubjectFrameHandle::GetSourceAnimationFrameData()
+{
+	if (CachedFrame.IsValid())
+	{
+		return &CachedFrame->GetSourceAnimationFrameData();
+	}
+	return nullptr;
+}
+
 void FSubjectFrameHandle::SetCachedFrame(TSharedPtr<FCachedSubjectFrame> InCachedFrame)
 {
 	CachedFrame = InCachedFrame;
diff --git a/Engine/Source/Runtime/LiveLinkInterface/Private/Roles/LiveLinkTransformRole.cpp b/Engine/Source/Runtime/LiveLinkInterface/Private/Roles/LiveLinkTransformRole.cpp
--- a/Engine/Source/Runtime/LiveLinkInterface/Private/Roles/LiveLinkTransformRole.cpp
+++ b/Engine/Source/Runtime/LiveLinkInterface/Private/Roles/LiveLinkTransformRole.cpp
@@ -5,6 +5,40 @@
 
 #define LOCTEXT_NAMESPACE "LiveLinkRole"
 
+namespace
+{
+	EAxis::Type LiveLinkAxisToMatrixAxis(ELiveLinkAxis InAxis)
+	{
+		switch (InAxis)
+		{
+		case ELiveLinkAxis::X:
+		case ELiveLinkAxis::XNeg:
+			return EAxis::X;
+		case ELiveLinkAxis::Y:
+		case ELiveLinkAxis::YNeg:
+			return EAxis::Y;
+		case ELiveLinkAxis::Z:
+		case ELiveLinkAxis::ZNeg:
+			return EAxis::Z;
+		default:
+			ensureMsgf(false, TEXT("EliveLinkAxis has no default value!"));
+			return EAxis::None;
+		}
+	}
+
+	float AxisSign(ELiveLinkAxis InAxis)
+	{
+		if (InAxis == ELiveLinkAxis::X || InAxis == ELiveLinkAxis::Y || InAxis == ELiveLinkAxis::Z)
+		{
+			return 1.0f;
+		}
+		else
+		{
+			return -1.0f;
+		}
+	}
+}
+
 UScriptStruct* ULiveLinkTransformRole::GetStaticDataStruct() const
 {
 	return FLiveLinkTransformStaticData::StaticStruct();
@@ -42,4 +76,63 @@ FText ULiveLinkTransformRole::GetDisplayName() const
 	return LOCTEXT("TransformRole", "Transform");
 }
 
+/**
+ * ULiveLinkAxisSwitchPreProcessor::FLiveLinkAxisSwitchPreProcessorWorker
+ */
+TSubclassOf<ULiveLinkRole> ULiveLinkAxisSwitchPreProcessor::FLiveLinkAxisSwitchPreProcessorWorker::GetRole() const
+{
+	return ULiveLinkTransformRole::StaticClass();
+}
+
+bool ULiveLinkAxisSwitchPreProcessor::FLiveLinkAxisSwitchPreProcessorWorker::PreProcessFrame(FLiveLinkFrameDataStruct& InOutFrame) const
+{
+	FLiveLinkTransformFrameData& TransformData = *InOutFrame.Cast<FLiveLinkTransformFrameData>();
+	const FMatrix InMatrix = TransformData.Transform.ToMatrixWithScale();
+
+	FVector DestAxisX = InMatrix.GetScaledAxis(LiveLinkAxisToMatrixAxis(AxisX)) * AxisSign(AxisX);
+	FVector DestAxisY = InMatrix.GetScaledAxis(LiveLinkAxisToMatrixAxis(AxisY)) * AxisSign(AxisY);
+	FVector DestAxisZ = InMatrix.GetScaledAxis(LiveLinkAxisToMatrixAxis(AxisZ)) * AxisSign(AxisZ);
+
+	FMatrix Result(InMatrix);
+	Result.SetAxes(&DestAxisX, &DestAxisY, &DestAxisZ);
+
+	TransformData.Transform.SetFromMatrix(Result);
+	return true;
+}
+
+/**
+ * ULiveLinkAxisSwitchPreProcessor
+ */
+TSubclassOf<ULiveLinkRole> ULiveLinkAxisSwitchPreProcessor::GetRole() const
+{
+	return ULiveLinkTransformRole::StaticClass();
+}
+
+ULiveLinkFramePreProcessor::FWorkerSharedPtr ULiveLinkAxisSwitchPreProcessor::FetchWorker()
+{
+	if (!Instance.IsValid())
+	{
+		Instance = MakeShared<FLiveLinkAxisSwitchPreProcessorWorker, ESPMode::ThreadSafe>();
+		Instance->AxisX = AxisX;
+		Instance->AxisY = AxisY;
+		Instance->AxisZ = AxisZ;
+	}
+
+	return Instance;
+}
+
+#if WITH_EDITOR
+void ULiveLinkAxisSwitchPreProcessor::PostEditChangeChainProperty(struct FPropertyChangedChainEvent& PropertyChangedEvent)
+{
+	if (PropertyChangedEvent.Property->GetFName() == GET_MEMBER_NAME_CHECKED(ThisClass, AxisX) ||
+		PropertyChangedEvent.Property->GetFName() == GET_MEMBER_NAME_CHECKED(ThisClass, AxisY) ||
+		PropertyChangedEvent.Property->GetFName() == GET_MEMBER_NAME_CHECKED(ThisClass, AxisZ))
+	{
+		Instance.Reset();
+	}
+
+	Super::PostEditChangeChainProperty(PropertyChangedEvent);
+}
+#endif //WITH_EDITOR
+
 #undef LOCTEXT_NAMESPACE
diff --git a/Engine/Source/Runtime/LiveLinkInterface/Public/LiveLinkRole.h b/Engine/Source/Runtime/LiveLinkInterface/Public/LiveLinkRole.h
--- a/Engine/Source/Runtime/LiveLinkInterface/Public/LiveLinkRole.h
+++ b/Engine/Source/Runtime/LiveLinkInterface/Public/LiveLinkRole.h
@@ -34,10 +34,10 @@ struct LIVELINKINTERFACE_API FLiveLinkSubjectRepresentation
 {
 	GENERATED_BODY()
 
-	UPROPERTY(EditAnywhere, Category="Live Link")
+	UPROPERTY(EditAnywhere, BlueprintReadWrite, Category="Live Link")
 	FLiveLinkSubjectName Subject;
 
-	UPROPERTY(EditAnywhere, Category="Live Link")
+	UPROPERTY(EditAnywhere, BlueprintReadWrite, Category="Live Link")
 	TSubclassOf<ULiveLinkRole> Role;
 
 	bool operator==(const FLiveLinkSubjectRepresentation& Other) const { return Subject == Other.Subject && Role == Other.Role; } 
diff --git a/Engine/Source/Runtime/LiveLinkInterface/Public/LiveLinkSourceFactory.h b/Engine/Source/Runtime/LiveLinkInterface/Public/LiveLinkSourceFactory.h
--- a/Engine/Source/Runtime/LiveLinkInterface/Public/LiveLinkSourceFactory.h
+++ b/Engine/Source/Runtime/LiveLinkInterface/Public/LiveLinkSourceFactory.h
@@ -11,15 +11,27 @@ class ILiveLinkSource;
 class ILiveLinkClient;
 class SWidget;
 
+/**
+ * Base class of factory that creates Source. A source can be created in the editor via the Live Link panel or at runtime via a connection string.
+ */
 UCLASS(Abstract)
 class LIVELINKINTERFACE_API ULiveLinkSourceFactory : public UObject
 {
 	GENERATED_BODY()
 
 public:
+	/** The name of the menu item (of any EMenuType) */
 	virtual FText GetSourceDisplayName() const PURE_VIRTUAL( ULiveLinkSourceFactory::GetSourceDisplayName, return FText(); );
+
+	/** The tooltip of the menu item (of any EMenyType) */
 	virtual FText GetSourceTooltip() const PURE_VIRTUAL(ULiveLinkSourceFactory::GetSourceTooltip, return FText(); );
 
+	UE_DEPRECATED(4.23, "CreateSourceCreationPanel is deprecated. LiveLinkSourceFactory can now be used at runtime. This factory won't work until it's been updated.")
+	virtual TSharedPtr<SWidget> CreateSourceCreationPanel();
+
+	UE_DEPRECATED(4.23, "OnSourceCreationPanelClosed is deprecated. LiveLinkSourceFactory can now be used at runtime. The factory won't work until it's been updated.")
+	virtual TSharedPtr<ILiveLinkSource> OnSourceCreationPanelClosed(bool bMakeSource);
+
 	enum class EMenuType
 	{
 		SubPanel,	// In the UI, a sub menu will used
@@ -34,7 +46,12 @@ public:
 	virtual EMenuType GetMenuType() const { return EMenuType::Disabled; }
 
 	DECLARE_DELEGATE_TwoParams(FOnLiveLinkSourceCreated, TSharedPtr<ILiveLinkSource> /*Created source*/, FString /*ConnectionString*/);
-	/** Build a UI that will create a Live Link source. */
+
+	/**
+	 * Create a widget responsible for the creation of a Live Link source.
+	 * @param OnLiveLinkSourceCreated Callback to call when the source has been created by the custom UI.
+	 * @return The subpanel UI created by the factory.
+	 */
 	virtual TSharedPtr<SWidget> BuildCreationPanel(FOnLiveLinkSourceCreated OnLiveLinkSourceCreated) const;
 
 	/** Create a new source from a ConnectionString */
diff --git a/Engine/Source/Runtime/LiveLinkInterface/Public/LiveLinkTypes.h b/Engine/Source/Runtime/LiveLinkInterface/Public/LiveLinkTypes.h
--- a/Engine/Source/Runtime/LiveLinkInterface/Public/LiveLinkTypes.h
+++ b/Engine/Source/Runtime/LiveLinkInterface/Public/LiveLinkTypes.h
@@ -11,6 +11,7 @@
 #include "UObject/StructOnScope.h"
 #include "Serialization/Archive.h"
 #include "Templates/HasGetTypeHash.h"
+#include "UObject/PropertyPortFlags.h"
 
 #include "LiveLinkTypes.generated.h"
 
@@ -26,7 +27,7 @@ public:
 	FLiveLinkSubjectName(EName InName) : Name(InName) {}
 
 	// Name of the subject
-	UPROPERTY(EditAnywhere, Category="Live Link")
+	UPROPERTY(EditAnywhere, BlueprintReadWrite, Category="Live Link")
 	FName Name;
 
 	bool IsNone() const { return Name.IsNone(); }
@@ -38,22 +39,28 @@ public:
 
 	bool operator==(const FLiveLinkSubjectName& Other) const { return Name == Other.Name; }
 	bool operator==(const FName Other) const { return Name == Other; }
+
+	friend FArchive& operator<<(FArchive& Ar, FLiveLinkSubjectName& InSubjectName)
+	{
+		Ar << InSubjectName.Name;
+		return Ar;
+	}
 };
 template <> struct TModels<CGetTypeHashable, FLiveLinkSubjectName> { enum { Value = TModels<CGetTypeHashable, FName>::Value }; };
 
 
 // Structure that identifies an individual subject
-USTRUCT()
+USTRUCT(BlueprintType)
 struct LIVELINKINTERFACE_API FLiveLinkSubjectKey
 {
 	GENERATED_BODY()
 
 	// The guid for this subjects source
-	UPROPERTY()
+	UPROPERTY(BlueprintReadOnly, Category="LiveLink")
 	FGuid Source;
 
 	// The Name of this subject
-	UPROPERTY()
+	UPROPERTY(BlueprintReadOnly, Category="LiveLink")
 	FLiveLinkSubjectName SubjectName;
 
 	FLiveLinkSubjectKey() = default;
@@ -61,6 +68,13 @@ struct LIVELINKINTERFACE_API FLiveLinkSubjectKey
 	FLiveLinkSubjectKey(const FLiveLinkSubjectKey& Rhs) : Source(Rhs.Source), SubjectName(Rhs.SubjectName) {}
 
 	bool operator== (const FLiveLinkSubjectKey& Other) const { return SubjectName == Other.SubjectName && Source == Other.Source; }
+
+	friend FArchive& operator<<(FArchive& Ar, FLiveLinkSubjectKey& InSubjectKey)
+	{
+		Ar << InSubjectKey.Source;
+		Ar << InSubjectKey.SubjectName;
+		return Ar;
+	}
 };
 
 
@@ -297,6 +311,19 @@ public:
 		return *this;
 	}
 
+	bool operator==(const FLiveLinkBaseDataStruct& Other) const
+	{
+		if (Other.GetStruct() == GetStruct())
+		{
+			if (GetBaseData())
+			{
+				return GetStruct()->CompareScriptStruct(GetBaseData(), Other.GetBaseData(), PPF_None);
+			}
+			return Other.GetBaseData() == nullptr; // same struct and both uninitialized
+		}
+		return false;
+	}
+
 	friend FArchive& operator<<(FArchive& Ar, FLiveLinkBaseDataStruct& InStruct)
 	{
 		if (Ar.IsLoading())
diff --git a/Engine/Source/Runtime/LiveLinkInterface/Public/Roles/LiveLinkAnimationBlueprintStructs.h b/Engine/Source/Runtime/LiveLinkInterface/Public/Roles/LiveLinkAnimationBlueprintStructs.h
--- a/Engine/Source/Runtime/LiveLinkInterface/Public/Roles/LiveLinkAnimationBlueprintStructs.h
+++ b/Engine/Source/Runtime/LiveLinkInterface/Public/Roles/LiveLinkAnimationBlueprintStructs.h
@@ -58,6 +58,12 @@ struct FCachedSubjectFrame
 
 	int32 GetRootIndex();
 
+	FLiveLinkSkeletonStaticData& GetSourceSkeletonData() { return SourceSkeletonData; }
+	const FLiveLinkSkeletonStaticData& GetSourceSkeletonData() const { return SourceSkeletonData; }
+
+	FLiveLinkAnimationFrameData& GetSourceAnimationFrameData() { return SourceAnimationFrameData; }
+	const FLiveLinkAnimationFrameData& GetSourceAnimationFrameData() const { return SourceAnimationFrameData; }
+
 private:
 	FLiveLinkSkeletonStaticData SourceSkeletonData;
 	FLiveLinkAnimationFrameData SourceAnimationFrameData;
@@ -131,6 +137,10 @@ struct LIVELINKINTERFACE_API FSubjectFrameHandle : public FLiveLinkBaseBlueprint
 
 	void SetCachedFrame(TSharedPtr<FCachedSubjectFrame> InCachedFrame);
 
+	FLiveLinkSkeletonStaticData* GetSourceSkeletonStaticData();
+
+	FLiveLinkAnimationFrameData* GetSourceAnimationFrameData();
+
 private:
 	TSharedPtr<FCachedSubjectFrame> CachedFrame;
 };
diff --git a/Engine/Source/Runtime/LiveLinkInterface/Public/Roles/LiveLinkTransformRole.h b/Engine/Source/Runtime/LiveLinkInterface/Public/Roles/LiveLinkTransformRole.h
--- a/Engine/Source/Runtime/LiveLinkInterface/Public/Roles/LiveLinkTransformRole.h
+++ b/Engine/Source/Runtime/LiveLinkInterface/Public/Roles/LiveLinkTransformRole.h
@@ -2,6 +2,7 @@
 
 #pragma once
 
+#include "LiveLinkFramePreProcessor.h"
 #include "LiveLinkTypes.h"
 #include "Roles/LiveLinkBasicRole.h"
 #include "LiveLinkTransformRole.generated.h"
@@ -23,3 +24,61 @@ public:
 
 	virtual FText GetDisplayName() const override;
 };
+
+UENUM()
+enum class ELiveLinkAxis : uint8
+{
+	X UMETA(DisplayName = "X-Axis"),
+	Y UMETA(DisplayName = "Y-Axis"),
+	Z UMETA(DisplayName = "Z-Axis"),
+	XNeg UMETA(DisplayName = "-X-Axis"),
+	YNeg UMETA(DisplayName = "-Y-Axis"),
+	ZNeg UMETA(DisplayName = "-Z-Axis"),
+};
+
+/**
+ * Allows to switch any axis of an incoming transform with another axis.
+ * @note For example the Z-Axis of an incoming transform can be set to the (optionally negated) Y-Axis of the transform in UE.
+ * @note This implies that translation, rotation and scale will be affected by switching an axis.
+ */
+UCLASS(meta = (DisplayName = "Axis Switch"))
+class LIVELINKINTERFACE_API ULiveLinkAxisSwitchPreProcessor : public ULiveLinkFramePreProcessor
+{
+	GENERATED_BODY()
+
+public:
+	class FLiveLinkAxisSwitchPreProcessorWorker : public ILiveLinkFramePreProcessorWorker
+	{
+	public:
+		ELiveLinkAxis AxisX = ELiveLinkAxis::X;
+		ELiveLinkAxis AxisY = ELiveLinkAxis::Y;
+		ELiveLinkAxis AxisZ = ELiveLinkAxis::Z;
+
+		virtual TSubclassOf<ULiveLinkRole> GetRole() const override;
+		virtual bool PreProcessFrame(FLiveLinkFrameDataStruct& InOutFrame) const override;
+	};
+
+protected:
+	UPROPERTY(EditAnywhere, Category = "LiveLink")
+	ELiveLinkAxis AxisX = ELiveLinkAxis::X;
+	UPROPERTY(EditAnywhere, Category = "LiveLink")
+	ELiveLinkAxis AxisY = ELiveLinkAxis::Y;
+	UPROPERTY(EditAnywhere, Category = "LiveLink")
+	ELiveLinkAxis AxisZ = ELiveLinkAxis::Z;
+
+
+public:
+	virtual TSubclassOf<ULiveLinkRole> GetRole() const override;
+	virtual ULiveLinkFramePreProcessor::FWorkerSharedPtr FetchWorker() override;
+
+public:
+	//~ UObject interface
+#if WITH_EDITOR
+	virtual void PostEditChangeChainProperty(struct FPropertyChangedChainEvent& PropertyChangedEvent) override;
+#endif //WITH_EDITOR
+	//~ End UObject interface
+
+private:
+	TSharedPtr<FLiveLinkAxisSwitchPreProcessorWorker, ESPMode::ThreadSafe> Instance;
+};
+
