Fix LiveLink Preset blueprint functions

diff --git a/Engine/Plugins/Animation/LiveLink/Source/LiveLink/Public/LiveLinkPreset.h b/Engine/Plugins/Animation/LiveLink/Source/LiveLink/Public/LiveLinkPreset.h
--- a/Engine/Plugins/Animation/LiveLink/Source/LiveLink/Public/LiveLinkPreset.h
+++ b/Engine/Plugins/Animation/LiveLink/Source/LiveLink/Public/LiveLinkPreset.h
@@ -9,7 +9,7 @@
 #include "LiveLinkPreset.generated.h"
 
 
-UCLASS()
+UCLASS(BlueprintType)
 class LIVELINK_API ULiveLinkPreset : public UObject
 {
 	GENERATED_BODY()
@@ -32,7 +32,7 @@ public:
 	 * Remove all previous sources and subjects and add the sources and subjects from this preset.
 	 * @return True is all sources and subjects from this preset could be created and added.
 	 */
-	UFUNCTION(BlueprintCallable, Category="LiveLink")
+	UFUNCTION(BlueprintCallable, BlueprintPure = false, Category="LiveLink")
 	bool ApplyToClient() const;
 
 	/** Reset this preset and build the list of sources and subjects from the client. */
