Merging fix from Dev-Enterprise CL#7074406 and #7075748

- Fix for LiveLink recording crash when no transforms in animation role
- Fix for animation deprecation path not adding curve names to static data
- Fix for livelink interpolation mode crash

diff --git a/Engine/Plugins/Animation/LiveLink/Source/LiveLink/Private/LiveLinkClient.cpp b/Engine/Plugins/Animation/LiveLink/Source/LiveLink/Private/LiveLinkClient.cpp
--- a/Engine/Plugins/Animation/LiveLink/Source/LiveLink/Private/LiveLinkClient.cpp
+++ b/Engine/Plugins/Animation/LiveLink/Source/LiveLink/Private/LiveLinkClient.cpp
@@ -1130,6 +1130,7 @@ void FLiveLinkClient_Base_DEPRECATED::PushSubjectData(FGuid InSourceGuid, FName
 			if (NumberOfPropertyNames == 0 && InFrameData.CurveElements.Num() > 0)
 			{
 				UpdateForAnimationStatic(AnimationStaticData->PropertyNames, InFrameData.CurveElements);
+				NumberOfPropertyNames = AnimationStaticData->PropertyNames.Num();
 			}
 		}
 
diff --git a/Engine/Plugins/Animation/LiveLink/Source/LiveLink/Private/LiveLinkSubject.h b/Engine/Plugins/Animation/LiveLink/Source/LiveLink/Private/LiveLinkSubject.h
--- a/Engine/Plugins/Animation/LiveLink/Source/LiveLink/Private/LiveLinkSubject.h
+++ b/Engine/Plugins/Animation/LiveLink/Source/LiveLink/Private/LiveLinkSubject.h
@@ -54,8 +54,8 @@ public:
 	virtual FLiveLinkSubjectKey GetSubjectKey() const { return SubjectKey; }
 	virtual TSubclassOf<ULiveLinkRole> GetRole() const override { return Role; }
 	virtual bool HasValidFrameSnapshot() const override;
-	virtual FLiveLinkStaticDataStruct& GetStaticData() override { return FrameSnapshot.StaticData; }
-	virtual const FLiveLinkStaticDataStruct& GetStaticData() const override { return FrameSnapshot.StaticData; }
+	virtual FLiveLinkStaticDataStruct& GetStaticData() override { return StaticData; }
+	virtual const FLiveLinkStaticDataStruct& GetStaticData() const override { return StaticData; }
 	virtual const TArray<ULiveLinkFrameTranslator::FWorkerSharedPtr> GetFrameTranslators() const override { return FrameTranslators; }
 protected:
 	virtual const FLiveLinkSubjectFrameData& GetFrameSnapshot() const { return FrameSnapshot; }
diff --git a/Engine/Plugins/Animation/LiveLink/Source/LiveLinkMovieScene/Private/MovieScene/MovieSceneLiveLinkSubSectionAnimation.cpp b/Engine/Plugins/Animation/LiveLink/Source/LiveLinkMovieScene/Private/MovieScene/MovieSceneLiveLinkSubSectionAnimation.cpp
--- a/Engine/Plugins/Animation/LiveLink/Source/LiveLinkMovieScene/Private/MovieScene/MovieSceneLiveLinkSubSectionAnimation.cpp
+++ b/Engine/Plugins/Animation/LiveLink/Source/LiveLinkMovieScene/Private/MovieScene/MovieSceneLiveLinkSubSectionAnimation.cpp
@@ -80,6 +80,11 @@ void UMovieSceneLiveLinkSubSectionAnimation::CreatePropertiesChannel()
 
 	const FName TransformsPropertyName = GET_MEMBER_NAME_CHECKED(FLiveLinkAnimationFrameData, Transforms);
 	const int32 TransformCount = SkeletonData->BoneNames.Num();
+	if (TransformCount <= 0)
+	{
+		return;
+	}
+
 	SubSectionData.Properties.SetNum(1);
 	SubSectionData.Properties[0].PropertyName = TransformsPropertyName;
 
@@ -92,12 +97,18 @@ void UMovieSceneLiveLinkSubSectionAnimation::RecordFrame(FFrameNumber InFrameNum
 	const FLiveLinkAnimationFrameData* AnimationFrameData = InFrameData.Cast<FLiveLinkAnimationFrameData>();
 	check(AnimationFrameData);
 
-	TransformHandler->RecordFrame(InFrameNumber, *FLiveLinkAnimationFrameData::StaticStruct(), InFrameData.GetBaseData());
+	if (TransformHandler.IsValid())
+	{
+		TransformHandler->RecordFrame(InFrameNumber, *FLiveLinkAnimationFrameData::StaticStruct(), InFrameData.GetBaseData());
+	}
 }
 
 void UMovieSceneLiveLinkSubSectionAnimation::FinalizeSection(bool bInReduceKeys, const FKeyDataOptimizationParams& InOptimizationParams)
 {
-	TransformHandler->Finalize(bInReduceKeys, InOptimizationParams);
+	if (TransformHandler.IsValid())
+	{
+		TransformHandler->Finalize(bInReduceKeys, InOptimizationParams);
+	}
 }
 
 bool UMovieSceneLiveLinkSubSectionAnimation::IsRoleSupported(const TSubclassOf<ULiveLinkRole>& RoleToSupport) const
diff --git a/Engine/Source/Runtime/LiveLinkInterface/Private/Roles/LiveLinkAnimationRole.cpp b/Engine/Source/Runtime/LiveLinkInterface/Private/Roles/LiveLinkAnimationRole.cpp
--- a/Engine/Source/Runtime/LiveLinkInterface/Private/Roles/LiveLinkAnimationRole.cpp
+++ b/Engine/Source/Runtime/LiveLinkInterface/Private/Roles/LiveLinkAnimationRole.cpp
@@ -125,6 +125,9 @@ void ULiveLinkAnimationFrameInterpolateProcessor::FLiveLinkAnimationFrameInterpo
 		{
 			OutLastFrameIndexUsed = FrameDataIndexA;
 
+			//Initialize the output frame for animation. It will be filled during blended values copied
+			OutBlendedFrame.FrameData.InitializeWith(FLiveLinkAnimationFrameData::StaticStruct(), nullptr);
+
 			const FLiveLinkFrameDataStruct& FrameDataA = InSourceFrames[FrameDataIndexA];
 			const FLiveLinkFrameDataStruct& FrameDataB = InSourceFrames[FrameDataIndexB];
 
