Move LiveLink graph node to their own developer module to work with -game

diff --git a/Engine/Plugins/Animation/LiveLink/Config/BaseLiveLink.ini b/Engine/Plugins/Animation/LiveLink/Config/BaseLiveLink.ini
--- a/Engine/Plugins/Animation/LiveLink/Config/BaseLiveLink.ini
+++ b/Engine/Plugins/Animation/LiveLink/Config/BaseLiveLink.ini
@@ -12,3 +12,8 @@
  +ClassRedirects=(OldName="MovieSceneLiveLinkSection", NewName="/Script/LiveLinkMovieScene.MovieSceneLiveLinkSection")
  +StructRedirects=(OldName="MovieSceneLiveLinkSectionTemplate", NewName="/Script/LiveLinkMovieScene.MovieSceneLiveLinkSectionTemplate")
 
+; LiveLink Graph
++ClassRedirects=(OldName="K2Node_EvaluateLiveLinkFrame", NewName="/Script/LiveLinkGraphNode.K2Node_EvaluateLiveLinkFrameWithSpecificRole")
++ClassRedirects=(OldName="K2Node_EvaluateLiveLinkFrameAtWorldTime", NewName="/Script/LiveLinkGraphNode.K2Node_EvaluateLiveLinkFrameAtWorldTime")
++ClassRedirects=(OldName="K2Node_EvaluateLiveLinkFrameAtSceneTime", NewName="/Script/LiveLinkGraphNode.K2Node_EvaluateLiveLinkFrameAtSceneTime")
++ClassRedirects=(OldName="AnimGraphNode_LiveLinkPose", NewName="/Script/LiveLinkGraphNode.AnimGraphNode_LiveLinkPose")
\ No newline at end of file
diff --git a/Engine/Plugins/Animation/LiveLink/LiveLink.uplugin b/Engine/Plugins/Animation/LiveLink/LiveLink.uplugin
--- a/Engine/Plugins/Animation/LiveLink/LiveLink.uplugin
+++ b/Engine/Plugins/Animation/LiveLink/LiveLink.uplugin
@@ -24,27 +24,32 @@
 		{
 			"Name": "LiveLink",
 			"Type": "Runtime",
-			"LoadingPhase": "PreDefault"
+			"LoadingPhase": "Default"
 		},
 		{
 			"Name": "LiveLinkComponents",
 			"Type": "Runtime",
-			"LoadingPhase": "PreDefault"
+			"LoadingPhase": "Default"
 		},
 		{
 			"Name": "LiveLinkEditor",
 			"Type": "Editor",
-			"LoadingPhase": "PreDefault"
+			"LoadingPhase": "Default"
+		},
+		{
+			"Name": "LiveLinkGraphNode",
+			"Type": "Developer",
+			"LoadingPhase": "Default"
 		},
 		{
 			"Name": "LiveLinkMovieScene",
 			"Type": "Runtime",
-			"LoadingPhase": "PreDefault"
+			"LoadingPhase": "Default"
 		},
 		{
 			"Name": "LiveLinkSequencer",
 			"Type": "Developer",
-			"LoadingPhase": "PreDefault"
+			"LoadingPhase": "Default"
 		}
 	]
 }
\ No newline at end of file
diff --git a/Engine/Plugins/Animation/LiveLink/Source/LiveLink/Public/LiveLinkInstance.h b/Engine/Plugins/Animation/LiveLink/Source/LiveLink/Public/LiveLinkInstance.h
--- a/Engine/Plugins/Animation/LiveLink/Source/LiveLink/Public/LiveLinkInstance.h
+++ b/Engine/Plugins/Animation/LiveLink/Source/LiveLink/Public/LiveLinkInstance.h
@@ -32,7 +32,7 @@ public:
 	virtual bool Evaluate(FPoseContext& Output) override;
 	virtual void UpdateAnimationNode(float DeltaSeconds) override;
 
-
+	UPROPERTY(EditAnywhere, Category = Settings)
 	FAnimNode_LiveLinkPose PoseNode;
 };
 
diff --git a/Engine/Plugins/Animation/LiveLink/Source/LiveLinkEditor/LiveLinkEditor.Build.cs b/Engine/Plugins/Animation/LiveLink/Source/LiveLinkEditor/LiveLinkEditor.Build.cs
--- a/Engine/Plugins/Animation/LiveLink/Source/LiveLinkEditor/LiveLinkEditor.Build.cs
+++ b/Engine/Plugins/Animation/LiveLink/Source/LiveLinkEditor/LiveLinkEditor.Build.cs
@@ -30,6 +30,7 @@ namespace UnrealBuildTool.Rules
 					"InputCore",
 					"KismetCompiler",
 					"GraphEditor",
+					"LiveLinkGraphNode",
 					"LiveLinkMessageBusFramework",
 					"LiveLinkMovieScene",
 					"Persona",
diff --git a/Engine/Plugins/Animation/LiveLink/Source/LiveLinkEditor/Private/AnimGraphNode_LiveLinkPose.cpp b/Engine/Plugins/Animation/LiveLink/Source/LiveLinkEditor/Private/AnimGraphNode_LiveLinkPose.cpp
deleted file mode 100644
diff --git a/Engine/Plugins/Animation/LiveLink/Source/LiveLinkEditor/Private/AnimGraphNode_LiveLinkPose.h b/Engine/Plugins/Animation/LiveLink/Source/LiveLinkEditor/Private/AnimGraphNode_LiveLinkPose.h
deleted file mode 100644
diff --git a/Engine/Plugins/Animation/LiveLink/Source/LiveLinkEditor/Private/K2Node_EvaluateLiveLinkFrame.cpp b/Engine/Plugins/Animation/LiveLink/Source/LiveLinkEditor/Private/K2Node_EvaluateLiveLinkFrame.cpp
deleted file mode 100644
diff --git a/Engine/Plugins/Animation/LiveLink/Source/LiveLinkEditor/Public/K2Node_EvaluateLiveLinkFrame.h b/Engine/Plugins/Animation/LiveLink/Source/LiveLinkEditor/Public/K2Node_EvaluateLiveLinkFrame.h
deleted file mode 100644
diff --git a/Engine/Plugins/Animation/LiveLink/Source/LiveLinkGraphNode/LiveLinkGraphNode.Build.cs b/Engine/Plugins/Animation/LiveLink/Source/LiveLinkGraphNode/LiveLinkGraphNode.Build.cs
new file mode 100644
--- /dev/null
+++ b/Engine/Plugins/Animation/LiveLink/Source/LiveLinkGraphNode/LiveLinkGraphNode.Build.cs
@@ -0,0 +1,47 @@
+// Copyright 1998-2019 Epic Games, Inc. All Rights Reserved.
+
+namespace UnrealBuildTool.Rules
+{
+	public class LiveLinkGraphNode : ModuleRules
+	{
+		public LiveLinkGraphNode(ReadOnlyTargetRules Target) : base(Target)
+		{
+			PublicDependencyModuleNames.AddRange(
+				new string[]
+				{
+
+				}
+			);
+
+			PrivateDependencyModuleNames.AddRange(
+				new string[]
+				{
+					"AnimationCore",
+					"AnimGraphRuntime",
+					"Core",
+					"CoreUObject",
+					"Engine",
+					"InputCore",
+					"KismetCompiler",
+					"LiveLink",
+					"LiveLinkInterface",
+					"Persona",
+					"SlateCore",
+				}
+			);
+
+			if (Target.bBuildEditor == true)
+			{
+				PrivateDependencyModuleNames.AddRange(
+					new string[]
+					{
+						"UnrealEd",
+						"Kismet",
+						"AnimGraph",
+						"BlueprintGraph",
+					}
+				);
+			}
+		}
+	}
+}
diff --git a/Engine/Plugins/Animation/LiveLink/Source/LiveLinkGraphNode/Private/AnimGraphNode_LiveLinkPose.cpp b/Engine/Plugins/Animation/LiveLink/Source/LiveLinkGraphNode/Private/AnimGraphNode_LiveLinkPose.cpp
new file mode 100644
--- /dev/null
+++ b/Engine/Plugins/Animation/LiveLink/Source/LiveLinkGraphNode/Private/AnimGraphNode_LiveLinkPose.cpp
@@ -0,0 +1,54 @@
+// Copyright 1998-2019 Epic Games, Inc. All Rights Reserved.
+
+#include "AnimGraphNode_LiveLinkPose.h"
+
+#include "EdGraph/EdGraphSchema.h"
+
+#define LOCTEXT_NAMESPACE "LiveLinkAnimNode"
+
+FText UAnimGraphNode_LiveLinkPose::GetNodeTitle(ENodeTitleType::Type TitleType) const
+{
+	return LOCTEXT("NodeTitle", "Live Link Pose");
+}
+
+FText UAnimGraphNode_LiveLinkPose::GetTooltipText() const
+{
+	return LOCTEXT("NodeTooltip", "Retrieves the current pose associated with the supplied subject");
+}
+
+FText UAnimGraphNode_LiveLinkPose::GetMenuCategory() const
+{
+	return LOCTEXT("NodeCategory", "Live Link");
+}
+
+void UAnimGraphNode_LiveLinkPose::ConvertDeprecatedNode(UEdGraph* Graph, bool bOnlySafeChanges)
+{
+	//Find deprecated SubjectName pin and set new pin with its value
+	const FName OldPinName = TEXT("SubjectName"); //Variable now has _DEPRECATED appended so can't use GET_MEMBER_NAME_CHECKED
+	const FName NewPinName = GET_MEMBER_NAME_CHECKED(FAnimNode_LiveLinkPose, LiveLinkSubjectName);
+	UEdGraphPin** FoundPinPtr = Pins.FindByPredicate([OldPinName](const UEdGraphPin* Other) { return Other->PinName == OldPinName; });
+	if (FoundPinPtr != nullptr)
+	{
+		UEdGraphPin** FoundNewPinPtr = Pins.FindByPredicate([NewPinName](const UEdGraphPin* Other) { return Other->PinName == NewPinName; });
+		if (FoundNewPinPtr != nullptr)
+		{
+			UScriptStruct* StructType = FLiveLinkSubjectName::StaticStruct();
+			UEdGraphPin* OldPin = *FoundPinPtr;
+			UEdGraphPin* NewPin = *FoundNewPinPtr;
+
+			//Create new structure from old data
+			FLiveLinkSubjectName NewName;
+			NewName.Name = *OldPin->DefaultValue;
+
+			//Apply new name structure 
+			FString StringValue;
+			StructType->ExportText(StringValue, &NewName, nullptr, nullptr, EPropertyPortFlags::PPF_None, nullptr);
+			NewPin->GetSchema()->TrySetDefaultValue(*NewPin, StringValue);
+
+			//Update node data with graph data
+			Node.LiveLinkSubjectName = NewName;
+		}
+	}
+}
+
+#undef LOCTEXT_NAMESPACE
\ No newline at end of file
diff --git a/Engine/Plugins/Animation/LiveLink/Source/LiveLinkGraphNode/Private/AnimGraphNode_LiveLinkPose.h b/Engine/Plugins/Animation/LiveLink/Source/LiveLinkGraphNode/Private/AnimGraphNode_LiveLinkPose.h
new file mode 100644
--- /dev/null
+++ b/Engine/Plugins/Animation/LiveLink/Source/LiveLinkGraphNode/Private/AnimGraphNode_LiveLinkPose.h
@@ -0,0 +1,34 @@
+// Copyright 1998-2019 Epic Games, Inc. All Rights Reserved.
+
+#pragma once
+
+#include "AnimGraphNode_Base.h"
+#include "AnimNode_LiveLinkPose.h"
+
+#include "AnimGraphNode_LiveLinkPose.generated.h"
+
+
+
+UCLASS()
+class UAnimGraphNode_LiveLinkPose : public UAnimGraphNode_Base
+{
+	GENERATED_BODY()
+
+public:
+
+	//~ UEdGraphNode interface
+	virtual FText GetNodeTitle(ENodeTitleType::Type TitleType) const override;
+	virtual FText GetTooltipText() const override;
+	virtual FText GetMenuCategory() const;
+	//~ End of UEdGraphNode
+
+	//~ Begin UK2Node interface
+	virtual void ConvertDeprecatedNode(UEdGraph* Graph, bool bOnlySafeChanges) override;
+	//~ End UK2Node interface
+
+public:
+
+	UPROPERTY(EditAnywhere, Category = Settings)
+	FAnimNode_LiveLinkPose Node;
+
+};
diff --git a/Engine/Plugins/Animation/LiveLink/Source/LiveLinkGraphNode/Private/K2Node_EvaluateLiveLinkCustom.cpp b/Engine/Plugins/Animation/LiveLink/Source/LiveLinkGraphNode/Private/K2Node_EvaluateLiveLinkCustom.cpp
new file mode 100644
--- /dev/null
+++ b/Engine/Plugins/Animation/LiveLink/Source/LiveLinkGraphNode/Private/K2Node_EvaluateLiveLinkCustom.cpp
@@ -0,0 +1,107 @@
+// Copyright 1998-2019 Epic Games, Inc. All Rights Reserved.
+
+#include "K2Node_EvaluateLiveLinkCustom.h"
+
+#include "BlueprintActionDatabaseRegistrar.h"
+#include "BlueprintNodeSpawner.h"
+#include "EdGraphSchema_K2.h"
+#include "EditorCategoryUtils.h"
+#include "K2Node_CallFunction.h"
+#include "K2Node_IfThenElse.h"
+#include "Kismet2/BlueprintEditorUtils.h"
+#include "KismetCompiler.h"
+#include "LiveLinkRole.h"
+#include "LiveLinkTypes.h"
+#include "LiveLinkBlueprintLibrary.h"
+#include "UObject/PropertyPortFlags.h"
+
+
+#define LOCTEXT_NAMESPACE "K2Node_EvaluateLiveLinkCustom"
+
+
+namespace UK2Node_EvaluateLiveLinkFrameHelper
+{
+	static FName LiveLinkWorldTimePinName = "WorldTime";
+	static FName LiveLinkSceneTimePinName = "SceneTime";
+};
+
+void UK2Node_EvaluateLiveLinkFrameWithSpecificRole::AllocateDefaultPins()
+{
+	Super::AllocateDefaultPins();
+}
+
+void UK2Node_EvaluateLiveLinkFrameAtWorldTime::AllocateDefaultPins()
+{
+	Super::AllocateDefaultPins();
+
+	const UEdGraphSchema_K2* K2Schema = GetDefault<UEdGraphSchema_K2>();
+
+	// float pin
+	UEdGraphPin* LiveLinkWorldTimePin = CreatePin(EGPD_Input, UEdGraphSchema_K2::PC_Float, UK2Node_EvaluateLiveLinkFrameHelper::LiveLinkWorldTimePinName);
+	LiveLinkWorldTimePin->PinFriendlyName = LOCTEXT("LiveLinkWorldTimePin", "World Time");
+	SetPinToolTip(*LiveLinkWorldTimePin, LOCTEXT("LiveLinkWorldTimePinDescription", "The World Time the subject will be evaluated to"));
+}
+
+void UK2Node_EvaluateLiveLinkFrameAtSceneTime::AllocateDefaultPins()
+{
+	Super::AllocateDefaultPins();
+
+	const UEdGraphSchema_K2* K2Schema = GetDefault<UEdGraphSchema_K2>();
+
+	// Timecode pin
+	UScriptStruct* TimecodeScriptStruct = FindObject<UScriptStruct>(ANY_PACKAGE, TEXT("Timecode"), true);
+	check(TimecodeScriptStruct);
+	UEdGraphPin* LiveLinkSceneTimePin = CreatePin(EGPD_Input, UEdGraphSchema_K2::PC_Struct, TimecodeScriptStruct, UK2Node_EvaluateLiveLinkFrameHelper::LiveLinkSceneTimePinName);
+	LiveLinkSceneTimePin->PinFriendlyName = LOCTEXT("LiveLinkSceneTimePin", "Scene Time");
+	SetPinToolTip(*LiveLinkSceneTimePin, LOCTEXT("LiveLinkSceneTimePinDescription", "The Scene Time the subject will be evaluated to"));
+}
+
+FText UK2Node_EvaluateLiveLinkFrameWithSpecificRole::GetNodeTitle(ENodeTitleType::Type TitleType) const
+{
+	return LOCTEXT("EvaluateLiveLinkFrameTitle", "Evaluate Live Link Frame");
+}
+
+FText UK2Node_EvaluateLiveLinkFrameAtWorldTime::GetNodeTitle(ENodeTitleType::Type TitleType) const
+{
+	return LOCTEXT("EvaluateLiveLinkFrameAtWorldTimeTitle", "Evaluate Live Link Frame at World Time");
+}
+
+FText UK2Node_EvaluateLiveLinkFrameAtSceneTime::GetNodeTitle(ENodeTitleType::Type TitleType) const
+{
+	return LOCTEXT("EvaluateLiveLinkFrameAtSceneTimeTitle", "Evaluate Live Link Frame at Scene Time");
+}
+
+FName UK2Node_EvaluateLiveLinkFrameWithSpecificRole::GetEvaluateFunctionName() const
+{
+	return GET_FUNCTION_NAME_CHECKED(ULiveLinkBlueprintLibrary, EvaluateLiveLinkFrameWithSpecificRole);
+}
+
+FName UK2Node_EvaluateLiveLinkFrameAtWorldTime::GetEvaluateFunctionName() const
+{
+	return GET_FUNCTION_NAME_CHECKED(ULiveLinkBlueprintLibrary, EvaluateLiveLinkFrameAtWorldTimeOffset);
+}
+
+FName UK2Node_EvaluateLiveLinkFrameAtSceneTime::GetEvaluateFunctionName() const
+{
+	return GET_FUNCTION_NAME_CHECKED(ULiveLinkBlueprintLibrary, EvaluateLiveLinkFrameAtSceneTime);
+}
+
+void UK2Node_EvaluateLiveLinkFrameWithSpecificRole::AddPins(FKismetCompilerContext& CompilerContext, UK2Node_CallFunction* EvaluateLiveLinkFrameFunction)
+{
+}
+
+void UK2Node_EvaluateLiveLinkFrameAtWorldTime::AddPins(FKismetCompilerContext& CompilerContext, UK2Node_CallFunction* EvaluateLiveLinkFrameFunction)
+{
+	UEdGraphPin* InPinSwitch = FindPinChecked(UK2Node_EvaluateLiveLinkFrameHelper::LiveLinkWorldTimePinName);
+	UEdGraphPin* TimePin = EvaluateLiveLinkFrameFunction->FindPinChecked(TEXT("WorldTime"));
+	CompilerContext.CopyPinLinksToIntermediate(*InPinSwitch, *TimePin);
+}
+
+void UK2Node_EvaluateLiveLinkFrameAtSceneTime::AddPins(FKismetCompilerContext& CompilerContext, UK2Node_CallFunction* EvaluateLiveLinkFrameFunction)
+{
+	UEdGraphPin* InPinSwitch = FindPinChecked(UK2Node_EvaluateLiveLinkFrameHelper::LiveLinkSceneTimePinName);
+	UEdGraphPin* TimePin = EvaluateLiveLinkFrameFunction->FindPinChecked(TEXT("SceneTime"));
+	CompilerContext.CopyPinLinksToIntermediate(*InPinSwitch, *TimePin);
+}
+
+#undef LOCTEXT_NAMESPACE
diff --git a/Engine/Plugins/Animation/LiveLink/Source/LiveLinkGraphNode/Private/K2Node_EvaluateLiveLinkCustom.h b/Engine/Plugins/Animation/LiveLink/Source/LiveLinkGraphNode/Private/K2Node_EvaluateLiveLinkCustom.h
new file mode 100644
--- /dev/null
+++ b/Engine/Plugins/Animation/LiveLink/Source/LiveLinkGraphNode/Private/K2Node_EvaluateLiveLinkCustom.h
@@ -0,0 +1,47 @@
+// Copyright 1998-2019 Epic Games, Inc. All Rights Reserved.
+
+#pragma once
+
+#include "CoreMinimal.h"
+#include "K2Node_EvaluateLiveLinkFrame.h"
+
+#include "K2Node_EvaluateLiveLinkCustom.generated.h"
+
+UCLASS()
+class UK2Node_EvaluateLiveLinkFrameWithSpecificRole : public UK2Node_EvaluateLiveLinkFrame
+{
+	GENERATED_BODY()
+
+public:
+	virtual void AllocateDefaultPins() override;
+	virtual FText GetNodeTitle(ENodeTitleType::Type TitleType) const override;
+
+	virtual FName GetEvaluateFunctionName() const override;
+	virtual void AddPins(FKismetCompilerContext& CompilerContext, UK2Node_CallFunction* EvaluateLiveLinkFrameFunction) override;
+};
+
+UCLASS()
+class UK2Node_EvaluateLiveLinkFrameAtWorldTime : public UK2Node_EvaluateLiveLinkFrame
+{
+	GENERATED_BODY()
+
+public:
+	virtual void AllocateDefaultPins() override;
+	virtual FText GetNodeTitle(ENodeTitleType::Type TitleType) const override;
+
+	virtual FName GetEvaluateFunctionName() const override;
+	virtual void AddPins(FKismetCompilerContext& CompilerContext, UK2Node_CallFunction* EvaluateLiveLinkFrameFunction) override;
+};
+
+UCLASS()
+class UK2Node_EvaluateLiveLinkFrameAtSceneTime : public UK2Node_EvaluateLiveLinkFrame
+{
+	GENERATED_BODY()
+
+public:
+	virtual void AllocateDefaultPins() override;
+	virtual FText GetNodeTitle(ENodeTitleType::Type TitleType) const override;
+
+	virtual FName GetEvaluateFunctionName() const override;
+	virtual void AddPins(FKismetCompilerContext& CompilerContext, UK2Node_CallFunction* EvaluateLiveLinkFrameFunction) override;
+};
\ No newline at end of file
diff --git a/Engine/Plugins/Animation/LiveLink/Source/LiveLinkGraphNode/Private/K2Node_EvaluateLiveLinkFrame.cpp b/Engine/Plugins/Animation/LiveLink/Source/LiveLinkGraphNode/Private/K2Node_EvaluateLiveLinkFrame.cpp
new file mode 100644
--- /dev/null
+++ b/Engine/Plugins/Animation/LiveLink/Source/LiveLinkGraphNode/Private/K2Node_EvaluateLiveLinkFrame.cpp
@@ -0,0 +1,395 @@
+// Copyright 1998-2019 Epic Games, Inc. All Rights Reserved.
+
+#include "K2Node_EvaluateLiveLinkFrame.h"
+
+#include "BlueprintActionDatabaseRegistrar.h"
+#include "BlueprintNodeSpawner.h"
+#include "EdGraphSchema_K2.h"
+#include "EditorCategoryUtils.h"
+#include "K2Node_CallFunction.h"
+#include "K2Node_IfThenElse.h"
+#include "Kismet2/BlueprintEditorUtils.h"
+#include "KismetCompiler.h"
+#include "LiveLinkRole.h"
+#include "LiveLinkTypes.h"
+#include "LiveLinkBlueprintLibrary.h"
+#include "UObject/PropertyPortFlags.h"
+
+
+#define LOCTEXT_NAMESPACE "K2Node_EvaluateLiveLinkFrame"
+
+
+namespace UK2Node_EvaluateLiveLinkFrameHelper
+{
+	static FName LiveLinkSubjectPinName = "Subject";
+	static FName LiveLinkRolePinName = "Role";
+	static FName LiveLinkDataResultPinName = "DataResult";
+	static FName FrameNotAvailablePinName = "InvalidFrame";
+};
+
+void UK2Node_EvaluateLiveLinkFrame::AllocateDefaultPins()
+{
+	const UEdGraphSchema_K2* K2Schema = GetDefault<UEdGraphSchema_K2>();
+
+	// Add execution pins
+	CreatePin(EGPD_Input, UEdGraphSchema_K2::PC_Exec, UEdGraphSchema_K2::PN_Execute);
+
+	// Output pin
+	{
+		UEdGraphPin* FrameAvailablePin = CreatePin(EGPD_Output, UEdGraphSchema_K2::PC_Exec, UEdGraphSchema_K2::PN_Then);
+		FrameAvailablePin->PinFriendlyName = LOCTEXT("FrameAvailablePin", "Valid Frame");
+		UEdGraphPin* FrameInvalidePin = CreatePin(EGPD_Output, UEdGraphSchema_K2::PC_Exec, UK2Node_EvaluateLiveLinkFrameHelper::FrameNotAvailablePinName);
+		FrameInvalidePin->PinFriendlyName = LOCTEXT("FrameNotAvailablePinName", "Invalid Frame");
+	}
+	// Output structs pins
+	{
+		UEdGraphPin* DataResultPin = CreatePin(EGPD_Output, UEdGraphSchema_K2::PC_Wildcard, UK2Node_EvaluateLiveLinkFrameHelper::LiveLinkDataResultPinName);
+		DataResultPin->PinFriendlyName = LOCTEXT("LiveLinkDataResultPinName", "DataResult");
+		SetPinToolTip(*DataResultPin, LOCTEXT("DataResultPinDescription", "The data struct, if a frame was present for the given role"));
+	}
+
+	// Subject pin
+	{
+		UEdGraphPin* LiveLinkSubjectRepPin = CreatePin(EGPD_Input, UEdGraphSchema_K2::PC_Struct, FLiveLinkSubjectName::StaticStruct(), UK2Node_EvaluateLiveLinkFrameHelper::LiveLinkSubjectPinName);
+		LiveLinkSubjectRepPin->PinFriendlyName = LOCTEXT("LiveLinkSubjectPinName", "Subject");
+		SetPinToolTip(*LiveLinkSubjectRepPin, LOCTEXT("LiveLinkSubjectNamePinDescription", "The Live Link Subject Name to get a frame from"));
+	}
+	// Role pin
+	{
+		UEdGraphPin* LiveLinkRoleRepPin = CreatePin(EGPD_Input, UEdGraphSchema_K2::PC_Class, ULiveLinkRole::StaticClass(), UK2Node_EvaluateLiveLinkFrameHelper::LiveLinkRolePinName);
+		LiveLinkRoleRepPin->PinFriendlyName = LOCTEXT("LiveLinkRolePinName", "Role");
+		LiveLinkRoleRepPin->bNotConnectable = true;
+		SetPinToolTip(*LiveLinkRoleRepPin, LOCTEXT("LiveLinkRolePinNameDescription", "The Live Link Role the data will be converted to."));
+	}
+
+	Super::AllocateDefaultPins();
+}
+
+void UK2Node_EvaluateLiveLinkFrame::SetPinToolTip(UEdGraphPin& InOutMutatablePin, const FText& InPinDescription) const
+{
+	InOutMutatablePin.PinToolTip = UEdGraphSchema_K2::TypeToText(InOutMutatablePin.PinType).ToString();
+
+	UEdGraphSchema_K2 const* const K2Schema = Cast<const UEdGraphSchema_K2>(GetSchema());
+	if (K2Schema != nullptr)
+	{
+		InOutMutatablePin.PinToolTip += TEXT(" ");
+		InOutMutatablePin.PinToolTip += K2Schema->GetPinDisplayName(&InOutMutatablePin).ToString();
+	}
+
+	InOutMutatablePin.PinToolTip += FString(TEXT("\n")) + InPinDescription.ToString();
+}
+
+void UK2Node_EvaluateLiveLinkFrame::RefreshDataOutputPinType()
+{
+	UScriptStruct* DataType = GetLiveLinkRoleOutputStructType();
+	SetReturnTypeForOutputStruct(DataType);
+}
+
+void UK2Node_EvaluateLiveLinkFrame::SetReturnTypeForOutputStruct(UScriptStruct* InClass)
+{
+	UScriptStruct* OldDataStruct = GetReturnTypeForOutputDataStruct();
+	if (InClass != OldDataStruct)
+	{
+		UEdGraphPin* ResultPin = GetResultingDataPin();
+
+		if (ResultPin->SubPins.Num() > 0)
+		{
+			GetSchema()->RecombinePin(ResultPin);
+		}
+
+		// NOTE: purposefully not disconnecting the ResultPin (even though it changed type)... we want the user to see the old
+		//       connections, and incompatible connections will produce an error (plus, some super-struct connections may still be valid)
+		ResultPin->PinType.PinSubCategoryObject = InClass;
+		ResultPin->PinType.PinCategory = (InClass == nullptr) ? UEdGraphSchema_K2::PC_Wildcard : UEdGraphSchema_K2::PC_Struct;
+	}
+}
+
+UScriptStruct* UK2Node_EvaluateLiveLinkFrame::GetReturnTypeForOutputDataStruct() const
+{
+	return Cast<UScriptStruct>(GetResultingDataPin()->PinType.PinSubCategoryObject.Get());
+}
+
+UScriptStruct* UK2Node_EvaluateLiveLinkFrame::GetLiveLinkRoleOutputStructType() const
+{
+	UScriptStruct* DataStructType = nullptr;
+
+	TSubclassOf<ULiveLinkRole> Role = GetDefaultRolePinValue();
+	if (Role != nullptr)
+	{
+		DataStructType = Role.GetDefaultObject()->GetBlueprintDataStruct();
+	}
+	return DataStructType;
+}
+
+void UK2Node_EvaluateLiveLinkFrame::ReallocatePinsDuringReconstruction(TArray<UEdGraphPin*>& OldPins)
+{
+	Super::ReallocatePinsDuringReconstruction(OldPins);
+}
+
+void UK2Node_EvaluateLiveLinkFrame::GetMenuActions(FBlueprintActionDatabaseRegistrar& ActionRegistrar) const
+{
+	// actions get registered under specific object-keys; the idea is that
+	// actions might have to be updated (or deleted) if their object-key is
+	// mutated (or removed)... here we use the node's class (so if the node
+	// type disappears, then the action should go with it)
+	UClass* ActionKey = GetClass();
+	// to keep from needlessly instantiating a UBlueprintNodeSpawner, first
+	// check to make sure that the registrar is looking for actions of this type
+	// (could be regenerating actions for a specific asset, and therefore the 
+	// registrar would only accept actions corresponding to that asset)
+	if (ActionRegistrar.IsOpenForRegistration(ActionKey))
+	{
+		UBlueprintNodeSpawner* NodeSpawner = UBlueprintNodeSpawner::Create(GetClass());
+		check(NodeSpawner != nullptr);
+
+		ActionRegistrar.AddBlueprintAction(ActionKey, NodeSpawner);
+	}
+}
+
+FText UK2Node_EvaluateLiveLinkFrame::GetMenuCategory() const
+{
+	return FText::FromString(TEXT("LiveLink"));
+}
+
+bool UK2Node_EvaluateLiveLinkFrame::IsConnectionDisallowed(const UEdGraphPin* MyPin, const UEdGraphPin* OtherPin, FString& OutReason) const
+{
+	if (MyPin == GetResultingDataPin() && MyPin->PinType.PinCategory == UEdGraphSchema_K2::PC_Wildcard)
+	{
+		bool bDisallowed = true;
+		if (OtherPin->PinType.PinCategory == UEdGraphSchema_K2::PC_Struct)
+		{
+			bDisallowed = false;
+			if (UScriptStruct* ConnectionType = Cast<UScriptStruct>(OtherPin->PinType.PinSubCategoryObject.Get()))
+			{
+				bDisallowed = GetLiveLinkRoleOutputStructType() != ConnectionType;
+			}
+		}
+		else if (OtherPin->PinType.PinCategory == UEdGraphSchema_K2::PC_Wildcard)
+		{
+			bDisallowed = false;
+		}
+
+		if (bDisallowed)
+		{
+			OutReason = TEXT("Must be a struct that inherits from FLiveLinkBaseBlueprintData");
+		}
+		return bDisallowed;
+	}
+	
+	return false;
+}
+
+void UK2Node_EvaluateLiveLinkFrame::PinDefaultValueChanged(UEdGraphPin* ChangedPin)
+{
+	if (ChangedPin && ChangedPin->PinName == UK2Node_EvaluateLiveLinkFrameHelper::LiveLinkRolePinName)
+	{
+		RefreshDataOutputPinType();
+	}
+}
+
+FText UK2Node_EvaluateLiveLinkFrame::GetTooltipText() const
+{
+	return LOCTEXT("NodeTooltip", "Attempts to Get a LiveLink Frame from a subject using a given Role");
+}
+
+UEdGraphPin* UK2Node_EvaluateLiveLinkFrame::GetThenPin()const
+{
+	const UEdGraphSchema_K2* K2Schema = GetDefault<UEdGraphSchema_K2>();
+
+	UEdGraphPin* Pin = FindPinChecked(UEdGraphSchema_K2::PN_Then);
+	check(Pin->Direction == EGPD_Output);
+	return Pin;
+}
+
+UEdGraphPin* UK2Node_EvaluateLiveLinkFrame::GetLiveLinkRolePin() const
+{
+	UEdGraphPin* Pin = FindPinChecked(UK2Node_EvaluateLiveLinkFrameHelper::LiveLinkRolePinName);
+	check(Pin == nullptr || Pin->Direction == EGPD_Input);
+	return Pin;
+}
+
+UEdGraphPin* UK2Node_EvaluateLiveLinkFrame::GetLiveLinkSubjectPin() const
+{
+	UEdGraphPin* Pin = FindPinChecked(UK2Node_EvaluateLiveLinkFrameHelper::LiveLinkSubjectPinName);
+	check(Pin == nullptr || Pin->Direction == EGPD_Input);
+	return Pin;
+}
+
+UEdGraphPin* UK2Node_EvaluateLiveLinkFrame::GetFrameNotAvailablePin() const
+{
+	UEdGraphPin* Pin = FindPinChecked(UK2Node_EvaluateLiveLinkFrameHelper::FrameNotAvailablePinName);
+	check(Pin->Direction == EGPD_Output);
+	return Pin;
+}
+
+UEdGraphPin* UK2Node_EvaluateLiveLinkFrame::GetResultingDataPin() const
+{
+	const UEdGraphSchema_K2* K2Schema = GetDefault<UEdGraphSchema_K2>();
+
+	UEdGraphPin* Pin = FindPinChecked(UK2Node_EvaluateLiveLinkFrameHelper::LiveLinkDataResultPinName);
+	check(Pin->Direction == EGPD_Output);
+	return Pin;
+}
+
+void UK2Node_EvaluateLiveLinkFrame::ExpandNode(class FKismetCompilerContext& CompilerContext, UEdGraph* SourceGraph)
+{
+	Super::ExpandNode(CompilerContext, SourceGraph);
+
+	TSubclassOf<ULiveLinkRole> Role = GetDefaultRolePinValue();
+	if (Role == nullptr)
+	{
+		CompilerContext.MessageLog.Error(*LOCTEXT("EvaluateLiveLinkRoleNoRole_Error", "EvaluateLiveLinkFrame must have a Role specified.").ToString(), this);
+		// we break exec links so this is the only error we get
+		BreakAllNodeLinks();
+		return;
+	}
+
+	// FUNCTION NODE
+	const FName FunctionName = GetEvaluateFunctionName();
+	UK2Node_CallFunction* EvaluateLiveLinkFrameFunction = CompilerContext.SpawnIntermediateNode<UK2Node_CallFunction>(this, SourceGraph);
+	EvaluateLiveLinkFrameFunction->FunctionReference.SetExternalMember(FunctionName, ULiveLinkBlueprintLibrary::StaticClass());
+	EvaluateLiveLinkFrameFunction->AllocateDefaultPins();
+	CompilerContext.MovePinLinksToIntermediate(*GetExecPin(), *(EvaluateLiveLinkFrameFunction->GetExecPin()));
+
+	// Connect the input of our EvaluateLiveLinkFrame to the Input of our Function pin
+	{
+		UEdGraphPin* OriginalLiveLinkSubjectPin = GetLiveLinkSubjectPin();
+		UEdGraphPin* LiveLinkSubjectInPin = EvaluateLiveLinkFrameFunction->FindPinChecked(TEXT("SubjectName"));
+		CompilerContext.CopyPinLinksToIntermediate(*OriginalLiveLinkSubjectPin, *LiveLinkSubjectInPin);
+	}
+	{
+		UEdGraphPin* OriginalLiveLinkRolePin = GetLiveLinkRolePin();
+		UEdGraphPin* LiveLinkRoleInPin = EvaluateLiveLinkFrameFunction->FindPinChecked(TEXT("Role"));
+		CompilerContext.CopyPinLinksToIntermediate(*OriginalLiveLinkRolePin, *LiveLinkRoleInPin);
+	}
+
+	AddPins(CompilerContext, EvaluateLiveLinkFrameFunction);
+
+	// Get some pins to work with
+	UEdGraphPin* OriginalDataOutPin = FindPinChecked(UK2Node_EvaluateLiveLinkFrameHelper::LiveLinkDataResultPinName);
+	UEdGraphPin* FunctionDataOutPin = EvaluateLiveLinkFrameFunction->FindPinChecked(TEXT("OutBlueprintData"));
+	UEdGraphPin* FunctionReturnPin = EvaluateLiveLinkFrameFunction->FindPinChecked(UEdGraphSchema_K2::PN_ReturnValue);
+	UEdGraphPin* FunctionThenPin = EvaluateLiveLinkFrameFunction->GetThenPin();
+
+	// Set the type of each output pins on this expanded mode to match original
+	FunctionDataOutPin->PinType = OriginalDataOutPin->PinType;
+	FunctionDataOutPin->PinType.PinSubCategoryObject = OriginalDataOutPin->PinType.PinSubCategoryObject;
+
+	//BRANCH NODE
+	UK2Node_IfThenElse* BranchNode = CompilerContext.SpawnIntermediateNode<UK2Node_IfThenElse>(this, SourceGraph);
+	BranchNode->AllocateDefaultPins();
+	// Hook up inputs to branch
+	FunctionThenPin->MakeLinkTo(BranchNode->GetExecPin());
+	FunctionReturnPin->MakeLinkTo(BranchNode->GetConditionPin());
+
+	// Hook up outputs
+	CompilerContext.MovePinLinksToIntermediate(*GetThenPin(), *(BranchNode->GetThenPin()));
+	CompilerContext.MovePinLinksToIntermediate(*GetFrameNotAvailablePin(), *(BranchNode->GetElsePin()));
+	CompilerContext.MovePinLinksToIntermediate(*OriginalDataOutPin, *FunctionDataOutPin);
+
+	BreakAllNodeLinks();
+}
+
+FSlateIcon UK2Node_EvaluateLiveLinkFrame::GetIconAndTint(FLinearColor& OutColor) const
+{
+	OutColor = GetNodeTitleColor();
+	static FSlateIcon Icon("EditorStyle", "Kismet.AllClasses.FunctionIcon");
+	return Icon;
+}
+
+void UK2Node_EvaluateLiveLinkFrame::PostReconstructNode()
+{
+	Super::PostReconstructNode();
+
+	RefreshDataOutputPinType();
+}
+
+void UK2Node_EvaluateLiveLinkFrame::EarlyValidation(class FCompilerResultsLog& MessageLog) const
+{
+	Super::EarlyValidation(MessageLog);
+
+	const UEdGraphPin* LiveLinkSubjectPin = GetLiveLinkSubjectPin();
+	if (!LiveLinkSubjectPin)
+	{
+		MessageLog.Error(*LOCTEXT("MissingPins", "Missing pins in @@").ToString(), this);
+		return;
+	}
+
+	TSubclassOf<ULiveLinkRole> Role = GetDefaultRolePinValue();
+	if (Role == nullptr || !Role->IsChildOf(ULiveLinkRole::StaticClass()))
+	{
+		MessageLog.Error(*LOCTEXT("NoLiveLinkRole", "No LiveLinkRole in @@").ToString(), this);
+		return;
+	}
+
+	UScriptStruct* DesiredOutputStruct = GetLiveLinkRoleOutputStructType();
+	UScriptStruct* ActualOutputStruct = GetReturnTypeForOutputDataStruct();
+	if (DesiredOutputStruct != ActualOutputStruct)
+	{
+		MessageLog.Error(*LOCTEXT("OutputPinDoNotMatches", "The output data pin do not maches in @@").ToString(), this);
+		return;
+	}
+
+	if (ActualOutputStruct != nullptr)
+	{
+		UEdGraphPin* ResultPin = GetResultingDataPin();
+		if (ResultPin && ResultPin->LinkedTo.Num() > 0)
+		{
+			for (UEdGraphPin* LinkPin : ResultPin->LinkedTo)
+			{
+				UScriptStruct* LinkType = Cast<UScriptStruct>(LinkPin->PinType.PinSubCategoryObject.Get());
+				if (ActualOutputStruct != LinkType)
+				{
+					MessageLog.Error(*LOCTEXT("OutputLinkPinDoNotMatches", "A linked pin data structure do not match the Live Link Role data structure in @@. Expected: @@.").ToString(), this, ActualOutputStruct);
+					return;
+				}
+			}
+		}
+	}
+}
+
+void UK2Node_EvaluateLiveLinkFrame::PreloadRequiredAssets()
+{
+	return Super::PreloadRequiredAssets();
+}
+
+void UK2Node_EvaluateLiveLinkFrame::NotifyPinConnectionListChanged(UEdGraphPin* Pin)
+{
+	Super::NotifyPinConnectionListChanged(Pin);
+
+	UEdGraphPin* LiveLinkRolePin = GetLiveLinkRolePin();
+	if (Pin == GetResultingDataPin())
+	{
+		// this connection would only change the output type if the role pin is undefined
+		const bool bIsTypeAuthority = (LiveLinkRolePin->LinkedTo.Num() <= 0 && LiveLinkRolePin->DefaultObject == nullptr);
+		if (bIsTypeAuthority)
+		{
+			RefreshDataOutputPinType();
+		}
+	}
+	else if (Pin == LiveLinkRolePin)
+	{
+		const bool bConnectionAdded = Pin->LinkedTo.Num() > 0;
+		if (bConnectionAdded)
+		{
+			RefreshDataOutputPinType();
+		}
+	}
+}
+
+TSubclassOf<ULiveLinkRole> UK2Node_EvaluateLiveLinkFrame::GetDefaultRolePinValue() const
+{
+	UEdGraphPin* LiveLinkRolePin = GetLiveLinkRolePin();
+	if (LiveLinkRolePin && LiveLinkRolePin->DefaultObject != nullptr && LiveLinkRolePin->LinkedTo.Num() == 0)
+	{
+		UClass* ClassValue = Cast<UClass>(LiveLinkRolePin->DefaultObject);
+		if (ClassValue && ClassValue->IsChildOf(ULiveLinkRole::StaticClass()))
+		{
+			return ClassValue;
+		}
+	}
+	return nullptr;
+}
+
+
+#undef LOCTEXT_NAMESPACE
diff --git a/Engine/Plugins/Animation/LiveLink/Source/LiveLinkGraphNode/Private/LiveLinkGraphNodeModule.cpp b/Engine/Plugins/Animation/LiveLink/Source/LiveLinkGraphNode/Private/LiveLinkGraphNodeModule.cpp
new file mode 100644
--- /dev/null
+++ b/Engine/Plugins/Animation/LiveLink/Source/LiveLinkGraphNode/Private/LiveLinkGraphNodeModule.cpp
@@ -0,0 +1,6 @@
+// Copyright 1998-2019 Epic Games, Inc. All Rights Reserved.
+
+#include "Modules/ModuleManager.h"
+
+IMPLEMENT_MODULE(FDefaultModuleImpl, LiveLinkGraphNode);
+
diff --git a/Engine/Plugins/Animation/LiveLink/Source/LiveLinkGraphNode/Public/K2Node_EvaluateLiveLinkFrame.h b/Engine/Plugins/Animation/LiveLink/Source/LiveLinkGraphNode/Public/K2Node_EvaluateLiveLinkFrame.h
new file mode 100644
--- /dev/null
+++ b/Engine/Plugins/Animation/LiveLink/Source/LiveLinkGraphNode/Public/K2Node_EvaluateLiveLinkFrame.h
@@ -0,0 +1,93 @@
+// Copyright 1998-2019 Epic Games, Inc. All Rights Reserved.
+
+#pragma once
+
+#include "CoreMinimal.h"
+#include "EdGraph/EdGraphNodeUtils.h"
+#include "LiveLinkRole.h"
+#include "K2Node.h"
+#include "Textures/SlateIcon.h"
+#include "UObject/ObjectMacros.h"
+
+#include "K2Node_EvaluateLiveLinkFrame.generated.h"
+
+class FBlueprintActionDatabaseRegistrar;
+class FKismetCompilerContext;
+class UDataTable;
+class UEdGraph;
+class UK2Node_CallFunction;
+
+
+UCLASS(Abstract)
+class LIVELINKGRAPHNODE_API UK2Node_EvaluateLiveLinkFrame : public UK2Node
+{
+	GENERATED_BODY()
+
+public:
+	//~ Begin UEdGraphNode Interface.
+	virtual void AllocateDefaultPins() override;
+	virtual void PinDefaultValueChanged(UEdGraphPin* Pin) override;
+	virtual FText GetTooltipText() const override;
+	virtual void ExpandNode(FKismetCompilerContext& CompilerContext, UEdGraph* SourceGraph) override;
+	virtual FSlateIcon GetIconAndTint(FLinearColor& OutColor) const override;
+	virtual void PostReconstructNode() override;
+	//~ End UEdGraphNode Interface.
+
+	//~ Begin UK2Node Interface
+	virtual bool IsNodeSafeToIgnore() const override { return true; }
+	virtual void ReallocatePinsDuringReconstruction(TArray<UEdGraphPin*>& OldPins) override;
+	virtual void GetMenuActions(FBlueprintActionDatabaseRegistrar& ActionRegistrar) const override;
+	virtual FText GetMenuCategory() const override;
+	virtual bool IsConnectionDisallowed(const UEdGraphPin* MyPin, const UEdGraphPin* OtherPin, FString& OutReason) const override;
+	virtual void EarlyValidation(class FCompilerResultsLog& MessageLog) const override;
+	virtual void PreloadRequiredAssets() override;
+	virtual void NotifyPinConnectionListChanged(UEdGraphPin* Pin) override;
+	//~ End UK2Node Interface
+
+	/** Get the return types of our struct */
+	UScriptStruct* GetReturnTypeForOutputDataStruct() const;
+
+	/** Get the then output pin */
+	UEdGraphPin* GetThenPin() const;
+
+	/** Get the Live Link Role input pin */
+	UEdGraphPin* GetLiveLinkRolePin() const;
+
+	/** Get the Live Link Subject input pin */
+	UEdGraphPin* GetLiveLinkSubjectPin() const;
+
+	/** Get the exec output pin for when no frame is available for the desired role */
+	UEdGraphPin* GetFrameNotAvailablePin() const;
+
+	/** Get the result output pin */
+	UEdGraphPin* GetResultingDataPin() const;
+
+	/** Get the type that the Live Link Role evaluates to to return */
+	UScriptStruct* GetLiveLinkRoleOutputStructType() const;
+	UScriptStruct* GetLiveLinkRoleOutputFrameStructType() const;
+
+protected:
+
+	/** Get the name of the UFunction we are trying to build from */
+	virtual FName GetEvaluateFunctionName() const PURE_VIRTUAL(UK2Node_EvaluateLiveLinkFrame::GetEvaluateFunctionName, return NAME_None; );
+
+	/** Add additionnal pin that the evaluate function would need */
+	virtual void AddPins(FKismetCompilerContext& CompilerContext, UK2Node_CallFunction* EvaluateLiveLinkFrameFunction) PURE_VIRTUAL(UK2Node_EvaluateLiveLinkFrame::AddPins, );
+
+	/**
+	 * Takes the specified "MutatablePin" and sets its 'PinToolTip' field (according
+	 * to the specified description)
+	 *
+	 * @param   MutatablePin	The pin you want to set tool-tip text on
+	 * @param   PinDescription	A string describing the pin's purpose
+	 */
+	void SetPinToolTip(UEdGraphPin& InOutMutatablePin, const FText& InPinDescription) const;
+
+	/** Set the return type of our structs */
+	void SetReturnTypeForOutputStruct(UScriptStruct* InClass);
+
+	/** Queries for the authoritative return type, then modifies the return pin to match */
+	void RefreshDataOutputPinType();
+
+	TSubclassOf<ULiveLinkRole> GetDefaultRolePinValue() const;
+};
