Change warning into verbose messages
Block any connection made to the Evaluate LiveLink Frame role pin.

diff --git a/Engine/Plugins/Animation/LiveLink/Source/LiveLinkGraphNode/Private/K2Node_EvaluateLiveLinkFrame.cpp b/Engine/Plugins/Animation/LiveLink/Source/LiveLinkGraphNode/Private/K2Node_EvaluateLiveLinkFrame.cpp
--- a/Engine/Plugins/Animation/LiveLink/Source/LiveLinkGraphNode/Private/K2Node_EvaluateLiveLinkFrame.cpp
+++ b/Engine/Plugins/Animation/LiveLink/Source/LiveLinkGraphNode/Private/K2Node_EvaluateLiveLinkFrame.cpp
@@ -175,15 +175,30 @@ bool UK2Node_EvaluateLiveLinkFrame::IsConnectionDisallowed(const UEdGraphPin* My
 		}
 		return bDisallowed;
 	}
-	
+
+	if (MyPin == GetLiveLinkRolePin())
+	{
+		return true;
+	}
+
 	return false;
 }
 
 void UK2Node_EvaluateLiveLinkFrame::PinDefaultValueChanged(UEdGraphPin* ChangedPin)
 {
-	if (ChangedPin && ChangedPin->PinName == UK2Node_EvaluateLiveLinkFrameHelper::LiveLinkRolePinName)
+	UEdGraphPin* LiveLinkRolePin = GetLiveLinkRolePin();
+	if (ChangedPin == LiveLinkRolePin && LiveLinkRolePin)
 	{
-		RefreshDataOutputPinType();
+		if (LiveLinkRolePin->DefaultObject != nullptr && LiveLinkRolePin->LinkedTo.Num() == 0)
+		{
+			UClass* ClassValue = Cast<UClass>(LiveLinkRolePin->DefaultObject);
+			bool bIsValid = ClassValue && ClassValue->IsChildOf(ULiveLinkRole::StaticClass()) && !ClassValue->HasAnyClassFlags(CLASS_Abstract | CLASS_Deprecated | CLASS_HideDropDown);
+			if (!bIsValid)
+			{
+				LiveLinkRolePin->DefaultObject = nullptr;
+			}
+			RefreshDataOutputPinType();
+		}
 	}
 }
 
@@ -383,7 +398,7 @@ TSubclassOf<ULiveLinkRole> UK2Node_EvaluateLiveLinkFrame::GetDefaultRolePinValue
 	if (LiveLinkRolePin && LiveLinkRolePin->DefaultObject != nullptr && LiveLinkRolePin->LinkedTo.Num() == 0)
 	{
 		UClass* ClassValue = Cast<UClass>(LiveLinkRolePin->DefaultObject);
-		if (ClassValue && ClassValue->IsChildOf(ULiveLinkRole::StaticClass()))
+		if (ClassValue && ClassValue->IsChildOf(ULiveLinkRole::StaticClass()) && !ClassValue->HasAnyClassFlags(CLASS_Abstract|CLASS_Deprecated|CLASS_HideDropDown))
 		{
 			return ClassValue;
 		}
diff --git a/Engine/Source/Runtime/LiveLinkInterface/Private/ILiveLinkSubject.cpp b/Engine/Source/Runtime/LiveLinkInterface/Private/ILiveLinkSubject.cpp
--- a/Engine/Source/Runtime/LiveLinkInterface/Private/ILiveLinkSubject.cpp
+++ b/Engine/Source/Runtime/LiveLinkInterface/Private/ILiveLinkSubject.cpp
@@ -25,7 +25,7 @@ bool ILiveLinkSubject::EvaluateFrame(TSubclassOf<ULiveLinkRole> InDesiredRole, F
 
 	if (!HasValidFrameSnapshot())
 	{
-		UE_LOG(LogLiveLinkSubject, Warning, TEXT("Can't evaluate frame for subject '%s'. No data was available."), *GetSubjectKey().SubjectName.ToString());
+		UE_LOG(LogLiveLinkSubject, Verbose, TEXT("Can't evaluate frame for subject '%s'. No data was available."), *GetSubjectKey().SubjectName.ToString());
 		return false;
 	}
 
