- Fix for editor monolithic build

diff --git a/Engine/Plugins/Animation/LiveLink/Source/LiveLinkEditor/Private/LiveLinkEditorModule.cpp b/Engine/Plugins/Animation/LiveLink/Source/LiveLinkEditor/Private/LiveLinkEditorModule.cpp
--- a/Engine/Plugins/Animation/LiveLink/Source/LiveLinkEditor/Private/LiveLinkEditorModule.cpp
+++ b/Engine/Plugins/Animation/LiveLink/Source/LiveLinkEditor/Private/LiveLinkEditorModule.cpp
@@ -45,14 +45,19 @@ DEFINE_LOG_CATEGORY(LogLiveLinkEditor);
 static const FName LiveLinkClientTabName(TEXT("LiveLink"));
 static const FName LevelEditorModuleName(TEXT("LevelEditor"));
 
-#define IMAGE_PLUGIN_BRUSH( RelativePath, ... ) FSlateImageBrush( InPluginContent( RelativePath, ".png" ), __VA_ARGS__ )
 
-FString InPluginContent(const FString& RelativePath, const ANSICHAR* Extension)
+namespace LiveLinkEditorModuleUtils
 {
-	static FString ContentDir = IPluginManager::Get().FindPlugin(TEXT("LiveLink"))->GetContentDir();
-	return (ContentDir / RelativePath) + Extension;
+	FString InPluginContent(const FString& RelativePath, const ANSICHAR* Extension)
+	{
+		static FString ContentDir = IPluginManager::Get().FindPlugin(TEXT("LiveLink"))->GetContentDir();
+		return (ContentDir / RelativePath) + Extension;
+	}
 }
 
+#define IMAGE_PLUGIN_BRUSH( RelativePath, ... ) FSlateImageBrush( LiveLinkEditorModuleUtils::InPluginContent( RelativePath, ".png" ), __VA_ARGS__ )
+
+
 class FLiveLinkEditorModule : public IModuleInterface
 {
 public:
diff --git a/Engine/Plugins/Animation/LiveLink/Source/LiveLinkSequencer/Private/LiveLinkSequencerModule.cpp b/Engine/Plugins/Animation/LiveLink/Source/LiveLinkSequencer/Private/LiveLinkSequencerModule.cpp
--- a/Engine/Plugins/Animation/LiveLink/Source/LiveLinkSequencer/Private/LiveLinkSequencerModule.cpp
+++ b/Engine/Plugins/Animation/LiveLink/Source/LiveLinkSequencer/Private/LiveLinkSequencerModule.cpp
@@ -37,13 +37,19 @@ static const FName TakeRecorderModuleName(TEXT("TakeRecorder"));
 static const FName MovieSceneSectionRecorderFactoryName(TEXT("MovieSceneSectionRecorderFactory"));
 static TArray<TSubclassOf<ULiveLinkRole>> SupportedRecordingRoles;
 
-#define IMAGE_PLUGIN_BRUSH( RelativePath, ... ) FSlateImageBrush( InPluginContent( RelativePath, ".png" ), __VA_ARGS__ )
 
-FString InPluginContent(const FString& RelativePath, const ANSICHAR* Extension)
+namespace LiveLinkSequencerModuleUtils
 {
-	static FString ContentDir = IPluginManager::Get().FindPlugin(TEXT("LiveLink"))->GetContentDir();
-	return (ContentDir / RelativePath) + Extension;
+	FString InPluginContent(const FString& RelativePath, const ANSICHAR* Extension)
+	{
+		static FString ContentDir = IPluginManager::Get().FindPlugin(TEXT("LiveLink"))->GetContentDir();
+		return (ContentDir / RelativePath) + Extension;
+	}
 }
+
+#define IMAGE_PLUGIN_BRUSH( RelativePath, ... ) FSlateImageBrush( LiveLinkSequencerModuleUtils::InPluginContent( RelativePath, ".png" ), __VA_ARGS__ )
+
+
 class FLiveLinkSequencerModule : public IModuleInterface
 {
 public:
