Merging fix from //UE4/Dev-Enterprise CL#6987938 : Allow to strip the scale from a LiveLinkController

diff --git a/Engine/Plugins/Animation/LiveLink/Source/LiveLinkComponents/Private/Controllers/LiveLinkTransformController.cpp b/Engine/Plugins/Animation/LiveLink/Source/LiveLinkComponents/Private/Controllers/LiveLinkTransformController.cpp
--- a/Engine/Plugins/Animation/LiveLink/Source/LiveLinkComponents/Private/Controllers/LiveLinkTransformController.cpp
+++ b/Engine/Plugins/Animation/LiveLink/Source/LiveLinkComponents/Private/Controllers/LiveLinkTransformController.cpp
@@ -23,13 +23,18 @@ void FLiveLinkTransformControllerData::ApplyTransform(USceneComponent* SceneComp
 {
 	if (SceneComponent)
 	{
+		FTransform ComponentTransform = Transform;
+		if (!bUseScale)
+		{
+			ComponentTransform.SetScale3D(FVector::OneVector);
+		}
 		if (bWorldTransform)
 		{
-			SceneComponent->SetWorldTransform(Transform, bSweep, nullptr, bTeleport ? ETeleportType::TeleportPhysics : ETeleportType::ResetPhysics);
+			SceneComponent->SetWorldTransform(ComponentTransform, bSweep, nullptr, bTeleport ? ETeleportType::TeleportPhysics : ETeleportType::ResetPhysics);
 		}
 		else
 		{
-			SceneComponent->SetRelativeTransform(Transform, bSweep, nullptr, bTeleport ? ETeleportType::TeleportPhysics : ETeleportType::ResetPhysics);
+			SceneComponent->SetRelativeTransform(ComponentTransform, bSweep, nullptr, bTeleport ? ETeleportType::TeleportPhysics : ETeleportType::ResetPhysics);
 		}
 	}
 }
diff --git a/Engine/Plugins/Animation/LiveLink/Source/LiveLinkComponents/Public/Controllers/LiveLinkTransformController.h b/Engine/Plugins/Animation/LiveLink/Source/LiveLinkComponents/Public/Controllers/LiveLinkTransformController.h
--- a/Engine/Plugins/Animation/LiveLink/Source/LiveLinkComponents/Public/Controllers/LiveLinkTransformController.h
+++ b/Engine/Plugins/Animation/LiveLink/Source/LiveLinkComponents/Public/Controllers/LiveLinkTransformController.h
@@ -18,6 +18,9 @@ struct FLiveLinkTransformControllerData
 	UPROPERTY(EditAnywhere, BlueprintReadOnly, Category = "LiveLink")
 	bool bWorldTransform = true;
 
+	UPROPERTY(EditAnywhere, Category = "LiveLink", AdvancedDisplay)
+	bool bUseScale = true;
+
 	/**
 	 * Whether we sweep to the destination location, triggering overlaps along the way and stopping short of the target if blocked by something.
 	 * Only the root component is swept and checked for blocking collision, child components move without sweeping. If collision is off, this has no effect.
