Merging fix from //UE4/Dev-Enterprise CL#7247953 : Fix for LiveLink connection drops when main thread is blocked

diff --git a/Engine/Plugins/Animation/LiveLink/Source/LiveLink/Private/ILiveLinkModule.h b/Engine/Plugins/Animation/LiveLink/Source/LiveLink/Private/ILiveLinkModule.h
new file mode 100644
--- /dev/null
+++ b/Engine/Plugins/Animation/LiveLink/Source/LiveLink/Private/ILiveLinkModule.h
@@ -0,0 +1,41 @@
+// Copyright 1998-2019 Epic Games, Inc. All Rights Reserved.
+
+#pragma once
+
+#include "Modules/ModuleInterface.h"
+#include "Containers/Array.h"
+#include "Modules/ModuleManager.h"
+
+class FLiveLinkHeartbeatEmitter;
+
+/**
+ * Interface for messaging modules.
+ */
+class ILiveLinkModule
+	: public IModuleInterface
+{
+public:
+
+	/**
+	 * Gets a reference to the live link module instance.
+	 *
+	 * @return A reference to the live link module.
+	 */
+	static ILiveLinkModule& Get()
+	{
+#if PLATFORM_IOS
+        static ILiveLinkModule& LiveLinkModule = FModuleManager::LoadModuleChecked<ILiveLinkModule>("LiveLink");
+        return LiveLinkModule;
+#else
+        return FModuleManager::LoadModuleChecked<ILiveLinkModule>("LiveLink");
+#endif
+	}
+
+	virtual FLiveLinkHeartbeatEmitter& GetHeartbeatEmitter() = 0;
+
+public:
+
+	/** Virtual destructor. */
+	virtual ~ILiveLinkModule() { }
+};
+
diff --git a/Engine/Plugins/Animation/LiveLink/Source/LiveLink/Private/LiveLinkHeartbeatEmitter.cpp b/Engine/Plugins/Animation/LiveLink/Source/LiveLink/Private/LiveLinkHeartbeatEmitter.cpp
new file mode 100644
--- /dev/null
+++ b/Engine/Plugins/Animation/LiveLink/Source/LiveLink/Private/LiveLinkHeartbeatEmitter.cpp
@@ -0,0 +1,95 @@
+// Copyright 1998-2019 Epic Games, Inc. All Rights Reserved.
+
+#include "LiveLinkHeartbeatEmitter.h"
+
+#include "LiveLinkMessages.h"
+#include "LiveLinkSettings.h"
+
+#include "GenericPlatform/GenericPlatformProcess.h"
+#include "HAL/Event.h"
+#include "HAL/RunnableThread.h"
+
+
+bool FLiveLinkHeartbeatEmitter::FHeartbeatRecipient::operator==(const FLiveLinkHeartbeatEmitter::FHeartbeatRecipient& other) const
+{
+	return other.MessageEndpoint == this->MessageEndpoint && other.ConnectionAddress == this->ConnectionAddress;
+}
+
+
+FLiveLinkHeartbeatEmitter::FLiveLinkHeartbeatEmitter()
+	: bIsRunning(false)
+	, HeartbeatFrequencyInMs(GetDefault<ULiveLinkSettings>()->GetMessageBusHeartbeatFrequency() * 1000.0f)
+	, HeartbeatEvent(nullptr)
+{
+
+}
+
+void FLiveLinkHeartbeatEmitter::StartHeartbeat(const FMessageAddress& RecipientAddress, const TSharedPtr<FMessageEndpoint, ESPMode::ThreadSafe>& MessageEndpoint)
+{
+	if (!bIsRunning)
+	{
+		bIsRunning = true;
+		HeartbeatEvent = FGenericPlatformProcess::GetSynchEventFromPool();
+		check(HeartbeatEvent);
+		Thread.Reset(FRunnableThread::Create(this, TEXT("LiveLinkHeartbeatEmitter")));
+	}
+
+	FScopeLock Lock(&CriticalSection);
+	HeartbeatRecipients.Add({ MessageEndpoint, RecipientAddress });
+}
+
+void FLiveLinkHeartbeatEmitter::StopHeartbeat(const FMessageAddress& RecipientAddress, const TSharedPtr<FMessageEndpoint, ESPMode::ThreadSafe>& MessageEndpoint)
+{
+	FScopeLock Lock(&CriticalSection);
+	HeartbeatRecipients.Remove({ MessageEndpoint, RecipientAddress });
+}
+
+uint32 FLiveLinkHeartbeatEmitter::Run()
+{
+	while (bIsRunning)
+	{
+		check(HeartbeatEvent);
+		HeartbeatEvent->Wait(HeartbeatFrequencyInMs);
+
+		FScopeLock Lock(&CriticalSection);
+		if (bIsRunning) // make sure we were not told to exit during the wait
+		{
+			for (const FHeartbeatRecipient& Recipient : HeartbeatRecipients)
+			{
+				if (TSharedPtr<FMessageEndpoint, ESPMode::ThreadSafe> MessageEndpoint = Recipient.MessageEndpoint.Pin())
+				{
+					MessageEndpoint->Send(new FLiveLinkHeartbeatMessage(), Recipient.ConnectionAddress);
+				}
+			}
+		}
+	}
+
+	return 0;
+}
+
+void FLiveLinkHeartbeatEmitter::Exit()
+{
+	if (!bIsRunning)
+	{
+		return;
+	}
+
+	bIsRunning = false;
+	{
+		FScopeLock Lock(&CriticalSection);
+		HeartbeatRecipients.Empty();
+	}
+
+	if (HeartbeatEvent)
+	{
+		HeartbeatEvent->Trigger();
+	}
+
+	Thread->WaitForCompletion();
+
+	if (HeartbeatEvent)
+	{
+		FGenericPlatformProcess::ReturnSynchEventToPool(HeartbeatEvent);
+		HeartbeatEvent = nullptr;
+	}
+}
diff --git a/Engine/Plugins/Animation/LiveLink/Source/LiveLink/Private/LiveLinkHeartbeatEmitter.h b/Engine/Plugins/Animation/LiveLink/Source/LiveLink/Private/LiveLinkHeartbeatEmitter.h
new file mode 100644
--- /dev/null
+++ b/Engine/Plugins/Animation/LiveLink/Source/LiveLink/Private/LiveLinkHeartbeatEmitter.h
@@ -0,0 +1,62 @@
+// Copyright 1998-2019 Epic Games, Inc. All Rights Reserved.
+
+#pragma once
+
+#include "HAL/Runnable.h"
+#include "HAL/ThreadSafeBool.h"
+
+#include "IMessageContext.h"
+#include "MessageEndpoint.h"
+
+class  FEvent;
+class  FRunnableThread;
+
+/**
+ * Sends out heartbeat signals at the rate defined by LiveLinkSettings::MessageBusHeartbeatFrequency to all recipients.
+ *
+ * Recipients for the heartbeat can be added by calling StartHeartbeat().
+ * On the first call to StartHeartbeat a thread is spawned that keeps sending the heartbeat to everyone.
+ * The heartbeats and the sender thread can be stopped by calling Exit().
+ */
+class FLiveLinkHeartbeatEmitter : public FRunnable
+{
+public:
+	FLiveLinkHeartbeatEmitter();
+
+	/**
+	 * Start sending a heartbeat to the specified Recipient by using the given MessageEndpoint.
+	 * @note This will spawn a new thread when first called.
+	 * @note It is assumed that each source that sends out heartbeats has its own MessageEndpoint.
+	 */
+	void StartHeartbeat(const FMessageAddress& RecipientAddress, const TSharedPtr<FMessageEndpoint, ESPMode::ThreadSafe>& MessageEndpoint);
+
+	/**
+	 * Stop sending a heartbeat to the specified Recipient.
+	 */
+	void StopHeartbeat(const FMessageAddress& RecipientAddress, const TSharedPtr<FMessageEndpoint, ESPMode::ThreadSafe>& MessageEndpoint);
+
+public:
+	/**
+	 * FRunnable Interface
+	 */
+	virtual uint32 Run() override;
+	virtual void Exit() override;
+
+private:
+	struct FHeartbeatRecipient
+	{
+		TWeakPtr<FMessageEndpoint, ESPMode::ThreadSafe> MessageEndpoint;
+		FMessageAddress ConnectionAddress;
+
+		bool operator==(const FHeartbeatRecipient& other) const;
+	};
+
+private:
+	mutable FCriticalSection CriticalSection;
+	TArray<FHeartbeatRecipient> HeartbeatRecipients;
+	FThreadSafeBool bIsRunning;
+	float HeartbeatFrequencyInMs;
+
+	FEvent* HeartbeatEvent;
+	TUniquePtr<FRunnableThread> Thread;
+};
diff --git a/Engine/Plugins/Animation/LiveLink/Source/LiveLink/Private/LiveLinkMessageBusSource.cpp b/Engine/Plugins/Animation/LiveLink/Source/LiveLink/Private/LiveLinkMessageBusSource.cpp
--- a/Engine/Plugins/Animation/LiveLink/Source/LiveLink/Private/LiveLinkMessageBusSource.cpp
+++ b/Engine/Plugins/Animation/LiveLink/Source/LiveLink/Private/LiveLinkMessageBusSource.cpp
@@ -3,7 +3,9 @@
 #include "LiveLinkMessageBusSource.h"
 
 #include "ILiveLinkClient.h"
+#include "ILiveLinkModule.h"
 #include "LiveLinkClient.h"
+#include "LiveLinkHeartbeatEmitter.h"
 #include "LiveLinkMessageBusDiscoveryManager.h"
 #include "LiveLinkMessages.h"
 #include "LiveLinkRoleTrait.h"
@@ -66,17 +68,10 @@ void FLiveLinkMessageBusSource::Update()
 	}
 	else
 	{
-		const double HeartbeatFrequency = GetDefault<ULiveLinkSettings>()->GetMessageBusHeartbeatFrequency();
 		const double HeartbeatTimeout = GetDefault<ULiveLinkSettings>()->GetMessageBusHeartbeatTimeout();
 		const double CurrentTime = FApp::GetCurrentTime();
 
 		bIsValid = CurrentTime - ConnectionLastActive < HeartbeatTimeout;
-		
-		if (CurrentTime - HeartbeatLastSent > HeartbeatFrequency)
-		{
-			MessageEndpoint->Send(new FLiveLinkHeartbeatMessage(), ConnectionAddress);
-			HeartbeatLastSent = CurrentTime;
-		}
 	}
 }
 
@@ -236,6 +231,8 @@ void FLiveLinkMessageBusSource::SendConnectMessage()
 	FLiveLinkConnectMessage* ConnectMessage = new FLiveLinkConnectMessage();
 	ConnectMessage->LiveLinkVersion = 2;
 	MessageEndpoint->Send(ConnectMessage, ConnectionAddress);
+	FLiveLinkHeartbeatEmitter& HeartbeatEmitter = ILiveLinkModule::Get().GetHeartbeatEmitter();
+	HeartbeatEmitter.StartHeartbeat(ConnectionAddress, MessageEndpoint);
 	bIsValid = true;
 }
 
@@ -247,6 +244,8 @@ bool FLiveLinkMessageBusSource::RequestSourceShutdown()
 		DiscoveryManager->RemoveDiscoveryMessageRequest();
 	}
 
+	FLiveLinkHeartbeatEmitter& HeartbeatEmitter = ILiveLinkModule::Get().GetHeartbeatEmitter();
+	HeartbeatEmitter.StopHeartbeat(ConnectionAddress, MessageEndpoint);
 	FMessageEndpoint::SafeRelease(MessageEndpoint);
 
 	return true;
diff --git a/Engine/Plugins/Animation/LiveLink/Source/LiveLink/Private/LiveLinkMessageBusSource.h b/Engine/Plugins/Animation/LiveLink/Source/LiveLink/Private/LiveLinkMessageBusSource.h
--- a/Engine/Plugins/Animation/LiveLink/Source/LiveLink/Private/LiveLinkMessageBusSource.h
+++ b/Engine/Plugins/Animation/LiveLink/Source/LiveLink/Private/LiveLinkMessageBusSource.h
@@ -24,7 +24,6 @@ public:
 		: ConnectionAddress(InConnectionAddress)
 		, SourceType(InSourceType)
 		, SourceMachineName(InSourceMachineName)
-		, HeartbeatLastSent(0.0)
 		, ConnectionLastActive(0.0)
 		, bIsValid(false)
 		, bInvalidSubjectNameErrorReported(false)
@@ -74,9 +73,6 @@ private:
 	FText SourceType;
 	FText SourceMachineName;
 
-	// Time we last sent connection heartbeat
-	double HeartbeatLastSent;
-
 	// Time we last received anything 
 	double ConnectionLastActive;
 
diff --git a/Engine/Plugins/Animation/LiveLink/Source/LiveLink/Private/LiveLinkModule.cpp b/Engine/Plugins/Animation/LiveLink/Source/LiveLink/Private/LiveLinkModule.cpp
--- a/Engine/Plugins/Animation/LiveLink/Source/LiveLink/Private/LiveLinkModule.cpp
+++ b/Engine/Plugins/Animation/LiveLink/Source/LiveLink/Private/LiveLinkModule.cpp
@@ -1,4 +1,4 @@
-﻿// Copyright 1998-2019 Epic Games, Inc. All Rights Reserved.
+// Copyright 1998-2019 Epic Games, Inc. All Rights Reserved.
 
 #include "ILiveLinkModule.h"
 
@@ -9,6 +9,8 @@
 
 #include "LiveLinkClient.h"
 
+#include "LiveLinkHeartbeatEmitter.h"
+
 /**
  * Implements the Messaging module.
  */
@@ -26,7 +28,9 @@ public:
 	FLiveLinkModule()
 		: LiveLinkClient()
 		, LiveLinkMotionController(LiveLinkClient)
+		, HeartbeatEmitter(MakeUnique<FLiveLinkHeartbeatEmitter>())
 	{}
+
 	// IModuleInterface interface
 
 	virtual void StartupModule() override
@@ -39,6 +43,7 @@ public:
 
 	virtual void ShutdownModule() override
 	{
+		HeartbeatEmitter->Exit();
 		LiveLinkMotionController.UnregisterController();
 		IModularFeatures::Get().UnregisterModularFeature(FLiveLinkClient::ModularFeatureName, &LiveLinkClient);
 		delete FLiveLinkMessageBusDiscoveryManager::Get();
@@ -48,6 +53,14 @@ public:
 	{
 		return false;
 	}
+
+	virtual FLiveLinkHeartbeatEmitter& GetHeartbeatEmitter() override
+	{
+		return *HeartbeatEmitter;
+	}
+
+private:
+	TUniquePtr<FLiveLinkHeartbeatEmitter> HeartbeatEmitter;
 };
 
 IMPLEMENT_MODULE(FLiveLinkModule, LiveLink);
diff --git a/Engine/Plugins/Animation/LiveLink/Source/LiveLink/Public/ILiveLinkModule.h b/Engine/Plugins/Animation/LiveLink/Source/LiveLink/Public/ILiveLinkModule.h
deleted file mode 100644
