- Integrating CL#7251061

- Fix for old LiveLinkSection data conversion not creating an AnimationRole

diff --git a/Engine/Plugins/Animation/LiveLink/Source/LiveLinkMovieScene/Private/MovieScene/MovieSceneLiveLinkSection.cpp b/Engine/Plugins/Animation/LiveLink/Source/LiveLinkMovieScene/Private/MovieScene/MovieSceneLiveLinkSection.cpp
--- a/Engine/Plugins/Animation/LiveLink/Source/LiveLinkMovieScene/Private/MovieScene/MovieSceneLiveLinkSection.cpp
+++ b/Engine/Plugins/Animation/LiveLink/Source/LiveLinkMovieScene/Private/MovieScene/MovieSceneLiveLinkSection.cpp
@@ -191,30 +191,23 @@ void UMovieSceneLiveLinkSection::ConvertPreRoleData()
 	SubjectPreset.Key.SubjectName = SubjectName_DEPRECATED;
 
 	//Verify if there are any bones for this subject. No bones means only curves were present so create a basic role in this case.
-	if (RefSkeleton_DEPRECATED.BoneNames.Num() == 0)
-	{
-		SubjectPreset.Role = ULiveLinkBasicRole::StaticClass();
-		StaticData->InitializeWith(FLiveLinkBaseStaticData::StaticStruct(), nullptr);
-
-		//Data initialization will be done in the common part
-	}
-	else
-	{
-		SubjectPreset.Role = ULiveLinkAnimationRole::StaticClass();
-		StaticData->InitializeWith(FLiveLinkSkeletonStaticData::StaticStruct(), nullptr);
-		
-		//Take old skeleton data for new static data structure
+	SubjectPreset.Role = ULiveLinkAnimationRole::StaticClass();
+	StaticData->InitializeWith(FLiveLinkSkeletonStaticData::StaticStruct(), nullptr);
+	
+	//Take old skeleton data for new static data structure
 
-		FLiveLinkSkeletonStaticData* SkeletonData = StaticData->Cast<FLiveLinkSkeletonStaticData>();
-		SkeletonData->BoneNames = MoveTemp(RefSkeleton_DEPRECATED.BoneNames);
-		SkeletonData->BoneParents = MoveTemp(RefSkeleton_DEPRECATED.BoneParents);
+	FLiveLinkSkeletonStaticData* SkeletonData = StaticData->Cast<FLiveLinkSkeletonStaticData>();
+	SkeletonData->BoneNames = MoveTemp(RefSkeleton_DEPRECATED.BoneNames);
+	SkeletonData->BoneParents = MoveTemp(RefSkeleton_DEPRECATED.BoneParents);
 
-		//Create subsection that will manage this data
-		UMovieSceneLiveLinkSubSectionAnimation* AnimationSubSection = NewObject<UMovieSceneLiveLinkSubSectionAnimation>(this, UMovieSceneLiveLinkSubSectionAnimation::StaticClass(), NAME_None, RF_Transactional);
-		AnimationSubSection->Initialize(SubjectPreset.Role, StaticData);
+	//Create subsection that will manage this data
+	UMovieSceneLiveLinkSubSectionAnimation* AnimationSubSection = NewObject<UMovieSceneLiveLinkSubSectionAnimation>(this, UMovieSceneLiveLinkSubSectionAnimation::StaticClass(), NAME_None, RF_Transactional);
+	AnimationSubSection->Initialize(SubjectPreset.Role, StaticData);
 
-		//Convert transforms to float data channels
-		const int32 TransformCount = SkeletonData->BoneNames.Num() * 9;
+	//Convert transforms to float data channels
+	const int32 TransformCount = SkeletonData->BoneNames.Num() * 9;
+	if (TransformCount > 0)
+	{
 		ensure(AnimationSubSection->SubSectionData.Properties[0].FloatChannel.Num() == TransformCount);
 		AnimationSubSection->SubSectionData.Properties[0].FloatChannel = TArray<FMovieSceneFloatChannel>(PropertyFloatChannels_DEPRECATED.GetData(), TransformCount);
 		PropertyFloatChannels_DEPRECATED.RemoveAtSwap(0, TransformCount);
