- Fixing LiveLink anim node and retarget asset to support basic role for curve data
- Filtering non finite numbers from curves

diff --git a/Engine/Plugins/Animation/LiveLink/Source/LiveLink/Private/AnimNode_LiveLinkPose.cpp b/Engine/Plugins/Animation/LiveLink/Source/LiveLink/Private/AnimNode_LiveLinkPose.cpp
--- a/Engine/Plugins/Animation/LiveLink/Source/LiveLink/Private/AnimNode_LiveLinkPose.cpp
+++ b/Engine/Plugins/Animation/LiveLink/Source/LiveLink/Private/AnimNode_LiveLinkPose.cpp
@@ -64,16 +64,37 @@ void FAnimNode_LiveLinkPose::Evaluate_AnyThread(FPoseContext& Output)
 		return;
 	}
 
-	FLiveLinkSubjectFrameData AnimationFrame;
-	if(LiveLinkClient->EvaluateFrame_AnyThread(LiveLinkSubjectName, ULiveLinkAnimationRole::StaticClass(), AnimationFrame))
+	FLiveLinkSubjectFrameData SubjectFrameData;
+
+	TSubclassOf<ULiveLinkRole> SubjectRole = LiveLinkClient->GetSubjectRole(LiveLinkSubjectName);
+	if (SubjectRole->IsChildOf(ULiveLinkAnimationRole::StaticClass()))
 	{
-		FLiveLinkSkeletonStaticData* SkeletonData = AnimationFrame.StaticData.Cast<FLiveLinkSkeletonStaticData>();
-		FLiveLinkAnimationFrameData* FrameData = AnimationFrame.FrameData.Cast<FLiveLinkAnimationFrameData>();
-		check(SkeletonData);
-		check(FrameData);
+		//Process animation data if the subject is from that type
+		if (LiveLinkClient->EvaluateFrame_AnyThread(LiveLinkSubjectName, ULiveLinkAnimationRole::StaticClass(), SubjectFrameData))
+		{
+			FLiveLinkSkeletonStaticData* SkeletonData = SubjectFrameData.StaticData.Cast<FLiveLinkSkeletonStaticData>();
+			FLiveLinkAnimationFrameData* FrameData = SubjectFrameData.FrameData.Cast<FLiveLinkAnimationFrameData>();
+			check(SkeletonData);
+			check(FrameData);
+
+			CurrentRetargetAsset->BuildPoseFromAnimationData(CachedDeltaTime, SkeletonData, FrameData, Output.Pose);
+			CurrentRetargetAsset->BuildPoseAndCurveFromBaseData(CachedDeltaTime, SkeletonData, FrameData, Output.Pose, Output.Curve);
+			CachedDeltaTime = 0.f; // Reset so that if we evaluate again we don't "create" time inside of the retargeter
+		}
+	}
+	else
+	{
+		//Otherwise, fetch basic data that contains property / curve data
+		if (LiveLinkClient->EvaluateFrame_AnyThread(LiveLinkSubjectName, ULiveLinkBasicRole::StaticClass(), SubjectFrameData))
+		{
+			FLiveLinkBaseStaticData* BaseStaticData = SubjectFrameData.StaticData.Cast<FLiveLinkBaseStaticData>();
+			FLiveLinkBaseFrameData* BaseFrameData = SubjectFrameData.FrameData.Cast<FLiveLinkBaseFrameData>();
+			check(BaseStaticData);
+			check(BaseFrameData);
 
-		CurrentRetargetAsset->BuildPoseForSubject(CachedDeltaTime, SkeletonData, FrameData, Output.Pose, Output.Curve);
-		CachedDeltaTime = 0.f; // Reset so that if we evaluate again we don't "create" time inside of the retargeter
+			CurrentRetargetAsset->BuildPoseAndCurveFromBaseData(CachedDeltaTime, BaseStaticData, BaseFrameData, Output.Pose, Output.Curve);
+			CachedDeltaTime = 0.f; // Reset so that if we evaluate again we don't "create" time inside of the retargeter
+		}
 	}
 }
 
diff --git a/Engine/Plugins/Animation/LiveLink/Source/LiveLink/Private/LiveLinkRemapAsset.cpp b/Engine/Plugins/Animation/LiveLink/Source/LiveLink/Private/LiveLinkRemapAsset.cpp
--- a/Engine/Plugins/Animation/LiveLink/Source/LiveLink/Private/LiveLinkRemapAsset.cpp
+++ b/Engine/Plugins/Animation/LiveLink/Source/LiveLink/Private/LiveLinkRemapAsset.cpp
@@ -1,9 +1,11 @@
 // Copyright 1998-2019 Epic Games, Inc. All Rights Reserved.
 
 #include "LiveLinkRemapAsset.h"
-#include "LiveLinkTypes.h"
+
 #include "BonePose.h"
 #include "Engine/Blueprint.h"
+#include "GenericPlatform/GenericPlatformMath.h"
+#include "LiveLinkTypes.h"
 #include "Roles/LiveLinkAnimationRole.h"
 #include "Roles/LiveLinkAnimationTypes.h"
   
@@ -36,25 +38,25 @@ void ULiveLinkRemapAsset::OnBlueprintClassCompiled(UBlueprint* TargetBlueprint)
 	CurveNameMap.Reset();
 }
 
-void MakeCurveMapFromFrame(const FCompactPose& InPose, const FLiveLinkSkeletonStaticData* InSkeletonData, const FLiveLinkAnimationFrameData* InFrameData, const TArray<FName, TMemStackAllocator<>>& TransformedCurveNames, TMap<FName, float>& OutCurveMap)
+void MakeCurveMapFromFrame(const FCompactPose& InPose, const FLiveLinkBaseStaticData* InBaseStaticData, const FLiveLinkBaseFrameData* InFrameData, const TArray<FName, TMemStackAllocator<>>& TransformedCurveNames, TMap<FName, float>& OutCurveMap)
 {
 	OutCurveMap.Reset();
 	OutCurveMap.Reserve(InFrameData->PropertyValues.Num());
 
-	const USkeleton* Skeleton = InPose.GetBoneContainer().GetSkeletonAsset();
-
-	if (InSkeletonData->PropertyNames.Num() == InFrameData->PropertyValues.Num())
+	if (InBaseStaticData->PropertyNames.Num() == InFrameData->PropertyValues.Num())
 	{
-		for (int32 CurveIdx = 0; CurveIdx < InSkeletonData->PropertyNames.Num(); ++CurveIdx)
+		for (int32 CurveIdx = 0; CurveIdx < InBaseStaticData->PropertyNames.Num(); ++CurveIdx)
 		{
-			float PropertyValue = InFrameData->PropertyValues[CurveIdx];
-			OutCurveMap.Add(TransformedCurveNames[CurveIdx]) = PropertyValue;
+			const float PropertyValue = InFrameData->PropertyValues[CurveIdx];
+			if (FMath::IsFinite(PropertyValue))
+			{
+				OutCurveMap.Add(TransformedCurveNames[CurveIdx]) = PropertyValue;
+			}
 		}
 	}
 }
 
-
-void ULiveLinkRemapAsset::BuildPoseForSubject(float DeltaTime, const FLiveLinkSkeletonStaticData* InSkeletonData, const FLiveLinkAnimationFrameData* InFrameData, FCompactPose& OutPose, FBlendedCurve& OutCurve)
+void ULiveLinkRemapAsset::BuildPoseFromAnimationData(float DeltaTime, const FLiveLinkSkeletonStaticData* InSkeletonData, const FLiveLinkAnimationFrameData* InFrameData, FCompactPose& OutPose)
 {
 	const TArray<FName>& SourceBoneNames = InSkeletonData->BoneNames;
 
@@ -92,12 +94,15 @@ void ULiveLinkRemapAsset::BuildPoseForSubject(float DeltaTime, const FLiveLinkSk
 			}
 		}
 	}
+}
 
-	const TArray<FName>& SourceCurveNames = InSkeletonData->PropertyNames;
+void ULiveLinkRemapAsset::BuildPoseAndCurveFromBaseData(float DeltaTime, const FLiveLinkBaseStaticData* InBaseStaticData, const FLiveLinkBaseFrameData* InBaseFrameData, FCompactPose& OutPose, FBlendedCurve& OutCurve)
+{
+	const TArray<FName>& SourceCurveNames = InBaseStaticData->PropertyNames;
 	TArray<FName, TMemStackAllocator<>> TransformedCurveNames;
 	TransformedCurveNames.Reserve(SourceCurveNames.Num());
 
-	for(const FName& SrcCurveName : SourceCurveNames)
+	for (const FName& SrcCurveName : SourceCurveNames)
 	{
 		FName* TargetCurveName = CurveNameMap.Find(SrcCurveName);
 		if (TargetCurveName == nullptr)
@@ -114,7 +119,7 @@ void ULiveLinkRemapAsset::BuildPoseForSubject(float DeltaTime, const FLiveLinkSk
 
 	TMap<FName, float> BPCurveValues;
 
-	MakeCurveMapFromFrame(OutPose, InSkeletonData, InFrameData, TransformedCurveNames, BPCurveValues);
+	MakeCurveMapFromFrame(OutPose, InBaseStaticData, InBaseFrameData, TransformedCurveNames, BPCurveValues);
 
 	RemapCurveElements(BPCurveValues);
 
diff --git a/Engine/Plugins/Animation/LiveLink/Source/LiveLink/Private/LiveLinkRetargetAsset.cpp b/Engine/Plugins/Animation/LiveLink/Source/LiveLink/Private/LiveLinkRetargetAsset.cpp
--- a/Engine/Plugins/Animation/LiveLink/Source/LiveLink/Private/LiveLinkRetargetAsset.cpp
+++ b/Engine/Plugins/Animation/LiveLink/Source/LiveLink/Private/LiveLinkRetargetAsset.cpp
@@ -6,6 +6,7 @@
 #include "Animation/AnimTypes.h"
 #include "Animation/Skeleton.h"
 #include "BonePose.h"
+#include "GenericPlatform/GenericPlatformMath.h"
 #include "Roles/LiveLinkAnimationTypes.h"
 
 ULiveLinkRetargetAsset::ULiveLinkRetargetAsset(const FObjectInitializer& ObjectInitializer)
@@ -30,8 +31,11 @@ void ULiveLinkRetargetAsset::BuildCurveData(const FLiveLinkSkeletonStaticData* I
 	{
 		for (int32 CurveIdx = 0; CurveIdx < InSkeletonData->PropertyNames.Num(); ++CurveIdx)
 		{
-			float Curve = InFrameData->PropertyValues[CurveIdx];
-			ApplyCurveValue(Skeleton, InSkeletonData->PropertyNames[CurveIdx], Curve, OutCurve);
+			const float Curve = InFrameData->PropertyValues[CurveIdx];
+			if (FMath::IsFinite(Curve))
+			{
+				ApplyCurveValue(Skeleton, InSkeletonData->PropertyNames[CurveIdx], Curve, OutCurve);
+			}
 		}
 	}
 }
diff --git a/Engine/Plugins/Animation/LiveLink/Source/LiveLink/Public/LiveLinkRemapAsset.h b/Engine/Plugins/Animation/LiveLink/Source/LiveLink/Public/LiveLinkRemapAsset.h
--- a/Engine/Plugins/Animation/LiveLink/Source/LiveLink/Public/LiveLinkRemapAsset.h
+++ b/Engine/Plugins/Animation/LiveLink/Source/LiveLink/Public/LiveLinkRemapAsset.h
@@ -19,7 +19,8 @@ class LIVELINK_API ULiveLinkRemapAsset : public ULiveLinkRetargetAsset
 	//~ End UObject Interface
 
 	//~ Begin ULiveLinkRetargetAsset interface
-	virtual void BuildPoseForSubject(float DeltaTime, const FLiveLinkSkeletonStaticData* InSkeletonData, const FLiveLinkAnimationFrameData* InFrameData, FCompactPose& OutPose, FBlendedCurve& OutCurve) override;
+	virtual void BuildPoseFromAnimationData(float DeltaTime, const FLiveLinkSkeletonStaticData* InSkeletonData, const FLiveLinkAnimationFrameData* InFrameData, FCompactPose& OutPose) override;
+	virtual void BuildPoseAndCurveFromBaseData(float DeltaTime, const FLiveLinkBaseStaticData* InBaseStaticData, const FLiveLinkBaseFrameData* InBaseFrameData, FCompactPose& OutPose, FBlendedCurve& OutCurve) override;
 	//~ End ULiveLinkRetargetAsset interface
 
 	/** Blueprint Implementable function for getting a remapped bone name from the original */
diff --git a/Engine/Plugins/Animation/LiveLink/Source/LiveLink/Public/LiveLinkRetargetAsset.h b/Engine/Plugins/Animation/LiveLink/Source/LiveLink/Public/LiveLinkRetargetAsset.h
--- a/Engine/Plugins/Animation/LiveLink/Source/LiveLink/Public/LiveLinkRetargetAsset.h
+++ b/Engine/Plugins/Animation/LiveLink/Source/LiveLink/Public/LiveLinkRetargetAsset.h
@@ -11,6 +11,8 @@
 
 class USkeleton;
 struct FLiveLinkAnimationFrameData;
+struct FLiveLinkBaseFrameData;
+struct FLiveLinkBaseStaticData;
 struct FLiveLinkSkeletonStaticData;
 struct FCompactPose;
 struct FBlendedCurve;
@@ -34,5 +36,12 @@ class LIVELINK_API ULiveLinkRetargetAsset : public UObject
 	virtual void Initialize() {}
 
 	// Build OutPose and OutCurve from the supplied InFrame.
+	UE_DEPRECATED(4.23, "ULiveLinkRetargetAsset::BuildPoseForSubject is deprecated. Please use ULiveLinkRetargetAsset::BuildPoseFromAnimationData and ULiveLinkRetargetAsset::BuildPoseAndCurveFromBaseData instead.")
 	virtual void BuildPoseForSubject(float DeltaTime, const FLiveLinkSkeletonStaticData* InSkeletonData, const FLiveLinkAnimationFrameData* InFrameData, FCompactPose& OutPose, FBlendedCurve& OutCurve) PURE_VIRTUAL(ULiveLinkRetargetAsset::BuildPoseForSubject, );
+	
+	// Build OutPose from AnimationData if subject was from this type
+	virtual void BuildPoseFromAnimationData(float DeltaTime, const FLiveLinkSkeletonStaticData* InSkeletonData, const FLiveLinkAnimationFrameData* InFrameData, FCompactPose& OutPose) {}
+
+	// Build OutPose and OutCurve from the basic data. Called for every type of subjects
+	virtual void BuildPoseAndCurveFromBaseData(float DeltaTime, const FLiveLinkBaseStaticData* InBaseStaticData, const FLiveLinkBaseFrameData* InBaseFrameData, FCompactPose& OutPose, FBlendedCurve& OutCurve) {}
 };
