- Fixing backward compatibility for EditorActiveCamera in Animation PreviewController

diff --git a/Engine/Plugins/Animation/LiveLink/Source/LiveLinkEditor/Private/LiveLinkPreviewController.cpp b/Engine/Plugins/Animation/LiveLink/Source/LiveLinkEditor/Private/LiveLinkPreviewController.cpp
--- a/Engine/Plugins/Animation/LiveLink/Source/LiveLinkEditor/Private/LiveLinkPreviewController.cpp
+++ b/Engine/Plugins/Animation/LiveLink/Source/LiveLinkEditor/Private/LiveLinkPreviewController.cpp
@@ -10,8 +10,9 @@
 #include "LiveLinkCustomVersion.h"
 #include "LiveLinkInstance.h"
 #include "LiveLinkRemapAsset.h"
-#include "Roles/LiveLinkCameraRole.h"
-#include "Roles/LiveLinkCameraTypes.h"
+#include "Roles/LiveLinkAnimationRole.h"
+#include "Roles/LiveLinkTransformRole.h"
+#include "Roles/LiveLinkTransformTypes.h"
 
 const FName EditorCamera(TEXT("EditorActiveCamera"));
 
@@ -32,15 +33,38 @@ public:
 	{
 		if (ILiveLinkClient* Client = ClientRef.GetClient())
 		{
-			FLiveLinkSubjectFrameData CameraFrame;
-			if (Client->EvaluateFrame_AnyThread(EditorCamera, ULiveLinkCameraRole::StaticClass(), CameraFrame))
+			TSubclassOf<ULiveLinkRole> SubjectRole = Client->GetSubjectRole(EditorCamera);
+			if (SubjectRole)
 			{
-				FLiveLinkCameraFrameData* FrameData = CameraFrame.FrameData.Cast<FLiveLinkCameraFrameData>();
-
-				FTransform Camera = FrameData->Transform;
-				InOutCameraPosition = Camera.GetLocation();
-				InOutCameraEuler = Camera.GetRotation().Euler();
-				return;
+				//Old plugin will stream EditorCamera as AnimationRole through backward compatibility path. Otherwise, it should be of the Camera Role
+				FLiveLinkSubjectFrameData CurrentFrameData;
+				if (SubjectRole->IsChildOf(ULiveLinkAnimationRole::StaticClass()))
+				{
+					if (Client->EvaluateFrame_AnyThread(EditorCamera, ULiveLinkAnimationRole::StaticClass(), CurrentFrameData))
+					{
+						FLiveLinkAnimationFrameData* FrameData = CurrentFrameData.FrameData.Cast<FLiveLinkAnimationFrameData>();
+
+						if (FrameData->Transforms.Num() > 0)
+						{
+							FTransform Camera = FrameData->Transforms[0];
+							InOutCameraPosition = Camera.GetLocation();
+							InOutCameraEuler = Camera.GetRotation().Euler();
+							return;
+						}
+					}
+				}
+				else
+				{
+					if (Client->EvaluateFrame_AnyThread(EditorCamera, ULiveLinkTransformRole::StaticClass(), CurrentFrameData))
+					{
+						FLiveLinkTransformFrameData* FrameData = CurrentFrameData.FrameData.Cast<FLiveLinkTransformFrameData>();
+
+						FTransform Camera = FrameData->Transform;
+						InOutCameraPosition = Camera.GetLocation();
+						InOutCameraEuler = Camera.GetRotation().Euler();
+						return;
+					}
+				}
 			}
 		}
 
