Sequencer: Icons for geometry cache and live link tracks.

Disable the ability to create a live link track until it's implemented.

diff --git a/Engine/Plugins/Animation/LiveLink/Source/LiveLinkEditor/Private/Sequencer/LiveLinkPropertyTrackEditor.cpp b/Engine/Plugins/Animation/LiveLink/Source/LiveLinkEditor/Private/Sequencer/LiveLinkPropertyTrackEditor.cpp
--- a/Engine/Plugins/Animation/LiveLink/Source/LiveLinkEditor/Private/Sequencer/LiveLinkPropertyTrackEditor.cpp
+++ b/Engine/Plugins/Animation/LiveLink/Source/LiveLinkEditor/Private/Sequencer/LiveLinkPropertyTrackEditor.cpp
@@ -6,7 +6,9 @@
 #include "Misc/QualifiedFrameTime.h"
 #include "EditorStyleSet.h" //for add track stuff mz todo remove maybe
 #include "Features/IModularFeatures.h"
+#include "Styling/SlateIconFinder.h"
 #include "ISequencerSection.h"
+#include "LiveLinkComponent.h"
 
 /**
 * An implementation of live link property sections.
@@ -67,18 +69,6 @@ void FLiveLinkPropertyTrackEditor::BuildTrackContextMenu(FMenuBuilder& MenuBuild
 
 void FLiveLinkPropertyTrackEditor::BuildAddTrackMenu(FMenuBuilder& MenuBuilder)
 {
-	IModularFeatures& ModularFeatures = IModularFeatures::Get();
-
-
-	MenuBuilder.AddMenuEntry(
-		LOCTEXT("AddLiveLinkTrack", "Live Link Track"),
-		LOCTEXT("AddLiveLinkTrackTooltip", "Adds a new track that exposes Live Link Sources."),
-		FSlateIcon(FEditorStyle::GetStyleSetName(), "Sequencer.Tracks.LiveLink"),
-		FUIAction(
-			FExecuteAction::CreateRaw(this, &FLiveLinkPropertyTrackEditor::HandleAddLiveLinkTrackMenuEntryExecute),
-			FCanExecuteAction::CreateRaw(this, &FLiveLinkPropertyTrackEditor::HandleAddLiveLinkTrackMenuEntryCanExecute)
-		)
-	);
 }
 
 TSharedRef<ISequencerSection> FLiveLinkPropertyTrackEditor::MakeSectionInterface(UMovieSceneSection& SectionObject, UMovieSceneTrack& Track, FGuid ObjectBinding)
@@ -97,6 +87,11 @@ bool FLiveLinkPropertyTrackEditor::SupportsType(TSubclassOf<UMovieSceneTrack> Ty
 	return (Type == UMovieSceneLiveLinkTrack::StaticClass());
 }
 
+const FSlateBrush* FLiveLinkPropertyTrackEditor::GetIconBrush() const
+{
+	return FSlateIconFinder::FindIconForClass(ULiveLinkComponent::StaticClass()).GetIcon();
+}
+
 
 /* FLiveLinkTrackEditor callbacks
 *****************************************************************************/
diff --git a/Engine/Plugins/Animation/LiveLink/Source/LiveLinkEditor/Private/Sequencer/LiveLinkPropertyTrackEditor.h b/Engine/Plugins/Animation/LiveLink/Source/LiveLinkEditor/Private/Sequencer/LiveLinkPropertyTrackEditor.h
--- a/Engine/Plugins/Animation/LiveLink/Source/LiveLinkEditor/Private/Sequencer/LiveLinkPropertyTrackEditor.h
+++ b/Engine/Plugins/Animation/LiveLink/Source/LiveLinkEditor/Private/Sequencer/LiveLinkPropertyTrackEditor.h
@@ -35,6 +35,7 @@ public:
 	virtual bool SupportsSequence(UMovieSceneSequence* InSequence) const override;
 	virtual bool SupportsType(TSubclassOf<UMovieSceneTrack> Type) const override;
 	virtual void BuildTrackContextMenu(FMenuBuilder& MenuBuilder, UMovieSceneTrack* Track) override;
+	virtual const FSlateBrush* GetIconBrush() const override;
 
 protected:
 
diff --git a/Engine/Plugins/Experimental/GeometryCache/Source/GeometryCacheSequencer/Classes/GeometryCacheTrackEditor.cpp b/Engine/Plugins/Experimental/GeometryCache/Source/GeometryCacheSequencer/Classes/GeometryCacheTrackEditor.cpp
--- a/Engine/Plugins/Experimental/GeometryCache/Source/GeometryCacheSequencer/Classes/GeometryCacheTrackEditor.cpp
+++ b/Engine/Plugins/Experimental/GeometryCache/Source/GeometryCacheSequencer/Classes/GeometryCacheTrackEditor.cpp
@@ -28,6 +28,7 @@
 #include "Widgets/Notifications/SNotificationList.h"
 #include "GeometryCacheComponent.h"
 #include "GeometryCache.h"
+#include "Styling/SlateIconFinder.h"
 
 namespace GeometryCacheEditorConstants
 {
@@ -375,5 +376,10 @@ TSharedPtr<SWidget> FGeometryCacheTrackEditor::BuildOutlinerEditWidget(const FGu
 	
 }
 
+const FSlateBrush* FGeometryCacheTrackEditor::GetIconBrush() const
+{
+	return FSlateIconFinder::FindIconForClass(UGeometryCacheComponent::StaticClass()).GetIcon();
+}
+
 
 #undef LOCTEXT_NAMESPACE
diff --git a/Engine/Plugins/Experimental/GeometryCache/Source/GeometryCacheSequencer/Classes/GeometryCacheTrackEditor.h b/Engine/Plugins/Experimental/GeometryCache/Source/GeometryCacheSequencer/Classes/GeometryCacheTrackEditor.h
--- a/Engine/Plugins/Experimental/GeometryCache/Source/GeometryCacheSequencer/Classes/GeometryCacheTrackEditor.h
+++ b/Engine/Plugins/Experimental/GeometryCache/Source/GeometryCacheSequencer/Classes/GeometryCacheTrackEditor.h
@@ -49,6 +49,7 @@ public:
 	virtual TSharedRef<ISequencerSection> MakeSectionInterface( UMovieSceneSection& SectionObject, UMovieSceneTrack& Track, FGuid ObjectBinding ) override;
 	virtual bool SupportsType( TSubclassOf<UMovieSceneTrack> Type ) const override;
 	virtual TSharedPtr<SWidget> BuildOutlinerEditWidget(const FGuid& ObjectBinding, UMovieSceneTrack* Track, const FBuildEditWidgetParams& Params) override;
+	virtual const FSlateBrush* GetIconBrush() const override;
 
 private:
 
