adding LIVELINK_API on the remap asset

#code reivew: Martin.Wilson

diff --git a/Engine/Plugins/Animation/LiveLink/Source/LiveLink/Public/LiveLinkRemapAsset.h b/Engine/Plugins/Animation/LiveLink/Source/LiveLink/Public/LiveLinkRemapAsset.h
--- a/Engine/Plugins/Animation/LiveLink/Source/LiveLink/Public/LiveLinkRemapAsset.h
+++ b/Engine/Plugins/Animation/LiveLink/Source/LiveLink/Public/LiveLinkRemapAsset.h
@@ -8,7 +8,7 @@
 
 // Remap asset for data coming from Live Link. Allows simple application of bone transforms into current pose based on name remapping only
 UCLASS(Blueprintable)
-class ULiveLinkRemapAsset : public ULiveLinkRetargetAsset
+class LIVELINK_API ULiveLinkRemapAsset : public ULiveLinkRetargetAsset
 {
 	GENERATED_UCLASS_BODY()
 
