Copied from Dev-Enterprise - Fix for Live Link disconnecting when editor is busy
Root cause was FHeartbeatManager would send the heartbeat on a separate thread however we were only receiving the response on the Game Thread causing replies to be dropped during a slow task.

Resolved by using ReceivingOnAnyThread() in the FMessageEndpoint::Builder and updating the relevant handlers to be Threadsafe

+ Fix for Message Bus Source timing out when connecting with large number of characters.

Resolved by correctly initialized ConnectionLastActive variable when the Message Bus Source receives a Client.

diff --git a/Engine/Plugins/Animation/LiveLink/Source/LiveLink/Private/LiveLinkMessageBusSource.cpp b/Engine/Plugins/Animation/LiveLink/Source/LiveLink/Private/LiveLinkMessageBusSource.cpp
--- a/Engine/Plugins/Animation/LiveLink/Source/LiveLink/Private/LiveLinkMessageBusSource.cpp
+++ b/Engine/Plugins/Animation/LiveLink/Source/LiveLink/Private/LiveLinkMessageBusSource.cpp
@@ -1,4 +1,4 @@
-﻿// Copyright 1998-2019 Epic Games, Inc. All Rights Reserved.
+// Copyright 1998-2019 Epic Games, Inc. All Rights Reserved.
 
 #include "LiveLinkMessageBusSource.h"
 #include "LiveLinkMessages.h"
@@ -19,7 +19,8 @@ void FLiveLinkMessageBusSource::ReceiveClient(ILiveLinkClient* InClient, FGuid I
 					  .Handling<FLiveLinkSubjectDataMessage>(this, &FLiveLinkMessageBusSource::HandleSubjectData)
 					  .Handling<FLiveLinkSubjectFrameMessage>(this, &FLiveLinkMessageBusSource::HandleSubjectFrame)
 					  .Handling<FLiveLinkHeartbeatMessage>(this, &FLiveLinkMessageBusSource::HandleHeartbeat)
-					  .Handling<FLiveLinkClearSubject>(this, &FLiveLinkMessageBusSource::HandleClearSubject);
+					  .Handling<FLiveLinkClearSubject>(this, &FLiveLinkMessageBusSource::HandleClearSubject)
+					  .ReceivingOnAnyThread();
 
 
 	MessageEndpoint->Send(new FLiveLinkConnectMessage(), ConnectionAddress);
@@ -27,17 +28,24 @@ void FLiveLinkMessageBusSource::ReceiveClient(ILiveLinkClient* InClient, FGuid I
 	// Register for heartbeats
 	bIsValid = true;
 	FHeartbeatManager::Get()->RegisterSource(this);
+
+	UpdateConnectionLastActive();
 }
 
 bool FLiveLinkMessageBusSource::SendHeartbeat()
 {
 	const double CurrentTime = FPlatformTime::Seconds();
 
-	if (HeartbeatLastSent > (CurrentTime - LL_HALF_CONNECTION_TIMEOUT) &&
-		ConnectionLastActive < (CurrentTime - LL_CONNECTION_TIMEOUT))
 	{
-		//We have recently tried to heartbeat and not received anything back
-		bIsValid = false;
+		// Ensure the read of ConnectionLastActive is Threadsafe
+		FScopeLock ConnectionTimeLock(&ConnectionLastActiveSection);
+
+		if (HeartbeatLastSent > (CurrentTime - LL_HALF_CONNECTION_TIMEOUT) &&
+			ConnectionLastActive < (CurrentTime - LL_CONNECTION_TIMEOUT))
+		{
+			//We have recently tried to heartbeat and not received anything back
+			bIsValid = false;
+		}
 	}
 
 	MessageEndpoint->Send(new FLiveLinkHeartbeatMessage(), ConnectionAddress);
@@ -53,15 +61,23 @@ bool FLiveLinkMessageBusSource::IsSourceStillValid()
 
 void FLiveLinkMessageBusSource::HandleHeartbeat(const FLiveLinkHeartbeatMessage& Message, const TSharedRef<IMessageContext, ESPMode::ThreadSafe>& Context)
 {
-	ConnectionLastActive = FPlatformTime::Seconds();
+	UpdateConnectionLastActive();
 }
 
 void FLiveLinkMessageBusSource::HandleClearSubject(const FLiveLinkClearSubject& Message, const TSharedRef<IMessageContext, ESPMode::ThreadSafe>& Context)
 {
-	ConnectionLastActive = FPlatformTime::Seconds();
+	UpdateConnectionLastActive();
+
 	Client->ClearSubject(Message.SubjectName);
 }
 
+FORCEINLINE void FLiveLinkMessageBusSource::UpdateConnectionLastActive()
+{
+	FScopeLock ConnectionTimeLock(&ConnectionLastActiveSection);
+
+	ConnectionLastActive = FPlatformTime::Seconds();
+}
+
 bool FLiveLinkMessageBusSource::RequestSourceShutdown()
 {
 	FHeartbeatManager* HeartbeatManager = FHeartbeatManager::Get();
@@ -75,50 +91,14 @@ bool FLiveLinkMessageBusSource::RequestSourceShutdown()
 
 void FLiveLinkMessageBusSource::HandleSubjectData(const FLiveLinkSubjectDataMessage& Message, const TSharedRef<IMessageContext, ESPMode::ThreadSafe>& Context)
 {
-	ConnectionLastActive = FPlatformTime::Seconds();
-
-	//FPlatformMisc::LowLevelOutputDebugStringf(TEXT("HandleSubjectData %s\n"), *Message.SubjectName);
-	/*for (const FString& Name : Message.BoneNames)
-	{
-	FPlatformMisc::LowLevelOutputDebugStringf(TEXT("\tName: %s\n"), *Name);
-	}*/
-	/*FScopeLock Lock(&GBoneDataCS);
-
-	if (BoneID != (Message.BoneID - 1))
-	{
-	UE_LOG(LogTemp, Warning, TEXT("BONE ID SKIP Was On:%i Now:%i"), BoneID, Message.BoneID);
-	}
-	if (BoneNames.Num() == Message.BoneNames.Num() || BoneNames.Num() == 0)
-	{
-	BoneNames.Reset();
-	for (const FString& Name : Message.BoneNames)
-	{
-	BoneNames.Add(FName(*Name));
-	}
-	//BoneTransforms.Reset();
-	BoneID = Message.BoneID;
-	}
-	else
-	{
-	UE_LOG(LogTemp, Warning, TEXT("INVALID BONE NAMES RECIEVED %i != existing %i"), Message.BoneNames.Num(), BoneNames.Num());
-	}*/
+	UpdateConnectionLastActive();
 
 	Client->PushSubjectSkeleton(SourceGuid, Message.SubjectName, Message.RefSkeleton);
 }
 
 void FLiveLinkMessageBusSource::HandleSubjectFrame(const FLiveLinkSubjectFrameMessage& Message, const TSharedRef<IMessageContext, ESPMode::ThreadSafe>& Context)
 {
-	ConnectionLastActive = FPlatformTime::Seconds();
-
-	//FPlatformMisc::LowLevelOutputDebugString(TEXT("HandleSubjectFrame\n"));
-	/*if (BoneID != Message.BoneID)
-	{
-	UE_LOG(LogTemp, Warning, TEXT("BONE ID MISMATCH Exp:%i Got:%i"), BoneID, Message.BoneID);
-	}*/
-	/*for (const FTransform& T : Message.Transforms)
-	{
-	FPlatformMisc::LowLevelOutputDebugStringf(TEXT("\tTransform: %s\n"), *T.ToString());
-	}*/
+	UpdateConnectionLastActive();
 
 	FLiveLinkFrameData FrameData;
 	FrameData.Transforms = Message.Transforms;
diff --git a/Engine/Plugins/Animation/LiveLink/Source/LiveLink/Public/LiveLinkMessageBusSource.h b/Engine/Plugins/Animation/LiveLink/Source/LiveLink/Public/LiveLinkMessageBusSource.h
--- a/Engine/Plugins/Animation/LiveLink/Source/LiveLink/Public/LiveLinkMessageBusSource.h
+++ b/Engine/Plugins/Animation/LiveLink/Source/LiveLink/Public/LiveLinkMessageBusSource.h
@@ -49,6 +49,9 @@ private:
 	void HandleClearSubject(const FLiveLinkClearSubject& Message, const TSharedRef<IMessageContext, ESPMode::ThreadSafe>& Context);
 	// End Message bus message handlers
 
+	// Threadsafe update of the last active time
+	FORCEINLINE void UpdateConnectionLastActive();
+
 	ILiveLinkClient* Client;
 
 	// Our identifier in LiveLink
@@ -70,4 +73,7 @@ private:
 
 	// Current Validity of Source
 	FThreadSafeBool bIsValid;
+
+	// Critical section to allow for threadsafe updating of the connection time
+	FCriticalSection ConnectionLastActiveSection;
 };
\ No newline at end of file
