Sequencer: Rename MovieSceneLiveLinkSectionTemplate.cpp

diff --git a/Engine/Plugins/Animation/LiveLink/Source/LiveLink/Private/MovieScene/MoivieSceneLiveLinkSectionTemplate.cpp b/Engine/Plugins/Animation/LiveLink/Source/LiveLink/Private/MovieScene/MoivieSceneLiveLinkSectionTemplate.cpp
deleted file mode 100644
diff --git a/Engine/Plugins/Animation/LiveLink/Source/LiveLink/Private/MovieScene/MovieSceneLiveLinkSectionTemplate.cpp b/Engine/Plugins/Animation/LiveLink/Source/LiveLink/Private/MovieScene/MovieSceneLiveLinkSectionTemplate.cpp
new file mode 100644
--- /dev/null
+++ b/Engine/Plugins/Animation/LiveLink/Source/LiveLink/Private/MovieScene/MovieSceneLiveLinkSectionTemplate.cpp
@@ -0,0 +1,259 @@
+// Copyright 1998-2019 Epic Games, Inc. All Rights Reserved.
+
+#include "MovieScene/MovieSceneLiveLinkSectionTemplate.h"
+#include "MovieScene/MovieSceneLiveLinkTrack.h"
+#include "MovieScene/MovieSceneLiveLinkSection.h"
+#include "MovieSceneLiveLinkSource.h"
+#include "Channels/MovieSceneChannelProxy.h"
+#include "Engine/Engine.h"
+#include "Engine/TimecodeProvider.h"
+#include "HAL/ConsoleManager.h"
+
+static TAutoConsoleVariable<int32> CVarSequencerAlwaysSendInterpolatedLiveLink(
+	TEXT("Sequencer.AlwaysSendInterpolatedLiveLink"),
+	0,
+	TEXT("If nonzero we always interpolate when sending out live link data, if 0 we may send out frames at a higher rate than engine tick, if the data is dense."),
+	ECVF_Default);
+
+struct FMovieSceneLiveLinkSectionTemplatePersistentData : IPersistentEvaluationData
+{
+	TSharedPtr<FMovieSceneLiveLinkSource> LiveLinkSource;
+};
+
+FMovieSceneLiveLinkTemplateData::FMovieSceneLiveLinkTemplateData(const UMovieSceneLiveLinkSection& Section)
+{
+	TArrayView<FMovieSceneFloatChannel*> ProxyFloatChannels = Section.GetChannelProxy().GetChannels<FMovieSceneFloatChannel>();
+	FloatChannels.SetNum(ProxyFloatChannels.Num());
+	for (int32 Index = 0; Index < ProxyFloatChannels.Num(); ++Index)
+	{
+		FloatChannels[Index] = *ProxyFloatChannels[Index];
+	}
+	TemplateToPush = Section.TemplateToPush;
+	SubjectName = Section.SubjectName;
+	RefSkeleton = Section.RefSkeleton;
+	ChannelMask = Section.ChannelMask;
+}
+
+//Converts time's in our movie scene frame rate to times in the time code frame rate, based upon where our frame time is and where the timecode frame time is.
+static FQualifiedFrameTime ConvertFrameTimeToTimeCodeTime(const FFrameTime& FrameTime, const FFrameRate& FrameRate, const FFrameTime& FrameTimeEqualToTimecodeFrameTime, const FQualifiedFrameTime& TimecodeTime)
+{
+	FFrameTime DiffFrameTime = FrameTime - FrameTimeEqualToTimecodeFrameTime;
+	DiffFrameTime = FFrameRate::TransformTime(DiffFrameTime, FrameRate, TimecodeTime.Rate);
+	return FQualifiedFrameTime(TimecodeTime.Time + DiffFrameTime, TimecodeTime.Rate);
+}
+
+static FLiveLinkWorldTime ConvertFrameTimeToLiveLinkWorldTime(const FFrameTime& FrameTime, const FFrameRate& FrameRate, const FFrameTime& FrameTimeEqualToWorldFrameTime, const FLiveLinkWorldTime& LiveLinkWorldTime)
+{
+	FLiveLinkWorldTime WorldTime;
+	FFrameTime DiffFrameTime = FrameTime - FrameTimeEqualToWorldFrameTime;
+	double DiffSeconds = FrameRate.AsSeconds(DiffFrameTime);
+	return WorldTime.Time = DiffSeconds + (LiveLinkWorldTime.Time  + LiveLinkWorldTime.Offset);
+}
+
+bool FMovieSceneLiveLinkTemplateData::GetLiveLinkFrameArray(const FFrameTime& FrameTime, const FFrameTime& LowerBound, const FFrameTime& UpperBound, TArray<FLiveLinkFrameData>&  LiveLinkFrameDataArray, const FFrameRate& FrameRate) const
+{
+	//See if we have a valid time code time. 
+	//If so we may can possible send raw data if not asked to only send interpolated.
+	TOptional<FQualifiedFrameTime> TimeCodeFrameTime;
+	if (GEngine && GEngine->GetTimecodeProvider() && GEngine->GetTimecodeProvider()->GetSynchronizationState() == ETimecodeProviderSynchronizationState::Synchronized)
+	{
+		const UTimecodeProvider* TimeCodeProvider = GEngine->GetTimecodeProvider();
+		FFrameRate TCFrameRate = TimeCodeProvider->GetFrameRate();
+		FTimecode TimeCode = TimeCodeProvider->GetTimecode(); //Same as FApp::GetTimecode();
+		FFrameNumber FrameNumber = TimeCode.ToFrameNumber(TCFrameRate);
+		TimeCodeFrameTime = FQualifiedFrameTime(FFrameTime(FrameNumber), TCFrameRate);
+	}
+
+	//Send interpolated if told to or no valid timecode synced.
+	const bool bAlwaysSendInterpolated = CVarSequencerAlwaysSendInterpolatedLiveLink->GetInt() == 0 ? false : true;
+
+	bool bSendInterpolated = bAlwaysSendInterpolated || !TimeCodeFrameTime.IsSet() || LowerBound == UpperBound;
+	FLiveLinkWorldTime WorldTime = FLiveLinkWorldTime(); //this calls FPlatform::Seconds()
+	FVector Vector;
+	
+	if (!bSendInterpolated)
+	{
+		FFrameTime FrameRangeEnd = LowerBound > UpperBound ? LowerBound : UpperBound;
+		FFrameTime FrameRangeStart = LowerBound > UpperBound ? UpperBound : LowerBound;
+		if (TemplateToPush.Transforms.Num() > 0)
+		{
+			//assume times will be the same for all channels...
+			TArrayView<const FFrameNumber> Times = FloatChannels[0].GetTimes();
+
+			int32 EndIndex = INDEX_NONE, StartIndex = INDEX_NONE;
+			EndIndex = Algo::LowerBound(Times, FrameRangeEnd.FrameNumber);
+
+			FFrameNumber Frame;
+			if (EndIndex != INDEX_NONE)
+			{
+				if (EndIndex >= Times.Num())
+				{
+					EndIndex = Times.Num() - 1;
+				}
+				
+				StartIndex = Algo::UpperBound(Times, FrameRangeStart.FrameNumber);
+				if (StartIndex == INDEX_NONE)
+				{
+					StartIndex = EndIndex;
+				}
+				
+			}
+			else
+			{
+				StartIndex = Algo::UpperBound(Times, FrameRangeStart.FrameNumber);
+				if (StartIndex >= Times.Num())
+				{
+					StartIndex = Times.Num() - 1;
+				}
+				if (StartIndex != INDEX_NONE)
+				{
+					EndIndex = StartIndex;
+				}
+			}
+			bSendInterpolated = true; //if we don't send at least one key send interpolated
+			if (EndIndex != INDEX_NONE)
+			{
+				//UE_LOG(LogTemp, Log, TEXT("Send Key LiveLink Start/End Index '%d'  '%d'"), StartIndex,EndIndex);
+				for (int32 Index = StartIndex; Index <= EndIndex; ++Index)
+				{
+					Frame = Times[Index];
+					if (Frame > FrameRangeStart && Frame <= FrameRangeEnd) // doing (begin,end] want to make sure we get the last frame always, future better than past.
+					{
+						//UE_LOG(LogTemp, Log, TEXT("Send Key LiveLink Key Index '%d'"), Index);
+						bSendInterpolated = false;
+						FLiveLinkFrameData FrameData;
+						FrameData.Transforms.Reset(TemplateToPush.Transforms.Num());
+						FrameData.WorldTime = ConvertFrameTimeToLiveLinkWorldTime(Times[Index], FrameRate, FrameTime, WorldTime);
+						if (TimeCodeFrameTime.IsSet())
+						{
+							FrameData.MetaData.SceneTime = ConvertFrameTimeToTimeCodeTime(Times[Index], FrameRate, FrameTime, TimeCodeFrameTime.GetValue());
+						}
+						int32 ChannelIndex = 0;
+						for (FTransform Transform : TemplateToPush.Transforms)
+						{
+							Vector.X = FloatChannels[ChannelIndex++].GetValues()[Index].Value;
+							Vector.Y = FloatChannels[ChannelIndex++].GetValues()[Index].Value;
+							Vector.Z = FloatChannels[ChannelIndex++].GetValues()[Index].Value;
+							Transform.SetLocation(Vector);
+							Vector.X = FloatChannels[ChannelIndex++].GetValues()[Index].Value;
+							Vector.Y = FloatChannels[ChannelIndex++].GetValues()[Index].Value;
+							Vector.Z = FloatChannels[ChannelIndex++].GetValues()[Index].Value;
+							FRotator Rotator(Vector.Y, Vector.Z, Vector.X); //pitch, yaw, roll
+							FQuat Quat = Rotator.Quaternion();
+							//mz todo handle flips somehow if needed. 
+							Transform.SetRotation(Quat);
+							Vector.X = FloatChannels[ChannelIndex++].GetValues()[Index].Value;
+							Vector.Y = FloatChannels[ChannelIndex++].GetValues()[Index].Value;
+							Vector.Z = FloatChannels[ChannelIndex++].GetValues()[Index].Value;
+							Transform.SetScale3D(Vector);
+							FrameData.Transforms.Add(Transform);
+						}
+
+						//handle curves also...
+						LiveLinkFrameDataArray.Add(FrameData);
+
+					}
+				}
+			}
+		}
+		else
+		{
+			int CurveStart = 9 * TemplateToPush.Transforms.Num();
+		}
+	}
+	if (bSendInterpolated)
+	{
+	//	UE_LOG(LogTemp, Log, TEXT("Send LiveLink Interpolated"));
+		static FLiveLinkFrameData FrameData; //NOTE if we send more then one frame this won't work.
+		
+		FrameData.Transforms.Reset(TemplateToPush.Transforms.Num());
+		FrameData.CurveElements.Reset(TemplateToPush.CurveElements.Num());
+		//send both engine time and if we have a synchronized timecode provider the qualified time also
+		FrameData.WorldTime =  WorldTime;
+
+		if (TimeCodeFrameTime.IsSet())
+		{
+			FrameData.MetaData.SceneTime = TimeCodeFrameTime.GetValue();
+		}
+		int32 ChannelIndex = 0;
+
+		for (FTransform Transform : TemplateToPush.Transforms)
+		{
+			FloatChannels[ChannelIndex++].Evaluate(FrameTime, Vector.X);
+			FloatChannels[ChannelIndex++].Evaluate(FrameTime, Vector.Y);
+			FloatChannels[ChannelIndex++].Evaluate(FrameTime, Vector.Z);
+			Transform.SetLocation(Vector);
+			FloatChannels[ChannelIndex++].Evaluate(FrameTime, Vector.X);
+			FloatChannels[ChannelIndex++].Evaluate(FrameTime, Vector.Y);
+			FloatChannels[ChannelIndex++].Evaluate(FrameTime, Vector.Z);
+			FRotator Rotator(Vector.Y, Vector.Z, Vector.X); //pitch, yaw, roll
+			FQuat Quat = Rotator.Quaternion();
+			//mz todo handle flips somehow if needed. 
+			Transform.SetRotation(Quat);
+			FloatChannels[ChannelIndex++].Evaluate(FrameTime, Vector.X);
+			FloatChannels[ChannelIndex++].Evaluate(FrameTime, Vector.Y);
+			FloatChannels[ChannelIndex++].Evaluate(FrameTime, Vector.Z);
+			Transform.SetScale3D(Vector);
+			FrameData.Transforms.Add(Transform);
+		}
+		for (FLiveLinkCurveElement CurveElement : TemplateToPush.CurveElements)
+		{
+			FloatChannels[ChannelIndex++].Evaluate(FrameTime, CurveElement.CurveValue);
+			FrameData.CurveElements.Add(CurveElement);
+		}
+
+		LiveLinkFrameDataArray.Add(FrameData);
+	}
+	return true;
+}
+
+FMovieSceneLiveLinkSectionTemplate::FMovieSceneLiveLinkSectionTemplate(const UMovieSceneLiveLinkSection& Section, const UMovieScenePropertyTrack& Track)
+	: FMovieScenePropertySectionTemplate(Track.GetPropertyName(), Track.GetPropertyPath()),  TemplateData(Section)
+{
+
+}
+
+
+void FMovieSceneLiveLinkSectionTemplate::EvaluateSwept(const FMovieSceneEvaluationOperand& Operand, const FMovieSceneContext& Context, const TRange<FFrameNumber>& SweptRange, const FPersistentEvaluationData& PersistentData, FMovieSceneExecutionTokens& ExecutionTokens) const
+{
+	FMovieSceneLiveLinkSectionTemplatePersistentData* Data = PersistentData.FindSectionData<FMovieSceneLiveLinkSectionTemplatePersistentData>();
+	if (Data && Data->LiveLinkSource.IsValid() && Data->LiveLinkSource->IsSourceStillValid())
+	{
+		TArray<FLiveLinkFrameData>  LiveLinkFrameDataArray;
+
+		TemplateData.GetLiveLinkFrameArray(Context.GetTime(), SweptRange.GetLowerBoundValue(), SweptRange.GetUpperBoundValue(), LiveLinkFrameDataArray, Context.GetFrameRate());
+		Data->LiveLinkSource->PublishLiveLinkFrameData(TemplateData.SubjectName, LiveLinkFrameDataArray, TemplateData.RefSkeleton);
+	}
+}
+
+void FMovieSceneLiveLinkSectionTemplate::Evaluate(const FMovieSceneEvaluationOperand& Operand, const FMovieSceneContext& Context, const FPersistentEvaluationData& PersistentData, FMovieSceneExecutionTokens& ExecutionTokens) const
+{
+	FMovieSceneLiveLinkSectionTemplatePersistentData* Data = PersistentData.FindSectionData<FMovieSceneLiveLinkSectionTemplatePersistentData>();
+	if (Data && Data->LiveLinkSource.IsValid() && Data->LiveLinkSource->IsSourceStillValid())
+	{
+		TArray<FLiveLinkFrameData>  LiveLinkFrameDataArray;
+		FFrameTime FrameTime = Context.GetTime();
+		TemplateData.GetLiveLinkFrameArray(FrameTime, FrameTime, FrameTime,LiveLinkFrameDataArray, Context.GetFrameRate());
+		Data->LiveLinkSource->PublishLiveLinkFrameData(TemplateData.SubjectName, LiveLinkFrameDataArray,TemplateData.RefSkeleton);
+	}
+
+}
+
+void FMovieSceneLiveLinkSectionTemplate::Setup(FPersistentEvaluationData& PersistentData, IMovieScenePlayer& Player) const
+{
+	PersistentData.GetOrAddSectionData<FMovieSceneLiveLinkSectionTemplatePersistentData>().LiveLinkSource = FMovieSceneLiveLinkSource::CreateLiveLinkSource(TemplateData.SubjectName);
+}
+
+void FMovieSceneLiveLinkSectionTemplate::TearDown(FPersistentEvaluationData& PersistentData, IMovieScenePlayer& Player) const
+{
+	FMovieSceneLiveLinkSectionTemplatePersistentData* Data = PersistentData.FindSectionData<FMovieSceneLiveLinkSectionTemplatePersistentData>();
+
+	if (Data && Data->LiveLinkSource.IsValid())
+	{
+		if ((Data->LiveLinkSource)->IsSourceStillValid())
+		{
+			FMovieSceneLiveLinkSource::RemoveLiveLinkSource(Data->LiveLinkSource, TemplateData.SubjectName);
+		}
+		Data->LiveLinkSource.Reset();
+	}
+}
\ No newline at end of file
