Fix for live link re-winding, needed to store the last rotation per recording object. Fixed both Take and Sequencer.

diff --git a/Engine/Plugins/Animation/LiveLink/Source/LiveLink/Public/MovieScene/MovieSceneLiveLinkBufferData.h b/Engine/Plugins/Animation/LiveLink/Source/LiveLink/Public/MovieScene/MovieSceneLiveLinkBufferData.h
--- a/Engine/Plugins/Animation/LiveLink/Source/LiveLink/Public/MovieScene/MovieSceneLiveLinkBufferData.h
+++ b/Engine/Plugins/Animation/LiveLink/Source/LiveLink/Public/MovieScene/MovieSceneLiveLinkBufferData.h
@@ -54,8 +54,7 @@ struct FLiveLinkTransformKeys
 	TArray<FMovieSceneFloatValue> RotationX, RotationY, RotationZ;
 	TArray<FMovieSceneFloatValue> ScaleX, ScaleY, ScaleZ;
 
-	//needed for rewinding, when we set the values we keep track of the last value set to restart the re-winding from that.
-	TOptional<FVector> LastRotationValues;
+
 
 	void Add(const FTransform& InTransform)
 	{
@@ -138,7 +137,7 @@ struct FLiveLinkTransformKeys
 	//This function is the one that's called when recording live link incrementally. We move the values over from our saved 
 	//Location, Rotation and Scale buffers into the specified float channels and then reset our buffers, re-using it's memory 
 	//for the next iteration. We also fix any euler flips during this process, avoiding iterating over the data once again during Finalize.
-	void AppendToFloatChannelsAndReset(int32 StartIndex, TArray<FMovieSceneFloatChannel> &FloatChannels, const TArray<FFrameNumber> &Times)
+	void AppendToFloatChannelsAndReset(int32 StartIndex, TArray<FMovieSceneFloatChannel>& FloatChannels, const TArray<FFrameNumber>& Times, TOptional<FVector>& LastRotationValues)
 	{
 		if (Times.Num() > 0)
 		{
diff --git a/Engine/Plugins/Animation/LiveLink/Source/LiveLinkEditor/Private/SequencerRecorderSections/MovieSceneLiveLinkSectionRecorder.cpp b/Engine/Plugins/Animation/LiveLink/Source/LiveLinkEditor/Private/SequencerRecorderSections/MovieSceneLiveLinkSectionRecorder.cpp
--- a/Engine/Plugins/Animation/LiveLink/Source/LiveLinkEditor/Private/SequencerRecorderSections/MovieSceneLiveLinkSectionRecorder.cpp
+++ b/Engine/Plugins/Animation/LiveLink/Source/LiveLinkEditor/Private/SequencerRecorderSections/MovieSceneLiveLinkSectionRecorder.cpp
@@ -216,6 +216,9 @@ void FMovieSceneLiveLinkSectionRecorder::CreateTracks(UMovieScene* InMovieScene,
 			Serializer.Close(); //just read header.
 		}
 		*/
+
+		LastRotationValues.Reset();
+
 	}
 }
 
@@ -323,7 +326,7 @@ void FMovieSceneLiveLinkSectionRecorder::Record(float CurrentTime)
 					int32 ChannelIndex = 0;
 					for (FLiveLinkTransformKeys& TransformKeys : LinkTransformKeysArray)
 					{
-						TransformKeys.AppendToFloatChannelsAndReset(ChannelIndex, FloatChannels, Times);
+						TransformKeys.AppendToFloatChannelsAndReset(ChannelIndex, FloatChannels, Times, LastRotationValues);
 						ChannelIndex += 9;
 					}
 					for (FLiveLinkCurveKeys CurveKeys : LinkCurveKeysArray)
diff --git a/Engine/Plugins/Animation/LiveLink/Source/LiveLinkEditor/Private/SequencerRecorderSections/MovieSceneLiveLinkSectionRecorder.h b/Engine/Plugins/Animation/LiveLink/Source/LiveLinkEditor/Private/SequencerRecorderSections/MovieSceneLiveLinkSectionRecorder.h
--- a/Engine/Plugins/Animation/LiveLink/Source/LiveLinkEditor/Private/SequencerRecorderSections/MovieSceneLiveLinkSectionRecorder.h
+++ b/Engine/Plugins/Animation/LiveLink/Source/LiveLinkEditor/Private/SequencerRecorderSections/MovieSceneLiveLinkSectionRecorder.h
@@ -80,4 +80,7 @@ private:
 	/** Guid for getting data from Live Link*/
 	FGuid HandlerGuid;
 
+	/** Needed for rewinding, when we set the values we keep track of the last value set to restart the re-winding from that. */
+	TOptional<FVector> LastRotationValues;
+
 };
diff --git a/Engine/Plugins/Animation/LiveLink/Source/LiveLinkEditor/Private/SequencerRecorderSections/MovieSceneLiveLinkTrackRecorder.cpp b/Engine/Plugins/Animation/LiveLink/Source/LiveLinkEditor/Private/SequencerRecorderSections/MovieSceneLiveLinkTrackRecorder.cpp
--- a/Engine/Plugins/Animation/LiveLink/Source/LiveLinkEditor/Private/SequencerRecorderSections/MovieSceneLiveLinkTrackRecorder.cpp
+++ b/Engine/Plugins/Animation/LiveLink/Source/LiveLinkEditor/Private/SequencerRecorderSections/MovieSceneLiveLinkTrackRecorder.cpp
@@ -142,6 +142,7 @@ void UMovieSceneLiveLinkTrackRecorder::SetSectionStartTimecodeImpl(const FTimeco
 	{
 		MovieSceneSection->TimecodeSource = FMovieSceneTimecodeSource(InSectionStartTimecode);
 	}
+	LastRotationValues.Reset();
 }
 
 void UMovieSceneLiveLinkTrackRecorder::StopRecordingImpl()
@@ -267,7 +268,7 @@ void UMovieSceneLiveLinkTrackRecorder::RecordSampleImpl(const FQualifiedFrameTim
 				int32 ChannelIndex = 0;
 				for (FLiveLinkTransformKeys& TransformKeys : LinkTransformKeysArray)
 				{
-					TransformKeys.AppendToFloatChannelsAndReset(ChannelIndex, FloatChannels, Times);
+					TransformKeys.AppendToFloatChannelsAndReset(ChannelIndex, FloatChannels, Times, LastRotationValues);
 					ChannelIndex += 9;
 				}
 				for (FLiveLinkCurveKeys CurveKeys : LinkCurveKeysArray)
diff --git a/Engine/Plugins/Animation/LiveLink/Source/LiveLinkEditor/Public/MovieSceneLiveLinkTrackRecorder.h b/Engine/Plugins/Animation/LiveLink/Source/LiveLinkEditor/Public/MovieSceneLiveLinkTrackRecorder.h
--- a/Engine/Plugins/Animation/LiveLink/Source/LiveLinkEditor/Public/MovieSceneLiveLinkTrackRecorder.h
+++ b/Engine/Plugins/Animation/LiveLink/Source/LiveLinkEditor/Public/MovieSceneLiveLinkTrackRecorder.h
@@ -87,4 +87,7 @@ private:
 	/** Cached Key Reduction from Live Link Source Propreties*/
 	bool bReduceKeys;
 
+	/** Needed for rewinding, when we set the values we keep track of the last value set to restart the re-winding from that. */
+	TOptional<FVector> LastRotationValues;
+
 };
