Live Link Euler Flip fix.

Removed cache and make sure values are within -180,180 when fixing winding.

Also includes Fury fixes that add a CVar to force it to always use timecode if it's present and when converting timecode to sequencer time we handle the possible case that the live link frame rate is not equal to the timecode provider framerate.

diff --git a/Engine/Plugins/Animation/LiveLink/Source/LiveLink/Public/MovieScene/MovieSceneLiveLinkBufferData.h b/Engine/Plugins/Animation/LiveLink/Source/LiveLink/Public/MovieScene/MovieSceneLiveLinkBufferData.h
--- a/Engine/Plugins/Animation/LiveLink/Source/LiveLink/Public/MovieScene/MovieSceneLiveLinkBufferData.h
+++ b/Engine/Plugins/Animation/LiveLink/Source/LiveLink/Public/MovieScene/MovieSceneLiveLinkBufferData.h
@@ -134,10 +134,19 @@ struct FLiveLinkTransformKeys
 		FloatChannels[StartIndex++].AutoSetTangents();
 	}
 
+	//returns rotation as -180 to 180.
+	FVector GetNormalizedRotation(float X, float Y, float Z)
+	{
+		FQuat Quat(FRotator(Y, Z, X));
+		FVector Rot = Quat.Euler();
+		return Rot;
+	}
+
+
 	//This function is the one that's called when recording live link incrementally. We move the values over from our saved 
 	//Location, Rotation and Scale buffers into the specified float channels and then reset our buffers, re-using it's memory 
 	//for the next iteration. We also fix any euler flips during this process, avoiding iterating over the data once again during Finalize.
-	void AppendToFloatChannelsAndReset(int32 StartIndex, TArray<FMovieSceneFloatChannel>& FloatChannels, const TArray<FFrameNumber>& Times, TOptional<FVector>& LastRotationValues)
+	void AppendToFloatChannelsAndReset(int32 StartIndex, TArray<FMovieSceneFloatChannel>& FloatChannels, const TArray<FFrameNumber>& Times)
 	{
 		if (Times.Num() > 0)
 		{
@@ -148,24 +157,47 @@ struct FLiveLinkTransformKeys
 			FloatChannels[StartIndex++].AddKeys(Times, LocationZ);
 			LocationZ.Reset();
 
+
+			FVector LastRotationValue;
 			//fix euler flips
-			if (LastRotationValues.IsSet())
+			TArrayView<const FMovieSceneFloatValue> XRotChannel = FloatChannels[StartIndex].GetValues();
+			TArrayView<const FMovieSceneFloatValue> YRotChannel = FloatChannels[StartIndex].GetValues();
+			TArrayView<const FMovieSceneFloatValue> ZRotChannel = FloatChannels[StartIndex].GetValues();
+
+			if (XRotChannel.Num() > 0)
 			{
-				FVector Val = LastRotationValues.GetValue();
-				FMath::WindRelativeAnglesDegrees(Val.X, RotationX[0].Value);
-				FMath::WindRelativeAnglesDegrees(Val.Y, RotationY[0].Value);
-				FMath::WindRelativeAnglesDegrees(Val.Z, RotationZ[0].Value);
+
+				int32 LastOne = XRotChannel.Num() - 1;
+				FVector Rotation = GetNormalizedRotation(RotationX[0].Value, RotationY[0].Value, RotationZ[0].Value);
+				RotationX[0].Value = Rotation.X;
+				RotationY[0].Value = Rotation.Y;
+				RotationZ[0].Value = Rotation.Z;
+
+				//Due winding from last one saved with the new one normalized.
+
+				FMath::WindRelativeAnglesDegrees(XRotChannel[LastOne].Value, RotationX[0].Value);
+				FMath::WindRelativeAnglesDegrees(YRotChannel[LastOne].Value, RotationY[0].Value);
+				FMath::WindRelativeAnglesDegrees(ZRotChannel[LastOne].Value, RotationZ[0].Value);
+			}
+			else
+			{
+				FVector Rotation = GetNormalizedRotation(RotationX[0].Value, RotationY[0].Value, RotationZ[0].Value);
+				RotationX[0].Value = Rotation.X;
+				RotationY[0].Value = Rotation.Y;
+				RotationZ[0].Value = Rotation.Z;
 			}
 			int32 TotalCount = Times.Num();
-			
+
 			for (int32 RotIndex = 0; RotIndex < TotalCount - 1; RotIndex++)
 			{
+				FVector Rotation = GetNormalizedRotation(RotationX[RotIndex + 1].Value, RotationY[RotIndex + 1].Value, RotationZ[RotIndex + 1].Value);
+				RotationX[RotIndex + 1].Value = Rotation.X;
+				RotationY[RotIndex + 1].Value = Rotation.Y;
+				RotationZ[RotIndex + 1].Value = Rotation.Z;
 				FMath::WindRelativeAnglesDegrees(RotationX[RotIndex].Value, RotationX[RotIndex + 1].Value);
 				FMath::WindRelativeAnglesDegrees(RotationY[RotIndex].Value, RotationY[RotIndex + 1].Value);
 				FMath::WindRelativeAnglesDegrees(RotationZ[RotIndex].Value, RotationZ[RotIndex + 1].Value);
 			}
-			FVector Vec(RotationX[TotalCount - 1].Value, RotationY[TotalCount - 1].Value, RotationZ[TotalCount - 1].Value); 
-			LastRotationValues = Vec;
 
 			FloatChannels[StartIndex++].AddKeys(Times, RotationX);
 			RotationX.Reset();
@@ -199,4 +231,5 @@ struct FLiveLinkTransformKeys
 		FloatChannels[StartIndex++].AutoSetTangents();
 
 	}
+
 };
diff --git a/Engine/Plugins/Animation/LiveLink/Source/LiveLinkEditor/Private/SequencerRecorderSections/MovieSceneLiveLinkSectionRecorder.cpp b/Engine/Plugins/Animation/LiveLink/Source/LiveLinkEditor/Private/SequencerRecorderSections/MovieSceneLiveLinkSectionRecorder.cpp
--- a/Engine/Plugins/Animation/LiveLink/Source/LiveLinkEditor/Private/SequencerRecorderSections/MovieSceneLiveLinkSectionRecorder.cpp
+++ b/Engine/Plugins/Animation/LiveLink/Source/LiveLinkEditor/Private/SequencerRecorderSections/MovieSceneLiveLinkSectionRecorder.cpp
@@ -217,8 +217,6 @@ void FMovieSceneLiveLinkSectionRecorder::CreateTracks(UMovieScene* InMovieScene,
 		}
 		*/
 
-		LastRotationValues.Reset();
-
 	}
 }
 
@@ -326,7 +324,7 @@ void FMovieSceneLiveLinkSectionRecorder::Record(float CurrentTime)
 					int32 ChannelIndex = 0;
 					for (FLiveLinkTransformKeys& TransformKeys : LinkTransformKeysArray)
 					{
-						TransformKeys.AppendToFloatChannelsAndReset(ChannelIndex, FloatChannels, Times, LastRotationValues);
+						TransformKeys.AppendToFloatChannelsAndReset(ChannelIndex, FloatChannels, Times);
 						ChannelIndex += 9;
 					}
 					for (FLiveLinkCurveKeys CurveKeys : LinkCurveKeysArray)
diff --git a/Engine/Plugins/Animation/LiveLink/Source/LiveLinkEditor/Private/SequencerRecorderSections/MovieSceneLiveLinkSectionRecorder.h b/Engine/Plugins/Animation/LiveLink/Source/LiveLinkEditor/Private/SequencerRecorderSections/MovieSceneLiveLinkSectionRecorder.h
--- a/Engine/Plugins/Animation/LiveLink/Source/LiveLinkEditor/Private/SequencerRecorderSections/MovieSceneLiveLinkSectionRecorder.h
+++ b/Engine/Plugins/Animation/LiveLink/Source/LiveLinkEditor/Private/SequencerRecorderSections/MovieSceneLiveLinkSectionRecorder.h
@@ -80,7 +80,4 @@ private:
 	/** Guid for getting data from Live Link*/
 	FGuid HandlerGuid;
 
-	/** Needed for rewinding, when we set the values we keep track of the last value set to restart the re-winding from that. */
-	TOptional<FVector> LastRotationValues;
-
 };
diff --git a/Engine/Plugins/Animation/LiveLink/Source/LiveLinkEditor/Private/SequencerRecorderSections/MovieSceneLiveLinkTrackRecorder.cpp b/Engine/Plugins/Animation/LiveLink/Source/LiveLinkEditor/Private/SequencerRecorderSections/MovieSceneLiveLinkTrackRecorder.cpp
--- a/Engine/Plugins/Animation/LiveLink/Source/LiveLinkEditor/Private/SequencerRecorderSections/MovieSceneLiveLinkTrackRecorder.cpp
+++ b/Engine/Plugins/Animation/LiveLink/Source/LiveLinkEditor/Private/SequencerRecorderSections/MovieSceneLiveLinkTrackRecorder.cpp
@@ -21,6 +21,11 @@
 //recording
 #include "HAL/FileManagerGeneric.h"
 
+static TAutoConsoleVariable<int32> CVarSequencerAlwaysUseRecordLiveLinkTimecode(
+	TEXT("Sequencer.AlwayRecordLiveLinkTimecode"),
+	0, TEXT("If nonzero we use the LiveLink Timecode for time, even if Subject isn't Synchronized."),
+	ECVF_Default);
+
 DEFINE_LOG_CATEGORY(LiveLinkSerialization);
 
 void UMovieSceneLiveLinkTrackRecorder::CreateTrack(UMovieScene* InMovieScene, const FName& InSubjectName, UMovieSceneTrackRecorderSettings* InSettingsObject)
@@ -142,7 +147,6 @@ void UMovieSceneLiveLinkTrackRecorder::SetSectionStartTimecodeImpl(const FTimeco
 	{
 		MovieSceneSection->TimecodeSource = FMovieSceneTimecodeSource(InSectionStartTimecode);
 	}
-	LastRotationValues.Reset();
 }
 
 void UMovieSceneLiveLinkTrackRecorder::StopRecordingImpl()
@@ -193,7 +197,9 @@ void UMovieSceneLiveLinkTrackRecorder::RecordSampleImpl(const FQualifiedFrameTim
 		FFrameNumber CurrentFrame = CurrentTime.ConvertTo(TickResolution).FloorToFrame();
 		FFrameNumber FrameNumber;
 
-		bool bSynced = LiveLinkClient->IsSubjectTimeSynchronized(SubjectName);
+		const bool bAlwaysUseTimecode = CVarSequencerAlwaysUseRecordLiveLinkTimecode->GetInt() == 0 ? false : true;
+
+		bool bSyncedOrForced = bAlwaysUseTimecode || LiveLinkClient->IsSubjectTimeSynchronized(SubjectName);
 
 		MovieSceneSection->ExpandToFrame(MovieSceneSection->GetInclusiveStartFrame() + CurrentFrame);
 
@@ -213,14 +219,17 @@ void UMovieSceneLiveLinkTrackRecorder::RecordSampleImpl(const FQualifiedFrameTim
 			{
 				int32 TransformIndex = 0;
 				int32 CurveIndex = 0;
-				if (bSynced && GEngine && GEngine->GetTimecodeProvider())
+				if (bSyncedOrForced && GEngine && GEngine->GetTimecodeProvider())
 				{
-					FQualifiedFrameTime QualifiedFrameTime = Frame.MetaData.SceneTime;
+					//Get StartTime on Section in TimeCode FrameRate
+					//Convert that to LiveLink FrameRate and subtract out from LiveLink Frame to get section starting from zero.
+					//Finally convert that to the actual MovieScene Section FrameRate(TickResolution).
+					FQualifiedFrameTime LiveLinkFrameTime = Frame.MetaData.SceneTime;
 					const UTimecodeProvider* TimecodeProvider = GEngine->GetTimecodeProvider();
-					FFrameNumber FrameNumberStart = MovieSceneSection->TimecodeSource.Timecode.ToFrameNumber(TimecodeProvider->GetFrameRate());
-					QualifiedFrameTime.Time.FrameNumber -= FrameNumberStart;
-
-					FFrameTime FrameTime = QualifiedFrameTime.ConvertTo(TickResolution);
+					FQualifiedFrameTime TimeProviderStartFrameTime = FQualifiedFrameTime(MovieSceneSection->TimecodeSource.Timecode, TimecodeProvider->GetFrameRate());
+					FFrameNumber FrameNumberStart = TimeProviderStartFrameTime.ConvertTo(LiveLinkFrameTime.Rate).FrameNumber;
+					LiveLinkFrameTime.Time.FrameNumber -= FrameNumberStart;
+					FFrameTime FrameTime = LiveLinkFrameTime.ConvertTo(TickResolution);
 					FrameNumber = FrameTime.FrameNumber;
 				}
 				else
@@ -268,7 +277,7 @@ void UMovieSceneLiveLinkTrackRecorder::RecordSampleImpl(const FQualifiedFrameTim
 				int32 ChannelIndex = 0;
 				for (FLiveLinkTransformKeys& TransformKeys : LinkTransformKeysArray)
 				{
-					TransformKeys.AppendToFloatChannelsAndReset(ChannelIndex, FloatChannels, Times, LastRotationValues);
+					TransformKeys.AppendToFloatChannelsAndReset(ChannelIndex, FloatChannels, Times);
 					ChannelIndex += 9;
 				}
 				for (FLiveLinkCurveKeys CurveKeys : LinkCurveKeysArray)
diff --git a/Engine/Plugins/Animation/LiveLink/Source/LiveLinkEditor/Public/MovieSceneLiveLinkTrackRecorder.h b/Engine/Plugins/Animation/LiveLink/Source/LiveLinkEditor/Public/MovieSceneLiveLinkTrackRecorder.h
--- a/Engine/Plugins/Animation/LiveLink/Source/LiveLinkEditor/Public/MovieSceneLiveLinkTrackRecorder.h
+++ b/Engine/Plugins/Animation/LiveLink/Source/LiveLinkEditor/Public/MovieSceneLiveLinkTrackRecorder.h
@@ -87,7 +87,4 @@ private:
 	/** Cached Key Reduction from Live Link Source Propreties*/
 	bool bReduceKeys;
 
-	/** Needed for rewinding, when we set the values we keep track of the last value set to restart the re-winding from that. */
-	TOptional<FVector> LastRotationValues;
-
 };
