Changed default color of the invalid live link subject

diff --git a/Engine/Plugins/Animation/LiveLink/Source/LiveLink/Private/LiveLinkSettings.cpp b/Engine/Plugins/Animation/LiveLink/Source/LiveLink/Private/LiveLinkSettings.cpp
--- a/Engine/Plugins/Animation/LiveLink/Source/LiveLink/Private/LiveLinkSettings.cpp
+++ b/Engine/Plugins/Animation/LiveLink/Source/LiveLink/Private/LiveLinkSettings.cpp
@@ -17,7 +17,7 @@ ULiveLinkSettings::ULiveLinkSettings()
 	, MessageBusHeartbeatTimeout(2.0)
 	, TimeWithoutFrameToBeConsiderAsInvalid(0.5)
 	, ValidColor(FLinearColor::Green)
-	, InvalidColor(FLinearColor::Green)
+	, InvalidColor(FLinearColor::Yellow)
 {
 	// Add the role animation default settings
 	FLiveLinkRoleProjectSetting BaseAnimationSetting;
