- Potential fix for a one off hang when shutting down LiveLinkMessagebus

diff --git a/Engine/Plugins/Animation/LiveLink/Source/LiveLink/Private/LiveLinkMessageBusDiscoveryManager.cpp b/Engine/Plugins/Animation/LiveLink/Source/LiveLink/Private/LiveLinkMessageBusDiscoveryManager.cpp
--- a/Engine/Plugins/Animation/LiveLink/Source/LiveLink/Private/LiveLinkMessageBusDiscoveryManager.cpp
+++ b/Engine/Plugins/Animation/LiveLink/Source/LiveLink/Private/LiveLinkMessageBusDiscoveryManager.cpp
@@ -31,7 +31,10 @@ FLiveLinkMessageBusDiscoveryManager::~FLiveLinkMessageBusDiscoveryManager()
 {
 	{
 		FScopeLock Lock(&SourcesCriticalSection);
-		FMessageEndpoint::SafeRelease(MessageEndpoint);
+		
+		// Disable the Endpoint message handling since the message could keep it alive a bit.
+		MessageEndpoint->Disable();
+		MessageEndpoint.Reset();
 	}
 
 	Stop();
diff --git a/Engine/Plugins/Animation/LiveLink/Source/LiveLink/Private/LiveLinkMessageBusSource.cpp b/Engine/Plugins/Animation/LiveLink/Source/LiveLink/Private/LiveLinkMessageBusSource.cpp
--- a/Engine/Plugins/Animation/LiveLink/Source/LiveLink/Private/LiveLinkMessageBusSource.cpp
+++ b/Engine/Plugins/Animation/LiveLink/Source/LiveLink/Private/LiveLinkMessageBusSource.cpp
@@ -246,7 +246,10 @@ bool FLiveLinkMessageBusSource::RequestSourceShutdown()
 
 	FLiveLinkHeartbeatEmitter& HeartbeatEmitter = ILiveLinkModule::Get().GetHeartbeatEmitter();
 	HeartbeatEmitter.StopHeartbeat(ConnectionAddress, MessageEndpoint);
-	FMessageEndpoint::SafeRelease(MessageEndpoint);
+
+	// Disable the Endpoint message handling since the message could keep it alive a bit.
+	MessageEndpoint->Disable();
+	MessageEndpoint.Reset();
 
 	return true;
 }
diff --git a/Engine/Source/Runtime/LiveLinkMessageBusFramework/Private/LiveLinkProvider.cpp b/Engine/Source/Runtime/LiveLinkMessageBusFramework/Private/LiveLinkProvider.cpp
--- a/Engine/Source/Runtime/LiveLinkMessageBusFramework/Private/LiveLinkProvider.cpp
+++ b/Engine/Source/Runtime/LiveLinkMessageBusFramework/Private/LiveLinkProvider.cpp
@@ -283,7 +283,9 @@ public:
 	{
 		if (MessageEndpoint.IsValid())
 		{
-			FMessageEndpoint::SafeRelease(MessageEndpoint);
+			// Disable the Endpoint message handling since the message could keep it alive a bit.
+			MessageEndpoint->Disable();
+			MessageEndpoint.Reset();
 		}
 	}
 
