- Fixing LiveLink Controller not being applied on other machines of a MU session
- Adding livelink controller out of concert filtering list as an example in concert config and by default in QAGame

diff --git a/Engine/Plugins/Animation/LiveLink/Source/LiveLinkComponents/Private/LiveLinkComponentController.cpp b/Engine/Plugins/Animation/LiveLink/Source/LiveLinkComponents/Private/LiveLinkComponentController.cpp
--- a/Engine/Plugins/Animation/LiveLink/Source/LiveLinkComponents/Private/LiveLinkComponentController.cpp
+++ b/Engine/Plugins/Animation/LiveLink/Source/LiveLinkComponents/Private/LiveLinkComponentController.cpp
@@ -54,58 +54,50 @@ void ULiveLinkComponentController::TickComponent(float DeltaTime, ELevelTick Tic
 
 
 #if WITH_EDITOR
-void ULiveLinkComponentController::PostEditChangeChainProperty(struct FPropertyChangedChainEvent& PropertyChangedEvent)
+
+void ULiveLinkComponentController::PostEditChangeProperty(struct FPropertyChangedEvent& PropertyChangedEvent)
 {
-	if (PropertyChangedEvent.PropertyChain.Num() > 0)
+	//Detect changes to the SubjectRepresentation blindly when one of our property has changed. In MultiUser, this will be called with an empty property event
+	bool bCreateAnewController = false;
+	if (SubjectRepresentation.Role.Get() == nullptr)
+	{
+		Controller = nullptr;
+	}
+	else if (Controller)
 	{
-		UProperty* MemberProperty = PropertyChangedEvent.PropertyChain.GetActiveMemberNode()->GetValue();
-		if (MemberProperty != NULL)
+		if (!Controller->IsRoleSupported(SubjectRepresentation.Role))
 		{
-			if (MemberProperty->GetFName() == GET_MEMBER_NAME_CHECKED(ULiveLinkComponentController, SubjectRepresentation))
-			{
-				bool bCreateAnewController = false;
-				if (SubjectRepresentation.Role.Get() == nullptr)
-				{
-					Controller = nullptr;
-				}
-				else if (Controller)
-				{
-					if (!Controller->IsRoleSupported(SubjectRepresentation.Role))
-					{
-						Controller = nullptr;
-						bCreateAnewController = true;
-					}
-				}
-				else
-				{
-					bCreateAnewController = true;
-				}
-
-				if (bCreateAnewController)
-				{
-					TSubclassOf<ULiveLinkControllerBase> NewControllerClass = ULiveLinkControllerBase::GetControllerForRole(SubjectRepresentation.Role);
-					if (NewControllerClass.Get())
-					{
-						const EObjectFlags ControllerObjectFlags = GetMaskedFlags(RF_Public | RF_Transactional | RF_ArchetypeObject);
-						Controller = NewObject<ULiveLinkControllerBase>(this, NewControllerClass.Get(), NAME_None, ControllerObjectFlags);
-						Controller->InitializeInEditor();
-					}
-					else
-					{
-						UE_LOG(LogLiveLinkComponents, Warning, TEXT("No controller was found for role '%s'."), *SubjectRepresentation.Role->GetName());
-						FNotificationInfo NotificationInfo(LOCTEXT("NoFoundController", "No controller was found for the role."));
-						NotificationInfo.ExpireDuration = 2.0f;
-						FSlateNotificationManager::Get().AddNotification(NotificationInfo);
-					}
-				}
-			}
+			Controller = nullptr;
+			bCreateAnewController = true;
 		}
 	}
+	else
+	{
+		bCreateAnewController = true;
+	}
 
+	if (bCreateAnewController)
+	{
+		TSubclassOf<ULiveLinkControllerBase> NewControllerClass = ULiveLinkControllerBase::GetControllerForRole(SubjectRepresentation.Role);
+		if (NewControllerClass.Get())
+		{
+			const EObjectFlags ControllerObjectFlags = GetMaskedFlags(RF_Public | RF_Transactional | RF_ArchetypeObject);
+			Controller = NewObject<ULiveLinkControllerBase>(this, NewControllerClass.Get(), NAME_None, ControllerObjectFlags);
+			Controller->InitializeInEditor();
+		}
+		else
+		{
+			UE_LOG(LogLiveLinkComponents, Warning, TEXT("No controller was found for role '%s'."), *SubjectRepresentation.Role->GetName());
+			FNotificationInfo NotificationInfo(LOCTEXT("NoFoundController", "No controller was found for the role."));
+			NotificationInfo.ExpireDuration = 2.0f;
+			FSlateNotificationManager::Get().AddNotification(NotificationInfo);
+		}
+	}
+	
 	bTickInEditor = bUpdateInEditor;
 	bIsDirty = true;
 
-	Super::PostEditChangeChainProperty(PropertyChangedEvent);
+	Super::PostEditChangeProperty(PropertyChangedEvent);
 }
 #endif //WITH_EDITOR
 
diff --git a/Engine/Plugins/Animation/LiveLink/Source/LiveLinkComponents/Public/LiveLinkComponentController.h b/Engine/Plugins/Animation/LiveLink/Source/LiveLinkComponents/Public/LiveLinkComponentController.h
--- a/Engine/Plugins/Animation/LiveLink/Source/LiveLinkComponents/Public/LiveLinkComponentController.h
+++ b/Engine/Plugins/Animation/LiveLink/Source/LiveLinkComponents/Public/LiveLinkComponentController.h
@@ -41,7 +41,7 @@ public:
 public:
 	//~ UObject interface
 #if WITH_EDITOR
-	virtual void PostEditChangeChainProperty(struct FPropertyChangedChainEvent& PropertyChangedEvent) override;
+	virtual void PostEditChangeProperty(struct FPropertyChangedEvent& PropertyChangedEvent) override;
 #endif //WITH_EDITOR
 	//~ End UObject interface
 
diff --git a/Engine/Plugins/Developer/Concert/ConcertSync/ConcertSyncCore/Config/Engine.ini b/Engine/Plugins/Developer/Concert/ConcertSync/ConcertSyncCore/Config/Engine.ini
--- a/Engine/Plugins/Developer/Concert/ConcertSync/ConcertSyncCore/Config/Engine.ini
+++ b/Engine/Plugins/Developer/Concert/ConcertSync/ConcertSyncCore/Config/Engine.ini
@@ -9,3 +9,6 @@
 
 ; If you wish to use Composure in Multi-user, you may add that line under [/Script/ConcertSyncCore.ConcertSyncConfig] in your project engine config settings.
 ;+IncludeObjectClassFilters=(ObjectOuterClass=None,ObjectClasses=(Class'/Script/Composure.CompositingElementPass',Class'/Script/Composure.ComposurePostProcessPassPolicy'))
+
+; If you wish to use LiveLink in Multi-user, you may add that line under [/Script/ConcertSyncCore.ConcertSyncConfig] in your project engine config settings.
+;+IncludeObjectClassFilters=(ObjectOuterClass=Class'/Script/LiveLinkComponents.LiveLinkComponentController',ObjectClasses=(Class'/Script/LiveLinkComponents.LiveLinkControllerBase'))
