- Adding missing debugging data in livelink panel
- Fix deprecation path not removing its subject
Integration of CL#7412589 and #7457403

diff --git a/Engine/Plugins/Animation/LiveLink/Source/LiveLink/Private/LiveLinkClient.cpp b/Engine/Plugins/Animation/LiveLink/Source/LiveLink/Private/LiveLinkClient.cpp
--- a/Engine/Plugins/Animation/LiveLink/Source/LiveLink/Private/LiveLinkClient.cpp
+++ b/Engine/Plugins/Animation/LiveLink/Source/LiveLink/Private/LiveLinkClient.cpp
@@ -348,6 +348,10 @@ void FLiveLinkClient::PushSubjectStaticData_Internal(FPendingSubjectStatic&& Sub
 			Collection->RemoveSubject(SubjectStaticData.SubjectKey);
 			LiveLinkSubject = nullptr;
 		}
+		else
+		{
+			LiveLinkSubject->ClearFrames();
+		}
 	}
 	
 	if(LiveLinkSubject == nullptr)
@@ -1062,7 +1066,7 @@ FLiveLinkSkeletonStaticData* FLiveLinkClient::GetSubjectAnimationStaticData_Depr
 	{
 		if (FLiveLinkCollectionSubjectItem* SubjectItem = Collection->FindSubject(InSubjectKey))
 		{
-			if (SubjectItem->GetSubject()->GetRole() == ULiveLinkAnimationRole::StaticClass())
+			if (SubjectItem->GetSubject()->GetRole() == ULiveLinkAnimationRole::StaticClass() && !SubjectItem->bPendingKill)
 			{
 				return SubjectItem->GetSubject()->GetStaticData().Cast<FLiveLinkSkeletonStaticData>();
 			}
@@ -1097,6 +1101,7 @@ void FLiveLinkClient_Base_DEPRECATED::PushSubjectSkeleton(FGuid SourceGuid, FNam
 
 	FLiveLinkSubjectKey Key = FLiveLinkSubjectKey(SourceGuid, SubjectName);
 
+	RemoveSubject_AnyThread(Key);
 
 	FLiveLinkStaticDataStruct StaticData(FLiveLinkSkeletonStaticData::StaticStruct());
 	FLiveLinkSkeletonStaticData* SkeletonData = StaticData.Cast<FLiveLinkSkeletonStaticData>();
diff --git a/Engine/Plugins/Animation/LiveLink/Source/LiveLinkEditor/LiveLinkEditor.Build.cs b/Engine/Plugins/Animation/LiveLink/Source/LiveLinkEditor/LiveLinkEditor.Build.cs
--- a/Engine/Plugins/Animation/LiveLink/Source/LiveLinkEditor/LiveLinkEditor.Build.cs
+++ b/Engine/Plugins/Animation/LiveLink/Source/LiveLinkEditor/LiveLinkEditor.Build.cs
@@ -39,6 +39,7 @@ namespace UnrealBuildTool.Rules
 					"Settings",
 					"Sequencer",
 					"Slate",
+					"TimeManagement",
 					"UnrealEd",
 					"WorkspaceMenuStructure",
 				}
diff --git a/Engine/Plugins/Animation/LiveLink/Source/LiveLinkEditor/Private/LiveLinkClientPanel.cpp b/Engine/Plugins/Animation/LiveLink/Source/LiveLinkEditor/Private/LiveLinkClientPanel.cpp
--- a/Engine/Plugins/Animation/LiveLink/Source/LiveLinkEditor/Private/LiveLinkClientPanel.cpp
+++ b/Engine/Plugins/Animation/LiveLink/Source/LiveLinkEditor/Private/LiveLinkClientPanel.cpp
@@ -8,6 +8,7 @@
 #include "Widgets/SBoxPanel.h"
 #include "Widgets/SOverlay.h"
 #include "Widgets/Layout/SSplitter.h"
+#include "Widgets/Layout/SWidgetSwitcher.h"
 #include "Widgets/Images/SImage.h"
 #include "Widgets/Input/SButton.h"
 #include "Widgets/Input/SCheckBox.h"
@@ -22,6 +23,7 @@
 #include "LiveLinkSourceFactory.h"
 #include "LiveLinkSourceSettings.h"
 #include "LiveLinkSubjectSettings.h"
+#include "SLiveLinkDataView.h"
 
 #include "Editor.h"
 #include "Editor/EditorPerformanceSettings.h"
@@ -393,6 +395,7 @@ void SLiveLinkClientPanel::Construct(const FArguments& Args, FLiveLinkClient* In
 	Client = InClient;
 
 	bSelectionChangedGuard = false;
+	DetailWidgetIndex = 0;
 
 	OnSourcesChangedHandle = Client->OnLiveLinkSourcesChanged().AddSP(this, &SLiveLinkClientPanel::OnSourcesChangedHandler);
 	OnSubjectsChangedHandle = Client->OnLiveLinkSubjectsChanged().AddSP(this, &SLiveLinkClientPanel::OnSubjectsChangedHandler);
@@ -511,7 +514,18 @@ void SLiveLinkClientPanel::Construct(const FArguments& Args, FLiveLinkClient* In
 			+SSplitter::Slot()
 			.Value(0.5f)
 			[
-				SettingsDetailsView.ToSharedRef()
+				SNew(SWidgetSwitcher)
+				.WidgetIndex(this, &SLiveLinkClientPanel::GetDetailWidgetIndex)
+				+ SWidgetSwitcher::Slot()
+				[
+					//[0] Detail view for Source
+					SettingsDetailsView.ToSharedRef()
+				]
+				+ SWidgetSwitcher::Slot()
+				[
+					// [1] Detail view for Subject, Frame data & Static data
+					SAssignNew(DataDetailsView, SLiveLinkDataView, Client)
+				]
 			]
 		]
 		+SVerticalBox::Slot()
@@ -586,6 +600,11 @@ void SLiveLinkClientPanel::RefreshSourceData(bool bRefreshUI)
 	}
 }
 
+int32 SLiveLinkClientPanel::GetDetailWidgetIndex() const
+{
+	return DataDetailsView->GetSubjectKey().Source.IsValid() && !DataDetailsView->GetSubjectKey().SubjectName.IsNone() ? 1 : 0;
+}
+
 TSharedRef<ITableRow> SLiveLinkClientPanel::MakeSourceListViewWidget(FLiveLinkSourceUIEntryPtr Entry, const TSharedRef<STableViewBase>& OwnerTable) const
 {
 	return SNew(SLiveLinkClientPanelSourcesRow, OwnerTable)
@@ -600,6 +619,8 @@ void SLiveLinkClientPanel::OnSourceListSelectionChanged(FLiveLinkSourceUIEntryPt
 	}
 	TGuardValue<bool> ReentrantGuard(bSelectionChangedGuard, true);
 
+	DataDetailsView->SetSubjectKey(FLiveLinkSubjectKey());
+
 	int32 FoundSubjectEntryIndex = INDEX_NONE;
 	if(Entry.IsValid())
 	{
@@ -660,7 +681,7 @@ void SLiveLinkClientPanel::OnSubjectTreeSelectionChanged(FLiveLinkSubjectUIEntry
 		return;
 	}
 	TGuardValue<bool> ReentrantGuard(bSelectionChangedGuard, true);
-	
+
 	int32 FoundSourceIndex = INDEX_NONE;
 	bool bDetailViewSet = false;
 	if (SubjectEntry.IsValid())
@@ -669,13 +690,23 @@ void SLiveLinkClientPanel::OnSubjectTreeSelectionChanged(FLiveLinkSubjectUIEntry
 		FGuid SourceGuid = SubjectEntry->SubjectKey.Source;
 		FoundSourceIndex = SourceData.IndexOfByPredicate([SourceGuid](FLiveLinkSourceUIEntryPtr SourceEntry) { return SourceEntry->GetGuid() == SourceGuid; });
 
-		SettingsDetailsView->SetObject(SubjectEntry->GetSettings());
+		if (SubjectEntry->IsSource())
+		{
+			SettingsDetailsView->SetObject(SubjectEntry->GetSettings());
+			DataDetailsView->SetSubjectKey(FLiveLinkSubjectKey());
+		}
+		else
+		{
+			SettingsDetailsView->SetObject(nullptr);
+			DataDetailsView->SetSubjectKey(SubjectEntry->SubjectKey);
+		}
 		bDetailViewSet = true;
 	}
 
 	if (!bDetailViewSet)
 	{
 		SettingsDetailsView->SetObject(nullptr);
+		DataDetailsView->SetSubjectKey(FLiveLinkSubjectKey());
 	}
 
 	// Select the corresponding Source entry
diff --git a/Engine/Plugins/Animation/LiveLink/Source/LiveLinkEditor/Private/LiveLinkClientPanel.h b/Engine/Plugins/Animation/LiveLink/Source/LiveLinkEditor/Private/LiveLinkClientPanel.h
--- a/Engine/Plugins/Animation/LiveLink/Source/LiveLinkEditor/Private/LiveLinkClientPanel.h
+++ b/Engine/Plugins/Animation/LiveLink/Source/LiveLinkEditor/Private/LiveLinkClientPanel.h
@@ -16,6 +16,7 @@ struct FLiveLinkSourceUIEntry;
 struct FLiveLinkSubjectUIEntry;
 class FMenuBuilder;
 class FUICommandList;
+class SLiveLinkDataView;
 class ULiveLinkSourceFactory;
 class SLiveLinkSourceListView;
 class STableViewBase;
@@ -81,6 +82,8 @@ private:
 	FReply DisableEditorPerformanceThrottling();
 
 private:
+	int32 GetDetailWidgetIndex() const;
+
 	TSharedRef<ITableRow> MakeSourceListViewWidget(FLiveLinkSourceUIEntryPtr Entry, const TSharedRef<STableViewBase>& OwnerTable) const;
 	TSharedPtr<SWidget> OnSourceConstructContextMenu();
 
@@ -114,6 +117,9 @@ private:
 	// Reference to connection settings struct details panel
 	TSharedPtr<class IDetailsView> SettingsDetailsView;
 
+	// Reference to the data value struct details panel
+	TSharedPtr<class SLiveLinkDataView> DataDetailsView;
+
 	// Handle to delegate when client sources list has changed */
 	FDelegateHandle OnSourcesChangedHandle;
 
@@ -123,6 +129,9 @@ private:
 	// Map to cover 
 	TMap<UClass*, UObject*> DetailsPanelEditorObjects;
 
+	// Details index
+	int32 DetailWidgetIndex;
+
 	// Guard from reentrant selection
 	mutable bool bSelectionChangedGuard;
 };
diff --git a/Engine/Plugins/Animation/LiveLink/Source/LiveLinkEditor/Private/LiveLinkMetaDataDetailCustomization.cpp b/Engine/Plugins/Animation/LiveLink/Source/LiveLinkEditor/Private/LiveLinkMetaDataDetailCustomization.cpp
new file mode 100644
--- /dev/null
+++ b/Engine/Plugins/Animation/LiveLink/Source/LiveLinkEditor/Private/LiveLinkMetaDataDetailCustomization.cpp
@@ -0,0 +1,167 @@
+// Copyright 1998-2019 Epic Games, Inc. All Rights Reserved.
+
+#include "LiveLinkMetaDataDetailCustomization.h"
+
+#include "LiveLinkTypes.h"
+
+#include "CommonFrameRates.h"
+#include "DetailLayoutBuilder.h"
+#include "DetailWidgetRow.h"
+#include "Editor.h"
+#include "IDetailChildrenBuilder.h"
+#include "IDetailGroup.h"
+#include "IPropertyUtilities.h"
+#include "IPropertyTypeCustomization.h"
+#include "Misc/QualifiedFrameTime.h"
+#include "PropertyHandle.h"
+
+#include "Widgets/Text/STextBlock.h"
+
+
+#define LOCTEXT_NAMESPACE "LiveLinkMetaDataDetailCustomization"
+
+
+void FLiveLinkMetaDataDetailCustomization::CustomizeChildren(TSharedRef<IPropertyHandle> StructPropertyHandle, IDetailChildrenBuilder& StructBuilder, IPropertyTypeCustomizationUtils& StructCustomizationUtils)
+{
+	uint32 NumChildren;
+	StructPropertyHandle->GetNumChildren(NumChildren);
+
+	for (uint32 ChildIndex = 0; ChildIndex < NumChildren; ++ChildIndex)
+	{
+		TSharedRef<IPropertyHandle> ChildHandle = StructPropertyHandle->GetChildHandle(ChildIndex).ToSharedRef();
+		const FName PropertyName = ChildHandle->GetProperty()->GetFName();
+
+		if (PropertyName == GET_MEMBER_NAME_CHECKED(FLiveLinkMetaData, SceneTime)
+			&& Cast<UStructProperty>(ChildHandle->GetProperty()))
+		{
+			SceneTimeHandle = ChildHandle;
+			check(CastChecked<UStructProperty>(SceneTimeHandle->GetProperty())->Struct->GetName() == TEXT("QualifiedFrameTime"));
+
+			IDetailGroup& SceneTimeGroup = StructBuilder.AddGroup("SceneTime", LOCTEXT("SceneTimeLabel", "Scene Time"));
+
+			SceneTimeGroup.AddWidgetRow()
+				.NameContent()
+				[
+					SNew(SBox)
+					.Content()
+					[
+						SNew(STextBlock)
+						.Text(LOCTEXT("FrameNumber", "Frame Number"))
+						.Font(IDetailLayoutBuilder::GetDetailFont())
+					]
+				]
+				.ValueContent()
+				.MinDesiredWidth(166.0f)
+				[
+					SNew(STextBlock)
+					.Text(this, &FLiveLinkMetaDataDetailCustomization::GetFrameNumber)
+					.Font(IDetailLayoutBuilder::GetDetailFont())
+				];
+
+			SceneTimeGroup.AddWidgetRow()
+				.NameContent()
+				[
+					SNew(SBox)
+					.Content()
+					[
+						SNew(STextBlock)
+						.Text(LOCTEXT("FrameRate", "Timecode Frame Rate"))
+						.Font(IDetailLayoutBuilder::GetDetailFont())
+					]
+				]
+				.ValueContent()
+				.MinDesiredWidth(166.0f)
+				[
+					SNew(STextBlock)
+					.Text(this, &FLiveLinkMetaDataDetailCustomization::GetTimecodeFrameRate)
+					.Font(IDetailLayoutBuilder::GetDetailFont())
+				];
+
+			SceneTimeGroup.AddWidgetRow()
+				.NameContent()
+				[
+					SNew(SBox)
+					.Content()
+					[
+						SNew(STextBlock)
+						.Text(LOCTEXT("Timecode", "Timecode"))
+						.Font(IDetailLayoutBuilder::GetDetailFont())
+					]
+				]
+				.ValueContent()
+				.MinDesiredWidth(166.0f)
+				[
+					SNew(STextBlock)
+					.Text(this, &FLiveLinkMetaDataDetailCustomization::GetTimecodeValue)
+					.Font(IDetailLayoutBuilder::GetDetailFont())
+				];
+		}
+		else
+		{
+			StructBuilder.AddProperty(ChildHandle).ShowPropertyButtons(false);
+		}
+	}
+}
+
+FText FLiveLinkMetaDataDetailCustomization::GetFrameNumber() const
+{
+	void* Data;
+	FPropertyAccess::Result Result = SceneTimeHandle->GetValueData(Data);
+	if (Result == FPropertyAccess::MultipleValues)
+	{
+		return LOCTEXT("MultipleValues", "Multiple Values");
+	}
+
+	if (Result == FPropertyAccess::Success)
+	{
+		FQualifiedFrameTime* QualifiedFrameTimePtr = reinterpret_cast<FQualifiedFrameTime*>(Data);
+		return FText::AsNumber(QualifiedFrameTimePtr->Time.FrameNumber.Value);
+	}
+
+	return LOCTEXT("Failed", "Failed");
+}
+
+FText FLiveLinkMetaDataDetailCustomization::GetTimecodeFrameRate() const
+{
+	void* Data;
+	FPropertyAccess::Result Result = SceneTimeHandle->GetValueData(Data);
+	if (Result == FPropertyAccess::MultipleValues)
+	{
+		return LOCTEXT("MultipleValues", "Multiple Values");
+	}
+
+	if (Result == FPropertyAccess::Success)
+	{
+		FQualifiedFrameTime* QualifiedFrameTimePtr = reinterpret_cast<FQualifiedFrameTime*>(Data);
+		if (const FCommonFrameRateInfo* Found = FCommonFrameRates::Find(QualifiedFrameTimePtr->Rate))
+		{
+			return Found->DisplayName;
+		}
+		else
+		{
+			QualifiedFrameTimePtr->Rate.ToPrettyText();
+		}
+	}
+
+	return LOCTEXT("Failed", "Failed");
+}
+
+FText FLiveLinkMetaDataDetailCustomization::GetTimecodeValue() const
+{
+	void* Data;
+	FPropertyAccess::Result Result = SceneTimeHandle->GetValueData(Data);
+	if (Result == FPropertyAccess::MultipleValues)
+	{
+		return LOCTEXT("MultipleValues", "Multiple Values");
+	}
+
+	if (Result == FPropertyAccess::Success)
+	{
+		FQualifiedFrameTime* QualifiedFrameTimePtr = reinterpret_cast<FQualifiedFrameTime*>(Data);
+		return FText::FromString(FTimecode::FromFrameNumber(QualifiedFrameTimePtr->Time.FrameNumber, QualifiedFrameTimePtr->Rate, false).ToString());
+	}
+
+	return LOCTEXT("Failed", "Failed");
+}
+
+#undef LOCTEXT_NAMESPACE
diff --git a/Engine/Plugins/Animation/LiveLink/Source/LiveLinkEditor/Private/LiveLinkMetaDataDetailCustomization.h b/Engine/Plugins/Animation/LiveLink/Source/LiveLinkEditor/Private/LiveLinkMetaDataDetailCustomization.h
new file mode 100644
--- /dev/null
+++ b/Engine/Plugins/Animation/LiveLink/Source/LiveLinkEditor/Private/LiveLinkMetaDataDetailCustomization.h
@@ -0,0 +1,23 @@
+// Copyright 1998-2019 Epic Games, Inc. All Rights Reserved.
+
+#pragma once
+
+#include "CoreMinimal.h"
+
+#include "IPropertyTypeCustomization.h"
+
+class IPropertyHandle;
+
+class FLiveLinkMetaDataDetailCustomization : public IPropertyTypeCustomization
+{
+public:
+	// IDetailCustomization interface
+	virtual void CustomizeHeader( TSharedRef<IPropertyHandle> PropertyHandle, FDetailWidgetRow& HeaderRow, IPropertyTypeCustomizationUtils& CustomizationUtils ) override {}
+	virtual void CustomizeChildren(TSharedRef<IPropertyHandle> PropertyHandle, IDetailChildrenBuilder& ChildBuilder, IPropertyTypeCustomizationUtils& CustomizationUtils) override;
+
+private:
+	FText GetFrameNumber() const;
+	FText GetTimecodeValue() const;
+	FText GetTimecodeFrameRate() const;
+	TSharedPtr<IPropertyHandle> SceneTimeHandle;
+};
diff --git a/Engine/Plugins/Animation/LiveLink/Source/LiveLinkEditor/Private/SLiveLinkDataView.cpp b/Engine/Plugins/Animation/LiveLink/Source/LiveLinkEditor/Private/SLiveLinkDataView.cpp
new file mode 100644
--- /dev/null
+++ b/Engine/Plugins/Animation/LiveLink/Source/LiveLinkEditor/Private/SLiveLinkDataView.cpp
@@ -0,0 +1,227 @@
+// Copyright 1998-2019 Epic Games, Inc. All Rights Reserved.
+
+#include "SLiveLinkDataView.h"
+
+#include "LiveLinkMetaDataDetailCustomization.h"
+#include "LiveLinkClient.h"
+#include "LiveLinkTypes.h"
+#include "LiveLinkRole.h"
+
+#include "ClassViewerModule.h"
+#include "IStructureDetailsView.h"
+#include "Modules/ModuleManager.h"
+#include "PropertyEditorModule.h"
+
+#include "Framework/MultiBox/MultiBoxBuilder.h"
+#include "Widgets/Images/SImage.h"
+#include "Widgets/Input/SComboButton.h"
+#include "Widgets/Input/SSpinBox.h"
+#include "Widgets/Layout/SBox.h"
+#include "Widgets/Layout/SWidgetSwitcher.h"
+#include "Widgets/Text/STextBlock.h"
+#include "Widgets/SNullWidget.h"
+
+
+#define LOCTEXT_NAMESPACE "LiveLinkDataView"
+
+
+void SLiveLinkDataView::Construct(const FArguments& Args, FLiveLinkClient* InClient)
+{
+	Client = InClient;
+	LastUpdateSeconds = 0.0;
+	UpdateDelay = 1.0;
+	DetailType = EDetailType::Property;
+
+	FMenuBuilder DetailViewOptions(true, nullptr);
+	DetailViewOptions.BeginSection("Properties", LOCTEXT("Properties", "Properties"));
+	{
+		DetailViewOptions.AddMenuEntry(
+			LOCTEXT("ShowSubjectDetail", "Show Subject Properties"),
+			LOCTEXT("ShowSubjectDetail_ToolTip", "Displays the subject properties"),
+			FSlateIcon(),
+			FUIAction(
+				FExecuteAction::CreateSP(this, &SLiveLinkDataView::OnSelectDetailWidget, EDetailType::Property),
+				FCanExecuteAction(),
+				FIsActionChecked::CreateSP(this, &SLiveLinkDataView::IsSelectedDetailWidget, EDetailType::Property)
+			),
+			NAME_None,
+			EUserInterfaceActionType::RadioButton
+		);
+		DetailViewOptions.AddMenuEntry(
+			LOCTEXT("ShowStaticDataDetail", "Show Static Data"),
+			LOCTEXT("ShowStaticData_ToolTip", "Displays the subject static data"),
+			FSlateIcon(),
+			FUIAction(
+				FExecuteAction::CreateSP(this, &SLiveLinkDataView::OnSelectDetailWidget, EDetailType::StaticData),
+				FCanExecuteAction(),
+				FIsActionChecked::CreateSP(this, &SLiveLinkDataView::IsSelectedDetailWidget, EDetailType::StaticData)
+			),
+			NAME_None,
+			EUserInterfaceActionType::RadioButton
+		);
+		DetailViewOptions.AddMenuEntry(
+			LOCTEXT("ShowFrameDataDetail", "Show Frame Data"),
+			LOCTEXT("ShowFrameData_ToolTip", "Displays the subject frame data"),
+			FSlateIcon(),
+			FUIAction(
+				FExecuteAction::CreateSP(this, &SLiveLinkDataView::OnSelectDetailWidget, EDetailType::FrameData),
+				FCanExecuteAction(),
+				FIsActionChecked::CreateSP(this, &SLiveLinkDataView::IsSelectedDetailWidget, EDetailType::FrameData)
+			),
+			NAME_None,
+			EUserInterfaceActionType::RadioButton
+		);
+	}
+	DetailViewOptions.EndSection();
+
+	DetailViewOptions.BeginSection("Settings", LOCTEXT("Settings", "Settings"));
+	{
+		TSharedRef<SWidget> RefreshDelay = SNew(SSpinBox<double>)
+			.MinValue(0.0)
+			.MaxValue(10.0)
+			.MinSliderValue(0.0)
+			.MaxSliderValue(10.0)
+			.ToolTipText(LOCTEXT("RefreshDelay", "Refresh delay of Static & Frame data."))
+			.Value(this, &SLiveLinkDataView::GetRefreshDelay)
+			.OnValueCommitted(this, &SLiveLinkDataView::SetRefreshDelayInternal)
+			.IsEnabled(this, &SLiveLinkDataView::CanEditRefreshDelay);
+
+		DetailViewOptions.AddWidget(RefreshDelay, LOCTEXT("RefreshDelay", "Refresh Delay"));
+	}
+	DetailViewOptions.EndSection();
+
+	FDetailsViewArgs DetailsViewArgs;
+	DetailsViewArgs.bUpdatesFromSelection = false;
+	DetailsViewArgs.bLockable = false;
+	DetailsViewArgs.bShowPropertyMatrixButton = false;
+	DetailsViewArgs.NameAreaSettings = FDetailsViewArgs::HideNameArea;
+	DetailsViewArgs.ViewIdentifier = NAME_None;
+	DetailsViewArgs.bShowCustomFilterOption = false;
+	DetailsViewArgs.bShowOptions = false;
+
+	FPropertyEditorModule& PropertyEditorModule = FModuleManager::GetModuleChecked<FPropertyEditorModule>("PropertyEditor");
+
+	SettingsDetailsView = PropertyEditorModule.CreateDetailView(DetailsViewArgs);
+	SettingsDetailsView->OnFinishedChangingProperties().AddSP(this, &SLiveLinkDataView::OnPropertyChanged);
+
+	FStructureDetailsViewArgs StructViewArgs;
+	StructureDetailsView = PropertyEditorModule.CreateStructureDetailView(DetailsViewArgs, StructViewArgs, TSharedPtr<FStructOnScope>());
+	StructureDetailsView->GetDetailsView()->SetIsPropertyEditingEnabledDelegate(FIsPropertyEditingEnabled::CreateLambda([]() { return false; }));
+	StructureDetailsView->GetDetailsView()->RegisterInstancedCustomPropertyTypeLayout(FLiveLinkMetaData::StaticStruct()->GetFName(),
+		FOnGetPropertyTypeCustomizationInstance::CreateLambda([](){ return MakeShared<FLiveLinkMetaDataDetailCustomization>(); })
+		);
+
+	ChildSlot
+	[
+		SNew(SVerticalBox)
+		+ SVerticalBox::Slot()
+		.AutoHeight()
+		[
+			SNew(SHorizontalBox)
+			+ SHorizontalBox::Slot()
+			[
+				SNullWidget::NullWidget
+			]
+			+ SHorizontalBox::Slot()
+			.HAlign(HAlign_Right)
+			.AutoWidth()
+			[
+				SNew(SComboButton)
+				.ContentPadding(0)
+				.ForegroundColor(FSlateColor::UseForeground())
+				.ButtonStyle(FEditorStyle::Get(), "ToggleButton")
+				.MenuContent()
+				[
+					DetailViewOptions.MakeWidget()
+				]
+				.ButtonContent()
+				[
+					SNew(SHorizontalBox)
+					+ SHorizontalBox::Slot()
+					.AutoWidth()
+					.VAlign(VAlign_Center)
+					[
+						SNew(SImage)
+						.Image(FEditorStyle::GetBrush("GenericViewButton"))
+					]
+					+ SHorizontalBox::Slot()
+					.AutoWidth()
+					.Padding(2, 0, 0, 0)
+					.VAlign(VAlign_Center)
+					[
+						SNew(STextBlock)
+						.Text(LOCTEXT("ViewButton", "View Options"))
+					]
+				]
+			]
+		]
+		+ SVerticalBox::Slot()
+		[
+			SNew(SWidgetSwitcher)
+			.WidgetIndex(this, &SLiveLinkDataView::GetDetailWidgetIndex)
+			+ SWidgetSwitcher::Slot()
+			[
+				SettingsDetailsView.ToSharedRef()
+			]
+			+ SWidgetSwitcher::Slot()
+			[
+				StructureDetailsView->GetWidget().ToSharedRef()
+			]
+		]
+	];
+}
+
+void SLiveLinkDataView::Tick(const FGeometry& AllottedGeometry, const double InCurrentTime, const float InDeltaTime)
+{
+	Super::Tick(AllottedGeometry, InCurrentTime, InDeltaTime);
+
+	if (DetailType != EDetailType::Property && InCurrentTime - LastUpdateSeconds > UpdateDelay)
+	{
+		LastUpdateSeconds = InCurrentTime;
+
+		if (Client->IsSubjectEnabled(SubjectKey))
+		{
+			TSubclassOf<ULiveLinkRole> SubjectRole = Client->GetSubjectRole(SubjectKey.SubjectName);
+			FLiveLinkSubjectFrameData SubjectData;
+			if (Client->EvaluateFrame_AnyThread(SubjectKey.SubjectName, SubjectRole, SubjectData))
+			{
+				if (DetailType == EDetailType::StaticData)
+				{
+					TSharedPtr<FStructOnScope> StructOnScope = MakeShared<FStructOnScope>(SubjectData.StaticData.GetStruct());
+					CastChecked<UScriptStruct>(StructOnScope->GetStruct())->CopyScriptStruct(StructOnScope->GetStructMemory(), SubjectData.StaticData.GetBaseData());
+					StructureDetailsView->SetStructureData(StructOnScope);
+				}
+				else
+				{
+					TSharedPtr<FStructOnScope> StructOnScope = MakeShared<FStructOnScope>(SubjectData.FrameData.GetStruct());
+					CastChecked<UScriptStruct>(StructOnScope->GetStruct())->CopyScriptStruct(StructOnScope->GetStructMemory(), SubjectData.FrameData.GetBaseData());
+					StructureDetailsView->SetStructureData(StructOnScope);
+				}
+			}
+		}
+	}
+}
+
+void SLiveLinkDataView::SetSubjectKey(FLiveLinkSubjectKey InSubjectKey)
+{
+	SubjectKey = InSubjectKey;
+	SettingsDetailsView->SetObject(Client->GetSubjectSettings(SubjectKey));
+}
+
+int32 SLiveLinkDataView::GetDetailWidgetIndex() const
+{
+	return DetailType == EDetailType::Property ? 0 : 1;
+}
+
+void SLiveLinkDataView::OnSelectDetailWidget(EDetailType InDetailType)
+{
+	DetailType = InDetailType;
+}
+
+void SLiveLinkDataView::OnPropertyChanged(const FPropertyChangedEvent& InEvent)
+{
+	Client->OnPropertyChanged(SubjectKey.Source, InEvent);
+}
+
+
+#undef LOCTEXT_NAMESPACE
\ No newline at end of file
diff --git a/Engine/Plugins/Animation/LiveLink/Source/LiveLinkEditor/Private/SLiveLinkDataView.h b/Engine/Plugins/Animation/LiveLink/Source/LiveLinkEditor/Private/SLiveLinkDataView.h
new file mode 100644
--- /dev/null
+++ b/Engine/Plugins/Animation/LiveLink/Source/LiveLinkEditor/Private/SLiveLinkDataView.h
@@ -0,0 +1,58 @@
+// Copyright 1998-2019 Epic Games, Inc. All Rights Reserved.
+
+#pragma once
+
+#include "Widgets/SCompoundWidget.h"
+#include "Widgets/DeclarativeSyntaxSupport.h"
+
+#include "LiveLinkTypes.h"
+#include "Types/SlateEnums.h"
+
+class IDetailsView;
+class FLiveLinkClient;
+class IStructureDetailsView;
+
+class SLiveLinkDataView : public SCompoundWidget
+{
+private:
+	using Super = SCompoundWidget;
+
+public:
+	SLATE_BEGIN_ARGS(SLiveLinkDataView){}
+	SLATE_END_ARGS()
+
+	void Construct(const FArguments& Args, FLiveLinkClient* InClient);
+
+	virtual void Tick(const FGeometry& AllottedGeometry, const double InCurrentTime, const float InDeltaTime) override;
+
+	void SetSubjectKey(FLiveLinkSubjectKey InSubjectKey);
+	FLiveLinkSubjectKey GetSubjectKey() const { return SubjectKey; }
+
+	void SetRefreshDelay(double DelaySeconds) { UpdateDelay = DelaySeconds; }
+	double GetRefreshDelay() const { return UpdateDelay; }
+
+private:
+	enum class EDetailType : uint32
+	{
+		Property,
+		StaticData,
+		FrameData,
+	};
+
+	void OnPropertyChanged(const FPropertyChangedEvent& InEvent);
+	int32 GetDetailWidgetIndex() const;
+	void OnSelectDetailWidget(EDetailType InDetailType);
+	bool IsSelectedDetailWidget(EDetailType InDetailType) const { return InDetailType == DetailType; }
+	bool CanEditRefreshDelay() const { return DetailType != EDetailType::Property; }
+	void SetRefreshDelayInternal(double InDelaySeconds, ETextCommit::Type) { SetRefreshDelay(InDelaySeconds); }
+
+private:
+	FLiveLinkClient* Client;
+	FLiveLinkSubjectKey SubjectKey;
+	double LastUpdateSeconds;
+	double UpdateDelay;
+	EDetailType DetailType;
+
+	TSharedPtr<IStructureDetailsView> StructureDetailsView;
+	TSharedPtr<IDetailsView> SettingsDetailsView;
+};
