Fix for log spamming in Evaluate Live Link Frame Blueprint node

diff --git a/Engine/Plugins/Animation/LiveLink/Source/LiveLinkEditor/Private/SLiveLinkSubjectNameGraphPin.cpp b/Engine/Plugins/Animation/LiveLink/Source/LiveLinkEditor/Private/SLiveLinkSubjectNameGraphPin.cpp
--- a/Engine/Plugins/Animation/LiveLink/Source/LiveLinkEditor/Private/SLiveLinkSubjectNameGraphPin.cpp
+++ b/Engine/Plugins/Animation/LiveLink/Source/LiveLinkEditor/Private/SLiveLinkSubjectNameGraphPin.cpp
@@ -27,8 +27,11 @@ FLiveLinkSubjectRepresentation SLiveLinkSubjectNameGraphPin::GetValue() const
 {
 	FLiveLinkSubjectRepresentation SubjectRepresentation;
 
-	FLiveLinkSubjectName::StaticStruct()->ImportText(*GraphPinObj->GetDefaultAsString(), &SubjectRepresentation.Subject, nullptr, EPropertyPortFlags::PPF_None, GLog, FLiveLinkSubjectName::StaticStruct()->GetName());
-	SubjectRepresentation.Role = Cast<UClass>(GraphPinObj->DefaultObject);
+	if (!GraphPinObj->GetDefaultAsString().IsEmpty())
+	{
+		FLiveLinkSubjectName::StaticStruct()->ImportText(*GraphPinObj->GetDefaultAsString(), &SubjectRepresentation.Subject, nullptr, EPropertyPortFlags::PPF_None, GLog, FLiveLinkSubjectName::StaticStruct()->GetName());
+		SubjectRepresentation.Role = Cast<UClass>(GraphPinObj->DefaultObject);
+	}
 
 	return SubjectRepresentation;
 }
