Add EvaluateFrameFromSource_AnyThread to LiveLink to evaluate a specific subject. Used for debugging. When requesting if a source is enabled, you can ask for the snapshot (took during the start of the frame) or the current value.
Change LiveLink client tick order to tick at the start of the frame.
Request the live enabled in LiveLinkUI
Change TimecodeSync tick order to after the frame.
Change LiveLink buffering system.
Add unit to livelink setting properties.
Add interpolation with timecode to LiveLink

diff --git a/Engine/Plugins/Animation/LiveLink/Source/LiveLink/LiveLink.Build.cs b/Engine/Plugins/Animation/LiveLink/Source/LiveLink/LiveLink.Build.cs
--- a/Engine/Plugins/Animation/LiveLink/Source/LiveLink/LiveLink.Build.cs
+++ b/Engine/Plugins/Animation/LiveLink/Source/LiveLink/LiveLink.Build.cs
@@ -21,6 +21,7 @@ namespace UnrealBuildTool.Rules
 			{
 				"HeadMountedDisplay",
 				"InputCore",
+				"Media",
 				"Projects",
 				"SlateCore",
 				"Slate",
diff --git a/Engine/Plugins/Animation/LiveLink/Source/LiveLink/Private/LiveLinkBlueprintLibrary.cpp b/Engine/Plugins/Animation/LiveLink/Source/LiveLink/Private/LiveLinkBlueprintLibrary.cpp
--- a/Engine/Plugins/Animation/LiveLink/Source/LiveLink/Private/LiveLinkBlueprintLibrary.cpp
+++ b/Engine/Plugins/Animation/LiveLink/Source/LiveLink/Private/LiveLinkBlueprintLibrary.cpp
@@ -179,13 +179,13 @@ TArray<FLiveLinkSubjectKey> ULiveLinkBlueprintLibrary::GetLiveLinkSubjects(bool
 	return TArray<FLiveLinkSubjectKey>();
 }
 
-bool ULiveLinkBlueprintLibrary::IsSpecificLiveLinkSubjectEnabled(const FLiveLinkSubjectKey SubjectKey)
+bool ULiveLinkBlueprintLibrary::IsSpecificLiveLinkSubjectEnabled(const FLiveLinkSubjectKey SubjectKey, bool bUseSnapshot)
 {
 	IModularFeatures& ModularFeatures = IModularFeatures::Get();
 	if (ModularFeatures.IsModularFeatureAvailable(ILiveLinkClient::ModularFeatureName))
 	{
 		ILiveLinkClient& LiveLinkClient = ModularFeatures.GetModularFeature<ILiveLinkClient>(ILiveLinkClient::ModularFeatureName);
-		return LiveLinkClient.IsSubjectEnabled(SubjectKey);
+		return LiveLinkClient.IsSubjectEnabled(SubjectKey, bUseSnapshot);
 	}
 	return false;
 }
diff --git a/Engine/Plugins/Animation/LiveLink/Source/LiveLink/Private/LiveLinkClient.cpp b/Engine/Plugins/Animation/LiveLink/Source/LiveLink/Private/LiveLinkClient.cpp
--- a/Engine/Plugins/Animation/LiveLink/Source/LiveLink/Private/LiveLinkClient.cpp
+++ b/Engine/Plugins/Animation/LiveLink/Source/LiveLink/Private/LiveLinkClient.cpp
@@ -15,7 +15,9 @@
 #include "LiveLinkSourceFactory.h"
 #include "LiveLinkSourceSettings.h"
 #include "LiveLinkSubject.h"
+#include "IMediaModule.h"
 #include "Misc/CoreDelegates.h"
+#include "Modules/ModuleManager.h"
 #include "Roles/LiveLinkAnimationRole.h"
 #include "Roles/LiveLinkAnimationTypes.h"
 #include "Roles/LiveLinkBasicTypes.h"
@@ -42,6 +44,9 @@ FLiveLinkClient::FLiveLinkClient()
 {
 	Collection = MakeUnique<FLiveLinkSourceCollection>();
 	FCoreDelegates::OnPreExit.AddRaw(this, &FLiveLinkClient::Shutdown);
+
+	IMediaModule& MediaModule = FModuleManager::LoadModuleChecked<IMediaModule>("Media");
+	MediaModule.GetOnTickPreEngineCompleted().AddRaw(this, &FLiveLinkClient::Tick);
 }
 
 FLiveLinkClient::~FLiveLinkClient()
@@ -50,7 +55,7 @@ FLiveLinkClient::~FLiveLinkClient()
 	Shutdown();
 }
 
-void FLiveLinkClient::Tick(float DeltaTime)
+void FLiveLinkClient::Tick()
 {
 	SCOPE_CYCLE_COUNTER(STAT_LiveLink_Client_Tick);
 
@@ -88,6 +93,13 @@ void FLiveLinkClient::UpdateSources()
 {
 	for (FLiveLinkCollectionSourceItem& SourceItem : Collection->GetSources())
 	{
+#if WITH_EDITOR
+		if (SourceItem.Setting)
+		{
+			SourceItem.Setting->SourceDebugInfos.Reset();
+		}
+#endif
+
 		SourceItem.Source->Update();
 	}
 }
@@ -135,6 +147,11 @@ void FLiveLinkClient::BuildThisTicksSubjectSnapshot()
 
 void FLiveLinkClient::Shutdown()
 {
+	if(IMediaModule* MediaModule = FModuleManager::GetModulePtr<IMediaModule>("Media"))
+	{
+		MediaModule->GetOnTickPreEngineCompleted().RemoveAll(this);
+	}
+
 	if (Collection)
 	{
 		FScopeLock Lock(&CollectionAccessCriticalSection);
@@ -687,8 +704,16 @@ bool FLiveLinkClient::IsSubjectValid(FLiveLinkSubjectName InSubjectName) const
 	return false;
 }
 
-bool FLiveLinkClient::IsSubjectEnabled(const FLiveLinkSubjectKey& InSubjectKey) const
+bool FLiveLinkClient::IsSubjectEnabled(const FLiveLinkSubjectKey& InSubjectKey, bool bUseSnapshot) const
 {
+	if (bUseSnapshot)
+	{
+		if (const FLiveLinkSubjectKey* FoundSubjectKey = EnabledSubjects.Find(InSubjectKey.SubjectName))
+		{
+			return *FoundSubjectKey == InSubjectKey;
+		}
+		return false;
+	}
 	return Collection->IsSubjectEnabled(InSubjectKey);
 }
 
@@ -974,45 +999,6 @@ void FLiveLinkClient::UnregisterSubjectFramesHandle(FLiveLinkSubjectName InSubje
 	}
 }
 
-void FLiveLinkClient::OnStartSynchronization(FLiveLinkSubjectName InSubjectName, const struct FTimeSynchronizationOpenData& OpenData, const int32 FrameOffset)
-{
-	FScopeLock Lock(&CollectionAccessCriticalSection);
-
-	if (const FLiveLinkCollectionSubjectItem* SubjectItem = Collection->FindEnabledSubject(InSubjectName))
-	{
-		if (FLiveLinkSubject* LinkSubject = SubjectItem->GetLiveSubject())
-		{
-			return LinkSubject->OnStartSynchronization(OpenData, FrameOffset);
-		}
-	}
-}
-
-void FLiveLinkClient::OnSynchronizationEstablished(FLiveLinkSubjectName InSubjectName, const struct FTimeSynchronizationStartData& StartData)
-{
-	FScopeLock Lock(&CollectionAccessCriticalSection);
-
-	if (const FLiveLinkCollectionSubjectItem* SubjectItem = Collection->FindEnabledSubject(InSubjectName))
-	{
-		if (FLiveLinkSubject* LinkSubject = SubjectItem->GetLiveSubject())
-		{
-			return LinkSubject->OnSynchronizationEstablished(StartData);
-		}
-	}
-}
-
-void FLiveLinkClient::OnStopSynchronization(FLiveLinkSubjectName InSubjectName)
-{
-	FScopeLock Lock(&CollectionAccessCriticalSection);
-
-	if (const FLiveLinkCollectionSubjectItem* SubjectItem = Collection->FindEnabledSubject(InSubjectName))
-	{
-		if (FLiveLinkSubject* LinkSubject = SubjectItem->GetLiveSubject())
-		{
-			return LinkSubject->OnStopSynchronization();
-		}
-	}
-}
-
 FSimpleMulticastDelegate& FLiveLinkClient::OnLiveLinkSourcesChanged()
 {
 	return Collection->OnLiveLinkSourcesChanged();
diff --git a/Engine/Plugins/Animation/LiveLink/Source/LiveLink/Private/LiveLinkPreset.cpp b/Engine/Plugins/Animation/LiveLink/Source/LiveLink/Private/LiveLinkPreset.cpp
--- a/Engine/Plugins/Animation/LiveLink/Source/LiveLink/Private/LiveLinkPreset.cpp
+++ b/Engine/Plugins/Animation/LiveLink/Source/LiveLink/Private/LiveLinkPreset.cpp
@@ -12,7 +12,7 @@ bool ULiveLinkPreset::ApplyToClient() const
 		FLiveLinkClient& LiveLinkClient = IModularFeatures::Get().GetModularFeature<FLiveLinkClient>(ILiveLinkClient::ModularFeatureName);
 
 		LiveLinkClient.RemoveAllSources();
-		LiveLinkClient.Tick(0.f);
+		LiveLinkClient.Tick();
 
 		bResult = true;
 		for (const FLiveLinkSourcePreset& SourcePreset : Sources)
diff --git a/Engine/Plugins/Animation/LiveLink/Source/LiveLink/Private/LiveLinkSubject.cpp b/Engine/Plugins/Animation/LiveLink/Source/LiveLink/Private/LiveLinkSubject.cpp
--- a/Engine/Plugins/Animation/LiveLink/Source/LiveLink/Private/LiveLinkSubject.cpp
+++ b/Engine/Plugins/Animation/LiveLink/Source/LiveLink/Private/LiveLinkSubject.cpp
@@ -13,10 +13,6 @@
 #include "TimeSynchronizationSource.h"
 #include "UObject/Class.h"
 
-
-const int32 MIN_FRAMES_TO_REMOVE = 5;
-
-
 void FLiveLinkSubject::Initialize(FLiveLinkSubjectKey InSubjectKey, TSubclassOf<ULiveLinkRole> InRole, ILiveLinkClient* InLiveLinkClient)
 {
 	SubjectKey = InSubjectKey;
@@ -27,31 +23,80 @@ void FLiveLinkSubject::Initialize(FLiveLinkSubjectKey InSubjectKey, TSubclassOf<
 
 void FLiveLinkSubject::Update()
 {
-	// Invalid the snapshot
-	FrameSnapshot.StaticData.Reset();
-	FrameSnapshot.FrameData.Reset();
+	// Clear all frames that are too old
+	if (FrameData.Num() > CachedSettings.BufferSettings.MaxNumberOfFrameToBuffered)
+	{
+		FrameData.RemoveAt(0, FrameData.Num() - CachedSettings.BufferSettings.MaxNumberOfFrameToBuffered, false);
+	}
 
-	// TODO Clear all frames that are too old
-	bool bCanClearOldFrames = TimeSyncData.IsSet() ? TimeSyncData.GetValue().bHasEstablishedSync : true;
-	if (bCanClearOldFrames)
+	if (GetMode() == ELiveLinkSourceMode::EngineTime)
 	{
+		double ValidEngineTime = FApp::GetCurrentTime() - CachedSettings.BufferSettings.EngineTimeOffset - CachedSettings.BufferSettings.ValidEngineTime;
+		int32 FrameIndex = 0;
+		for (const FLiveLinkFrameDataStruct& SourceFrameData : FrameData)
+		{
+			double FrameTime = SourceFrameData.GetBaseData()->WorldTime.GetOffsettedTime();
+			double OffsetTime = ValidEngineTime;
+			if (FrameTime > OffsetTime)
+			{
+				break;
+			}
+			++FrameIndex;
+		}
 
+		if (FrameIndex - 1 >= 0)
+		{
+			FrameData.RemoveAt(0, FrameIndex, false);
+		}
+	}
+	else if (GetMode() == ELiveLinkSourceMode::Timecode)
+	{
+		const FQualifiedFrameTime CurrentSyncTime = FQualifiedFrameTime(FApp::GetTimecode(), FApp::GetTimecodeFrameRate());
+		const FFrameTime CurrentFrameTimeInFrameSpace = CurrentSyncTime.ConvertTo(CachedSettings.BufferSettings.TimecodeFrameRate);
+		int32 FrameIndex = 0;
+		for (const FLiveLinkFrameDataStruct& SourceFrameData : FrameData)
+		{
+			FFrameTime UsedFrameTime = CurrentFrameTimeInFrameSpace - CachedSettings.BufferSettings.TimecodeFrameOffset - CachedSettings.BufferSettings.ValidTimecodeFrame;
+			if (SourceFrameData.GetBaseData()->MetaData.SceneTime.Time > UsedFrameTime)
+			{
+				break;
+			}
+			++FrameIndex;
+		}
+
+		if (FrameIndex - 1 >= 0)
+		{
+			FrameData.RemoveAt(0, FrameIndex, false);
+		}
 	}
 
-	// Build a snapshot for this role 
+	// Build a snapshot for this role
+	bool bSnapshotIsValid = false;
 	if (FrameData.Num() > 0)
 	{
-		if (GetMode() == ELiveLinkSourceMode::TimeSynchronized)
+		switch(GetMode())
 		{
-			const FFrameRate FrameRate = FApp::GetTimecodeFrameRate();
-			const FTimecode Timecode = FApp::GetTimecode();
-			const FQualifiedFrameTime CurrentSyncTime(Timecode.ToFrameNumber(FrameRate), FrameRate);
-			GetFrameAtSceneTime(CurrentSyncTime, FrameSnapshot);
-		}
-		else
+		case ELiveLinkSourceMode::Timecode:
 		{
-			GetFrameAtWorldTime(FApp::GetCurrentTime(), FrameSnapshot);
+			const FQualifiedFrameTime CurrentSyncTime = FQualifiedFrameTime(FApp::GetTimecode(), FApp::GetTimecodeFrameRate());
+			bSnapshotIsValid = GetFrameAtSceneTime(CurrentSyncTime, FrameSnapshot);
 		}
+		break;
+
+		case ELiveLinkSourceMode::EngineTime:
+			bSnapshotIsValid = GetFrameAtWorldTime(FApp::GetCurrentTime(), FrameSnapshot);
+			break;
+
+		case ELiveLinkSourceMode::Latest:
+		default:
+			bSnapshotIsValid = GetLatestFrame(FrameSnapshot);
+		}
+	}
+
+	if (!bSnapshotIsValid)
+	{
+		// Invalidate the snapshot
+		FrameSnapshot.FrameData.Reset();
 	}
 }
 
@@ -59,9 +104,6 @@ void FLiveLinkSubject::ClearFrames()
 {
 	FrameSnapshot.StaticData.Reset();
 	FrameSnapshot.FrameData.Reset();
-
-	LastReadFrame = INDEX_NONE;
-	LastReadTime = 0;
 	FrameData.Reset();
 }
 
@@ -74,13 +116,13 @@ bool FLiveLinkSubject::EvaluateFrameAtWorldTime(double InWorldTime, TSubclassOf<
 {
 	if (Role == nullptr)
 	{
-		UE_LOG(LogLiveLink, Warning, TEXT("Can't evaluate frame for subject %s. No role has been set yet."), *SubjectKey.SubjectName.ToString());
+		UE_LOG(LogLiveLink, Warning, TEXT("Can't evaluate frame for subject '%s'. No role has been set yet."), *SubjectKey.SubjectName.ToString());
 		return false;
 	}
 
 	if (InDesiredRole == nullptr)
 	{
-		UE_LOG(LogLiveLink, Warning, TEXT("Can't evaluate frame for subject %s. Invalid role was received for evaluation."), *SubjectKey.SubjectName.ToString());
+		UE_LOG(LogLiveLink, Warning, TEXT("Can't evaluate frame for subject '%s'. Invalid role was received for evaluation."), *SubjectKey.SubjectName.ToString());
 		return false;
 	}
 
@@ -111,22 +153,20 @@ bool FLiveLinkSubject::EvaluateFrameAtSceneTime(const FTimecode& InSceneTime, TS
 {
 	if (Role == nullptr)
 	{
-		UE_LOG(LogLiveLink, Warning, TEXT("Can't evaluate frame for subject %s. No role has been set yet."), *SubjectKey.SubjectName.ToString());
+		UE_LOG(LogLiveLink, Warning, TEXT("Can't evaluate frame for subject '%s'. No role has been set yet."), *SubjectKey.SubjectName.ToString());
 		return false;
 	}
 
 	if (InDesiredRole == nullptr)
 	{
-		UE_LOG(LogLiveLink, Warning, TEXT("Can't evaluate frame for subject %s. Invalid role was received for evaluation."), *SubjectKey.SubjectName.ToString());
+		UE_LOG(LogLiveLink, Warning, TEXT("Can't evaluate frame for subject '%s'. Invalid role was received for evaluation."), *SubjectKey.SubjectName.ToString());
 		return false;
 	}
 
 	bool bSuccess = false;
 	if (FrameData.Num() != 0)
 	{
-		const FFrameRate FrameRate = FApp::GetTimecodeFrameRate();
-		const FQualifiedFrameTime UseTime(InSceneTime.ToFrameNumber(FrameRate), FrameRate);
-
+		const FQualifiedFrameTime UseTime = FQualifiedFrameTime(InSceneTime, FApp::GetTimecodeFrameRate());
 		if (Role == InDesiredRole || Role->IsChildOf(InDesiredRole))
 		{
 			GetFrameAtSceneTime(UseTime, OutFrame);
@@ -140,7 +180,7 @@ bool FLiveLinkSubject::EvaluateFrameAtSceneTime(const FTimecode& InSceneTime, TS
 		}
 		else
 		{
-			UE_LOG(LogLiveLink, Warning, TEXT("Can't evaluate frame for subject %s. Role %s is incompatible with current role %s"), *SubjectKey.SubjectName.ToString(), *InDesiredRole->GetName(), *Role->GetName());
+			UE_LOG(LogLiveLink, Warning, TEXT("Can't evaluate frame for subject '%s'. Role '%s' is incompatible with current role '%s'"), *SubjectKey.SubjectName.ToString(), *InDesiredRole->GetName(), *Role->GetName());
 		}
 	}
 
@@ -155,7 +195,11 @@ bool FLiveLinkSubject::HasStaticData() const
 void FLiveLinkSubject::AddFrameData(FLiveLinkFrameDataStruct&& InFrameData)
 {
 	check(IsInGameThread());
-	check(StaticData.IsValid());
+	if (!StaticData.IsValid())
+	{
+		UE_LOG(LogLiveLink, Warning, TEXT("Can't add frame for subject '%s'. The static frame data is invalid."), *SubjectKey.SubjectName.ToString());
+		return;
+	}
 
 	if (Role == nullptr)
 	{
@@ -164,381 +208,285 @@ void FLiveLinkSubject::AddFrameData(FLiveLinkFrameDataStruct&& InFrameData)
 
 	if (Role->GetDefaultObject<ULiveLinkRole>()->GetFrameDataStruct() != InFrameData.GetStruct())
 	{
-		UE_LOG(LogLiveLink, Warning, TEXT("Can't add frame for subject %s. The frame data is incompatible with current role %s"), *SubjectKey.SubjectName.ToString(), *Role->GetName());
+		UE_LOG(LogLiveLink, Warning, TEXT("Can't add frame for subject '%s'. The frame data is incompatible with current role '%s'."), *SubjectKey.SubjectName.ToString(), *Role->GetName());
 		return;
 	}
 
 	if (!FLiveLinkRoleTrait::Validate(Role, InFrameData))
 	{
-		UE_LOG(LogLiveLink, Warning, TEXT("Trying to add unsupported frame data type to role %s."), *Role->GetName());
+		UE_LOG(LogLiveLink, Warning, TEXT("Trying to add unsupported frame data type to role '%s'."), *Role->GetName());
 		return;
 	}
 
 	int32 FrameIndex = INDEX_NONE;
 	switch (CachedSettings.SourceMode)
 	{
-	case ELiveLinkSourceMode::TimeSynchronized:
-		if (TimeSyncData.IsSet())
-		{
-			FrameIndex = AddFrame_TimeSynchronized(InFrameData.GetBaseData()->MetaData.SceneTime.Time, TimeSyncData.GetValue());
-		}
-		else
-		{
-			FrameIndex = AddFrame_Default(InFrameData.GetBaseData()->WorldTime);
-		}
+	case ELiveLinkSourceMode::EngineTime:
+		FrameIndex = FindNewFrame_WorldTime(InFrameData.GetBaseData()->WorldTime);
 		break;
-
-	case ELiveLinkSourceMode::Interpolated:
-		FrameIndex = AddFrame_Interpolated(InFrameData.GetBaseData()->WorldTime);
+	case ELiveLinkSourceMode::Timecode:
+		FrameIndex = FindNewFrame_SceneTime(InFrameData.GetBaseData()->MetaData.SceneTime);
 		break;
-
+	case ELiveLinkSourceMode::Latest:
 	default:
-		FrameIndex = AddFrame_Default(InFrameData.GetBaseData()->WorldTime);
+		FrameIndex = FindNewFrame_Latest();
 		break;
 	}
 	
-	for (ULiveLinkFramePreProcessor::FWorkerSharedPtr PreProcessor : FramePreProcessors)
+	if (FrameIndex != INDEX_NONE)
 	{
-		PreProcessor->PreProcessFrame(InFrameData);
+		for (ULiveLinkFramePreProcessor::FWorkerSharedPtr PreProcessor : FramePreProcessors)
+		{
+			PreProcessor->PreProcessFrame(InFrameData);
+		}
+
+		FrameData.Insert(MoveTemp(InFrameData), FrameIndex);
 	}
 
-	FrameData.Insert(MoveTemp(InFrameData), FrameIndex);
 	LastPushTime = FApp::GetCurrentTime();
 }
 
-int32 FLiveLinkSubject::AddFrame_Default(const FLiveLinkWorldTime& WorldTime)
+int32 FLiveLinkSubject::FindNewFrame_WorldTime(const FLiveLinkWorldTime& WorldTime) const
 {
-	if (WorldTime.Time < LastReadTime)
+	int32 FrameIndex = FrameData.Num() - 1;
+	const double NewFrameOffsettedTime = WorldTime.GetOffsettedTime();
+	for (; FrameIndex >= 0; --FrameIndex)
 	{
-		//Gone back in time
-		FrameData.Reset();
-		LastReadTime = 0;
-		SubjectTimeOffset = WorldTime.Offset;
-	}
-
-	int32 FrameIndex = 0;
-	if (FrameData.Num() == 0)
-	{
-		LastReadFrame = 0;
-	}
-	else
-	{
-		if (LastReadFrame > MIN_FRAMES_TO_REMOVE)
-		{
-			check(FrameData.Num() > LastReadFrame);
-			FrameData.RemoveAt(0, LastReadFrame, false);
-			LastReadFrame = 0;
-		}
-
-		for (FrameIndex = FrameData.Num() - 1; FrameIndex >= 0; --FrameIndex)
+		if (FrameData[FrameIndex].GetBaseData()->WorldTime.GetOffsettedTime() <= NewFrameOffsettedTime)
 		{
-			if (FrameData[FrameIndex].GetBaseData()->WorldTime.Time <= WorldTime.Time)
-			{
-				break;
-			}
+			break;
 		}
-
-		FrameIndex += 1;
 	}
 
-	return FrameIndex;
+	return FrameIndex + 1;
 }
 
-int32 FLiveLinkSubject::AddFrame_Interpolated(const FLiveLinkWorldTime& WorldTime)
+int32 FLiveLinkSubject::FindNewFrame_SceneTime(const FQualifiedFrameTime& QualifiedFrameTime) const
 {
-	if (LastReadFrame > MIN_FRAMES_TO_REMOVE)
+	if (QualifiedFrameTime.Time.FloorToFrame() < 0)
 	{
-		check(FrameData.Num() > LastReadFrame);
-		FrameData.RemoveAt(0, LastReadFrame - 1, false);
-		LastReadFrame = 1;
+		UE_LOG(LogLiveLink, Warning, TEXT("Trying to add a frame that does not have a valid scene time (timecode). The Subject is '%s'."), *SubjectKey.SubjectName.ToString());
+		return INDEX_NONE;
 	}
 
-	return AddFrame_Default(WorldTime);
-}
-
-int32 FLiveLinkSubject::AddFrame_TimeSynchronized(const FFrameTime& FrameTime, FLiveLinkTimeSynchronizationData& InSynchronizationData)
-{
-	int32 FrameIndex = 0;
-
-	// If we're not actively synchronizing, we don't need to do anything special.
-	if (FrameData.Num() == 0)
-	{
-		LastReadTime = 0;
-		LastReadFrame = 0;
-	}
-	else if (InSynchronizationData.RolloverModulus.IsSet())
+	if (QualifiedFrameTime.Rate != CachedSettings.BufferSettings.TimecodeFrameRate)
 	{
-		const FFrameTime UseFrameTime = UTimeSynchronizationSource::AddOffsetWithRolloverModulus(FrameTime, InSynchronizationData.Offset, InSynchronizationData.RolloverModulus.GetValue());
-		FrameIndex = AddFrame_TimeSynchronized</*bWithRollover=*/true>(UseFrameTime, InSynchronizationData);
+		UE_LOG(LogLiveLink, Warning, TEXT("Trying to add a frame that timecode frame rate does not match with the expected frame rate. The Subject is '%s'."), *SubjectKey.SubjectName.ToString());
+		return INDEX_NONE;
 	}
-	else
-	{
-		FrameIndex = AddFrame_TimeSynchronized</*bWithRollover=*/false>(FrameTime + InSynchronizationData.Offset, InSynchronizationData);
-	}
-
-	return FrameIndex;
-}
 
-template<bool bWithRollover>
-int32 FLiveLinkSubject::AddFrame_TimeSynchronized(const FFrameTime& FrameTime, FLiveLinkTimeSynchronizationData& InSynchronizationData)
-{
-	//We keep buffering as long as synchronization hasn't been achieved
-	if (InSynchronizationData.bHasEstablishedSync && LastReadFrame > MIN_FRAMES_TO_REMOVE)
+	int32 FrameIndex = FrameData.Num() - 1;
+	for (; FrameIndex >= 0; --FrameIndex)
 	{
-		check(FrameData.Num() > LastReadFrame);
-
-		if (bWithRollover)
+		if (FrameData[FrameIndex].GetBaseData()->MetaData.SceneTime.Time <= QualifiedFrameTime.Time)
 		{
-			int32& RolloverFrame = InSynchronizationData.RolloverFrame;
-
-			// If we had previously detected that a roll over had occurred in the range of frames we have,
-			// then we need to adjust that as well.
-			if (RolloverFrame > 0)
-			{
-				RolloverFrame = RolloverFrame - LastReadFrame;
-				if (RolloverFrame <= 0)
-				{
-					RolloverFrame = INDEX_NONE;
-				}
-			}
+			break;
 		}
-
-		FrameData.RemoveAt(0, LastReadFrame, false);
-		LastReadFrame = 0;
 	}
 
-	return FindFrameIndex_TimeSynchronized</*bForInsert=*/true, bWithRollover>(FrameTime, InSynchronizationData);
-}
-
-void FLiveLinkSubject::ResetFrame(FLiveLinkSubjectFrameData& OutFrame) const
-{
-	//Allocate and copy over our static data for that frame.
-	OutFrame.StaticData.InitializeWith(StaticData);
-
-	//Only reset the frame data. Copy will be done later on depending on sampling type
-	OutFrame.FrameData.Reset();
+	return FrameIndex + 1;
 }
 
-void FLiveLinkSubject::GetFrameAtWorldTime(const double InSeconds, FLiveLinkSubjectFrameData& OutFrame)
+int32 FLiveLinkSubject::FindNewFrame_Latest() const
 {
-	ResetFrame(OutFrame);
-
-	switch (CachedSettings.SourceMode)
-	{
-	case ELiveLinkSourceMode::TimeSynchronized:
-		ensureMsgf(false, TEXT("Attempting to use WorldTime for a TimeSynchronized source! Source = %s"), *SubjectKey.SubjectName.ToString());
-		GetFrameAtWorldTime_Default(InSeconds, OutFrame);
-		break;
-
-	case ELiveLinkSourceMode::Interpolated:
-
-		GetFrameAtWorldTime_Interpolated(InSeconds, OutFrame);
-		break;
-
-	default:
-		GetFrameAtWorldTime_Default(InSeconds, OutFrame);
-		break;
-	}
+	return FrameData.Num();
 }
 
-void FLiveLinkSubject::GetFrameAtSceneTime(const FQualifiedFrameTime& InSceneTime, FLiveLinkSubjectFrameData& OutFrame)
+bool FLiveLinkSubject::GetFrameAtWorldTime(const double InSeconds, FLiveLinkSubjectFrameData& OutFrame)
 {
-	ResetFrame(OutFrame);
-	
-	switch (CachedSettings.SourceMode)
+	bool bResult = FrameData.Num() != 0;
+	if (bResult)
 	{
-	case ELiveLinkSourceMode::TimeSynchronized:
-
-		if (TimeSyncData.IsSet())
+		if (FrameInterpolationProcessor.IsValid())
 		{
-			/**
-			 * Get the actual frame of data at the specified time.
-			 * Note, unlike GetTimeSyncData which subtracts the offset from the Saved Frames' Timecodes,
-			 * we're adding the offset to the Incoming Timecode Value.
-			 *
-			 * The reason for this can be thought of as a simple translation between two spaces: Real Space and Source Space.
-			 * InSceneTime (and generally timecode the engine is using) is considered Real Space.
-			 * The timecode reported with our frames is considered Source Space.
-			 * `Source Space = Real Space + Offset`, and therefore `Real Space = Source Space - Offset`.
-			 *
-			 * We could just apply the offset on the frame's when they are Pushed to Live Link, but in general
-			 * we don't want to modify source data. Further, the Frame Offset is not available until synchronization
-			 * is started now, but it may change to be more dynamic in the future, and we don't want to have to
-			 * "restripe" the offsets in that case.
-			 */
-			const FFrameTime FrameTime = InSceneTime.ConvertTo(CachedSettings.TimeSynchronizationSettings->FrameRate) - CachedSettings.TimeSynchronizationSettings->FrameDelay;
-			if (TimeSyncData->RolloverModulus.IsSet())
-			{
-				const FFrameTime UseFrameTime = UTimeSynchronizationSource::AddOffsetWithRolloverModulus(FrameTime, TimeSyncData->Offset, TimeSyncData->RolloverModulus.GetValue());
-				GetFrameAtSceneTime_TimeSynchronized</*bWithRollover=*/true>(UseFrameTime, TimeSyncData.GetValue(), OutFrame);
-			}
-			else
-			{
-				const FFrameTime UseFrameTime = FrameTime + TimeSyncData->Offset;
-				GetFrameAtSceneTime_TimeSynchronized</*bWithRollover=*/false>(UseFrameTime, TimeSyncData.GetValue(), OutFrame);
-			}
+			bResult = GetFrameAtWorldTime_Interpolated(InSeconds, OutFrame);
 		}
 		else
 		{
-			GetFrameAtWorldTime_Default(InSceneTime.AsSeconds(), OutFrame);
+			bResult = GetFrameAtWorldTime_Closest(InSeconds, OutFrame);
 		}
-		break;
 
-	default:
-		ensureMsgf(false, TEXT("Attempting to use SceneTime for a non TimeSynchronized source! Source = %s Mode = %d"), *SubjectKey.SubjectName.ToString(), static_cast<int32>(CachedSettings.SourceMode));
-		GetFrameAtWorldTime_Default(InSceneTime.AsSeconds(), OutFrame);
-		break;
-	}
-}
-
-void FLiveLinkSubject::GetFrameAtWorldTime_Default(const double InSeconds, FLiveLinkSubjectFrameData& OutFrame)
-{
-	FLiveLinkFrameDataStruct& LastFrame = FrameData.Last();
-
-	//Copy current snapshot of our frame data
-	OutFrame.FrameData.InitializeWith(LastFrame);
-
-	LastReadTime = LastFrame.GetBaseData()->WorldTime.Time;
-	LastReadFrame = FrameData.Num() - 1;
-}
-
-void FLiveLinkSubject::GetFrameAtWorldTime_Interpolated(const double InSeconds, FLiveLinkSubjectFrameData& OutFrame)
-{
-	if (FrameInterpolationProcessor.IsValid())
-	{
-		LastReadTime = (InSeconds - SubjectTimeOffset) - CachedSettings.InterpolationSettings->InterpolationOffset;
-		FrameInterpolationProcessor->Interpolate(LastReadTime, StaticData, FrameData, OutFrame, LastReadFrame);
-	}
-	else
-	{
-		GetFrameAtWorldTime_Closest(InSeconds, OutFrame);
+		if (bResult && !OutFrame.StaticData.IsValid())
+		{
+			OutFrame.StaticData.InitializeWith(StaticData.GetStruct(), StaticData.GetBaseData());
+		}
 	}
+	return bResult;
 }
 
-void FLiveLinkSubject::GetFrameAtWorldTime_Closest(const double InSeconds, FLiveLinkSubjectFrameData& OutFrame)
+bool FLiveLinkSubject::GetFrameAtWorldTime_Closest(const double InSeconds, FLiveLinkSubjectFrameData& OutFrame)
 {
 	check(FrameData.Num() != 0);
 
-	LastReadTime = InSeconds - SubjectTimeOffset - CachedSettings.InterpolationSettings->InterpolationOffset;
+	bool bBuiltFrame = false;
+	const double ReadTime = (InSeconds) - CachedSettings.BufferSettings.EngineTimeOffset;
 	for (int32 FrameIndex = FrameData.Num() - 1; FrameIndex >= 0; --FrameIndex)
 	{
-		const FLiveLinkFrameDataStruct& SourceFrameData = FrameData[FrameIndex];
-		if (SourceFrameData.GetBaseData()->WorldTime.Time < LastReadTime)
+		const double Time = FrameData[FrameIndex].GetBaseData()->WorldTime.GetOffsettedTime();
+		if (Time < ReadTime)
 		{
 			if (FrameIndex == FrameData.Num() - 1)
 			{
-				LastReadFrame = FrameIndex;
-
-				// Copy over the frame directly
-				OutFrame.FrameData.InitializeWith(SourceFrameData);
-				return;
+				//Copy over the frame directly
+				OutFrame.FrameData.InitializeWith(FrameData[FrameIndex]);
+				bBuiltFrame = true;
+#if WITH_EDITOR
+				SnapshotIndex = FrameIndex;
+				NumberOfBufferAtSnapshot = FrameData.Num();
+#endif 
+				break;
 			}
 			else
 			{
-				LastReadFrame = FrameIndex;
-
-				const FLiveLinkFrameDataStruct& PostFrame = FrameData[FrameIndex + 1];
-				const float BlendWeight = (LastReadTime - SourceFrameData.GetBaseData()->WorldTime.Time) / (PostFrame.GetBaseData()->WorldTime.Time - SourceFrameData.GetBaseData()->WorldTime.Time);
-				int32 CopyIndex = (BlendWeight > 0.5f) ? FrameIndex+1 : FrameIndex;
+				const double NextTime = FrameData[FrameIndex + 1].GetBaseData()->WorldTime.GetOffsettedTime();
+				const float BlendWeight = (ReadTime - NextTime) / (NextTime - Time);
+				int32 CopyIndex = (BlendWeight > 0.5f) ? FrameIndex : FrameIndex + 1;
 				OutFrame.FrameData.InitializeWith(FrameData[CopyIndex].GetStruct(), FrameData[CopyIndex].GetBaseData());
-				return;
+				bBuiltFrame = true;
+#if WITH_EDITOR
+				SnapshotIndex = CopyIndex;
+				NumberOfBufferAtSnapshot = FrameData.Num();
+#endif 
+				break;
 			}
 		}
 	}
 
-	// Failed to find an interp point so just take oldest frame
-	LastReadFrame = 0;
-	OutFrame.FrameData.InitializeWith(FrameData[0].GetStruct(), FrameData[0].GetBaseData());
+	if (!bBuiltFrame)
+	{
+		// Failed to find an interp point so just take oldest frame
+		OutFrame.FrameData.InitializeWith(FrameData[0].GetStruct(), FrameData[0].GetBaseData());
+#if WITH_EDITOR
+		SnapshotIndex = INDEX_NONE;
+		NumberOfBufferAtSnapshot = FrameData.Num();
+#endif 
+	}
+
+	return true;
 }
 
-template<bool bWithRollover>
-void FLiveLinkSubject::GetFrameAtSceneTime_TimeSynchronized(const FFrameTime& InTime, FLiveLinkTimeSynchronizationData& InSynchronizationData, FLiveLinkSubjectFrameData& OutFrame)
+bool FLiveLinkSubject::GetFrameAtWorldTime_Interpolated(const double InSeconds, FLiveLinkSubjectFrameData& OutFrame)
 {
-	const int32 UseFrame = FindFrameIndex_TimeSynchronized</*bForInsert=*/false, bWithRollover>(InTime, InSynchronizationData);
-
-	OutFrame.FrameData.InitializeWith(FrameData[UseFrame]);
+	check(FrameData.Num() != 0);
 
-	LastReadTime = FrameData[UseFrame].GetBaseData()->WorldTime.Time;
-	LastReadFrame = UseFrame;
+	double ReadTime = InSeconds - CachedSettings.BufferSettings.EngineTimeOffset;
+	FrameInterpolationProcessor->Interpolate(ReadTime, StaticData, FrameData, OutFrame);
+	return true;
 }
 
-template<bool bForInsert, bool bWithRollover>
-int32 FLiveLinkSubject::FindFrameIndex_TimeSynchronized(const FFrameTime& InFrameTime, FLiveLinkTimeSynchronizationData& InSynchronizationData)
+bool FLiveLinkSubject::GetFrameAtSceneTime(const FQualifiedFrameTime& InTimeInEngineFrameRate, FLiveLinkSubjectFrameData& OutFrame)
 {
-	if (FrameData.Num() == 0)
+	bool bResult = FrameData.Num() != 0;
+	if (bResult)
 	{
-		return 0;
+		if (FrameInterpolationProcessor.IsValid())
+		{
+			bResult = GetFrameAtSceneTime_Interpolated(InTimeInEngineFrameRate, OutFrame);
+		}
+		else
+		{
+			bResult = GetFrameAtSceneTime_Closest(InTimeInEngineFrameRate, OutFrame);
+		}
+
+		if (bResult && !OutFrame.StaticData.IsValid())
+		{
+			OutFrame.StaticData.InitializeWith(StaticData.GetStruct(), StaticData.GetBaseData());
+		}
 	}
+	return bResult;
+}
 
-	// Preroll / Synchronization should handle the case where there are any time skips by simply clearing out the buffered data.
-	// Therefore, there are only 2 cases where time would go backwards:
-	// 1. We've received frames out of order. In this case, we want to push it backwards.
-	// 2. We've rolled over. In that case, value have wrapped around zero (and appear "smaller") but should be treated as newer.
+bool FLiveLinkSubject::GetFrameAtSceneTime_Closest(const FQualifiedFrameTime& InTimeInEngineFrameRate, FLiveLinkSubjectFrameData& OutFrame)
+{
+	check(FrameData.Num() != 0);
 
-	// Further, when we're not inserting a value, we're guaranteed that the frame time should always go up
-	// (or stay the same). So, in that case we only need to search between our LastReadFrameTime and the Newest Frame.
-	// That assumption will break if external code tries to grab anything other than the frame of data we build internally.
+	bool bBuiltFrame = false;
 
-	// Finally, we only update the RolloverFrame value when inserting values. This is because we may query for a rollover frame
-	// before we receive a rollover frame (in the case of missing or unordered frames).
-	// We generally don't want to modify state if we're just reading data.
+	const FFrameTime FrameOffset = FQualifiedFrameTime(CachedSettings.BufferSettings.TimecodeFrameOffset, CachedSettings.BufferSettings.TimecodeFrameRate).ConvertTo(InTimeInEngineFrameRate.Rate);
+	const FFrameTime ReadTime = InTimeInEngineFrameRate.Time - FrameOffset;
 
-	int32 HighFrame = FrameData.Num() - 1;
-	int32 LowFrame = bForInsert ? 0 : LastReadFrame;
-	int32 FrameIndex = HighFrame;
+	FTimecode ReadTimecode = FTimecode::FromFrameNumber(ReadTime.GetFrame(), InTimeInEngineFrameRate.Rate, false);
 
-	if (bWithRollover)
+	for (int32 FrameIndex = FrameData.Num() - 1; FrameIndex >= 0; --FrameIndex)
 	{
-		bool bDidRollover = false;
-		int32& RolloverFrame = InSynchronizationData.RolloverFrame;
-		const FFrameTime& CompareFrameTime = ((RolloverFrame == INDEX_NONE) ? FrameData.Last() : FrameData[RolloverFrame - 1]).GetBaseData()->MetaData.SceneTime.Time;
-		UTimeSynchronizationSource::FindDistanceBetweenFramesWithRolloverModulus(CompareFrameTime, InFrameTime, InSynchronizationData.RolloverModulus.GetValue(), bDidRollover);
-
-		if (RolloverFrame == INDEX_NONE)
+		const FFrameTime Time = FrameData[FrameIndex].GetBaseData()->MetaData.SceneTime.ConvertTo(InTimeInEngineFrameRate.Rate);
+		FTimecode Timecode = FTimecode::FromFrameNumber(Time.GetFrame(), InTimeInEngineFrameRate.Rate, false);
+		if (Time < ReadTime)
 		{
-			if (bDidRollover)
+			if (FrameIndex == FrameData.Num() - 1)
 			{
-				if (bForInsert)
-				{
-					RolloverFrame = HighFrame;
-					FrameIndex = FrameData.Num();
-				}
-				else
-				{
-					FrameIndex = HighFrame;
-				}
+				//Copy over the frame directly
+				OutFrame.FrameData.InitializeWith(FrameData[FrameIndex]);
+				bBuiltFrame = true;
+#if WITH_EDITOR
+				SnapshotIndex = FrameIndex;
+				NumberOfBufferAtSnapshot = FrameData.Num();
+#endif 
 
-				return FrameIndex;
-			}
-		}
-		else
-		{
-			if (bDidRollover)
-			{
-				LowFrame = RolloverFrame;
+				break;
 			}
 			else
 			{
-				HighFrame = RolloverFrame - 1;
-				if (bForInsert)
-				{
-					++RolloverFrame;
-				}
+				const FFrameTime NextTime = FrameData[FrameIndex+1].GetBaseData()->MetaData.SceneTime.ConvertTo(InTimeInEngineFrameRate.Rate);
+				const double BlendWeight = (ReadTime - NextTime).AsDecimal() / (NextTime - Time).AsDecimal();
+				int32 CopyIndex = (BlendWeight > 0.5) ? FrameIndex : FrameIndex + 1;
+				OutFrame.FrameData.InitializeWith(FrameData[CopyIndex].GetStruct(), FrameData[CopyIndex].GetBaseData());
+				bBuiltFrame = true;
+#if WITH_EDITOR
+				SnapshotIndex = CopyIndex;
+				NumberOfBufferAtSnapshot = FrameData.Num();
+#endif 
+				break;
 			}
 		}
 	}
 
-	if (bForInsert)
+	if (!bBuiltFrame)
 	{
-		for (; LowFrame <= FrameIndex && FrameData[FrameIndex].GetBaseData()->MetaData.SceneTime.Time > InFrameTime; --FrameIndex);
-		FrameIndex += 1;
+#if WITH_EDITOR
+		SnapshotIndex = INDEX_NONE;
+		NumberOfBufferAtSnapshot = FrameData.Num();
+#endif 
+		// Failed to find an interp point so just take oldest frame
+		OutFrame.FrameData.InitializeWith(FrameData[0].GetStruct(), FrameData[0].GetBaseData());
 	}
-	else
+
+	return true;
+}
+
+bool FLiveLinkSubject::GetFrameAtSceneTime_Interpolated(const FQualifiedFrameTime& InTimeInEngineFrameRate, FLiveLinkSubjectFrameData& OutFrame)
+{
+	check(FrameData.Num() != 0);
+
+	const FFrameTime ReadTime = InTimeInEngineFrameRate.Time - CachedSettings.BufferSettings.TimecodeFrameOffset;
+	FrameInterpolationProcessor->Interpolate(FQualifiedFrameTime(ReadTime, InTimeInEngineFrameRate.Rate), StaticData, FrameData, OutFrame);
+	return true;
+}
+
+bool FLiveLinkSubject::GetLatestFrame(FLiveLinkSubjectFrameData& OutFrame)
+{
+	if (FrameData.Num())
 	{
-		for (; LowFrame < FrameIndex && FrameData[FrameIndex].GetBaseData()->MetaData.SceneTime.Time > InFrameTime; --FrameIndex);
+		FLiveLinkFrameDataStruct& LastDataStruct = FrameData.Last();
+		OutFrame.FrameData.InitializeWith(LastDataStruct.GetStruct(), LastDataStruct.GetBaseData());
+		OutFrame.StaticData.InitializeWith(StaticData.GetStruct(), StaticData.GetBaseData());
+#if WITH_EDITOR
+		SnapshotIndex = FrameData.Num() - 1;
+		NumberOfBufferAtSnapshot = FrameData.Num();
+#endif 
+
+		return true;
 	}
+	return false;
+}
 
-	return FrameIndex;
+void FLiveLinkSubject::ResetFrame(FLiveLinkSubjectFrameData& OutFrame) const
+{
+	//Allocate and copy over our static data for that frame.
+	OutFrame.StaticData.InitializeWith(StaticData);
+
+	//Only reset the frame data. Copy will be done later on depending on sampling type
+	OutFrame.FrameData.Reset();
 }
 
 void FLiveLinkSubject::SetStaticData(TSubclassOf<ULiveLinkRole> InRole, FLiveLinkStaticDataStruct&& InStaticData)
@@ -556,7 +504,6 @@ void FLiveLinkSubject::SetStaticData(TSubclassOf<ULiveLinkRole> InRole, FLiveLin
 		//Set initial blending processor to the role's default one. User will be able to modify it afterwards.
 		FrameData.Reset();
 		StaticData = MoveTemp(InStaticData);
-		StaticDataGuid = FGuid::NewGuid();
 	}
 	else
 	{
@@ -570,47 +517,32 @@ void FLiveLinkSubject::CacheSettings(ULiveLinkSourceSettings* SourceSetting, ULi
 
 	if (SourceSetting)
 	{
-		const bool bSourceModeChanged = SourceSetting->Mode != CachedSettings.SourceMode;
-		if (bSourceModeChanged)
+#if WITH_EDITOR
+		if (NumberOfBufferAtSnapshot > 0)
 		{
-			CachedSettings.TimeSynchronizationSettings.Reset();
-			CachedSettings.InterpolationSettings.Reset();
+			FLiveLinkSourceDebugInfo DebugInfo;
+			DebugInfo.SubjectName = SubjectKey.SubjectName;
+			DebugInfo.SnapshotIndex = SnapshotIndex + 1;
+			DebugInfo.NumberOfBufferAtSnapshot = NumberOfBufferAtSnapshot;
+			SourceSetting->SourceDebugInfos.Add(DebugInfo);
+		}
+		SnapshotIndex = INDEX_NONE;
+		NumberOfBufferAtSnapshot = 0;
+#endif
 
+		const bool bSourceModeChanged = SourceSetting->Mode != CachedSettings.SourceMode;
+		const bool bTimecodeFrameRateChanged = SourceSetting->Mode == ELiveLinkSourceMode::Timecode && SourceSetting->BufferSettings.TimecodeFrameRate != CachedSettings.BufferSettings.TimecodeFrameRate;
+		if (bSourceModeChanged || bTimecodeFrameRateChanged)
+		{
 			FrameData.Reset();
-
-			switch (CachedSettings.SourceMode)
-			{
-			case ELiveLinkSourceMode::TimeSynchronized:
-				TimeSyncData.Reset();
-				break;
-
-			default:
-				break;
-			}
 		}
 
 		CachedSettings.SourceMode = SourceSetting->Mode;
+		CachedSettings.BufferSettings = SourceSetting->BufferSettings;
 
-		// Even if the mode didn't change, settings may have updated.
-		// Handle those changes now.
-		switch (CachedSettings.SourceMode)
-		{
-		case ELiveLinkSourceMode::TimeSynchronized:
-			CachedSettings.TimeSynchronizationSettings = SourceSetting->TimeSynchronizationSettings;
-			break;
-
-		case ELiveLinkSourceMode::Interpolated:
-			CachedSettings.InterpolationSettings = SourceSetting->InterpolationSettings;
-			break;
+		CachedSettings.BufferSettings.MaxNumberOfFrameToBuffered = FMath::Max(CachedSettings.BufferSettings.MaxNumberOfFrameToBuffered, 1);
 
-		default:
-			break;
-		}
-	}
-
-	if (SubjectSetting)
-	{
-		// Create the new PreProcessors for this frame
+		// Create a new or fetch the PreProcessors for this frame
 		FramePreProcessors.Reset();
 		for (ULiveLinkFramePreProcessor* PreProcessor : SubjectSetting->PreProcessors)
 		{
@@ -624,14 +556,14 @@ void FLiveLinkSubject::CacheSettings(ULiveLinkSourceSettings* SourceSetting, ULi
 			}
 		}
 
-		// Create the new interpolation for this frame
+		// Create a new or fetch the interpolation for this frame
 		FrameInterpolationProcessor.Reset();
 		if (SubjectSetting->InterpolationProcessor)
 		{
 			FrameInterpolationProcessor = SubjectSetting->InterpolationProcessor->FetchWorker();
 		}
 
-		// Create the new translators for this frame
+		// Create a new or fetch the translators for this frame
 		FrameTranslators.Reset();
 		for (ULiveLinkFrameTranslator* Translator : SubjectSetting->Translators)
 		{
@@ -651,102 +583,28 @@ FLiveLinkSubjectTimeSyncData FLiveLinkSubject::GetTimeSyncData()
 {
 	FLiveLinkSubjectTimeSyncData SyncData;
 	SyncData.bIsValid = FrameData.Num() > 0;
-	SyncData.Settings = CachedSettings.TimeSynchronizationSettings.Get(FLiveLinkTimeSynchronizationSettings());
 
 	if (SyncData.bIsValid)
 	{
-		const FFrameTime LastSceneTime = FrameData.Last().GetBaseData()->MetaData.SceneTime.Time;
-		const FFrameTime FirstSceneTime = FrameData[0].GetBaseData()->MetaData.SceneTime.Time;
-
-		if (TimeSyncData.IsSet())
-		{
-			/**
-			* It's possible that the Timecode received by the Subject / Source doesn't align perfectly with other Timecode
-			* Sources. This is because there is inherent delays between the initial transmission of a timecode signal
-			* and its receipt on any given source machine. This can further be exacerbated if the source doesn't associate
-			* timecode with a frame of data until after it's finished processing it.
-			* To make things align properly in that case, users specify an additional FrameOffset for Subjects.
-			*
-			* For example, a Subject may report that it is at timecode X, whereas the raw source data that was processed
-			* to generate that frame may correspond to timecode X-O (where O = offset).
-			*
-			* To compensate for that, we will just report our frame data shifted over by the offset here, and then adjust
-			* the desired timecode in GetFrameAtSceneTime.
-			*/
-			const int32 FrameOffset = TimeSyncData->Offset;
-			if (!TimeSyncData->RolloverModulus.IsSet())
-			{
-				SyncData.NewestSampleTime = LastSceneTime - FrameOffset;
-				SyncData.OldestSampleTime = FirstSceneTime - FrameOffset;
-			}
-			else
-			{
-				const FFrameTime& RolloverModulus = TimeSyncData->RolloverModulus.GetValue();
-				SyncData.NewestSampleTime = UTimeSynchronizationSource::AddOffsetWithRolloverModulus(LastSceneTime, -FrameOffset, RolloverModulus);
-				SyncData.OldestSampleTime = UTimeSynchronizationSource::AddOffsetWithRolloverModulus(FirstSceneTime, -FrameOffset, RolloverModulus);
-			}
-		}
-		else
-		{
-			SyncData.NewestSampleTime = LastSceneTime;
-			SyncData.OldestSampleTime = FirstSceneTime;
-		}
-		
-		SyncData.SkeletonGuid = StaticDataGuid;
+		SyncData.NewestSampleTime = FrameData.Last().GetBaseData()->MetaData.SceneTime.Time;
+		SyncData.OldestSampleTime = FrameData[0].GetBaseData()->MetaData.SceneTime.Time;
+		SyncData.SampleFrameRate = FrameData[0].GetBaseData()->MetaData.SceneTime.Rate;
 	}
 
 	return SyncData;
 }
 
-void FLiveLinkSubject::OnStartSynchronization(const FTimeSynchronizationOpenData& OpenData, const int32 FrameOffset)
-{
-	if (CachedSettings.SourceMode == ELiveLinkSourceMode::TimeSynchronized)
-	{
-		ensure(!TimeSyncData.IsSet());
-		TimeSyncData = FLiveLinkTimeSynchronizationData();
-		TimeSyncData->RolloverModulus = OpenData.RolloverFrame;
-		TimeSyncData->SyncFrameRate = OpenData.SynchronizationFrameRate;
-		TimeSyncData->Offset = FrameOffset;
-
-		// Still need to check this, because OpenData.RolloverFrame is a TOptional which may be unset.
-		if (TimeSyncData->RolloverModulus.IsSet())
-		{
-			TimeSyncData->RolloverModulus = FFrameRate::TransformTime(TimeSyncData->RolloverModulus.GetValue(), OpenData.SynchronizationFrameRate, CachedSettings.TimeSynchronizationSettings->FrameRate);
-		}
-
-		ClearFrames();
-	}
-	else
-	{
-		TimeSyncData.Reset();
-	}
-}
-
-void FLiveLinkSubject::OnSynchronizationEstablished(const struct FTimeSynchronizationStartData& StartData)
+bool FLiveLinkSubject::IsTimeSynchronized() const
 {
-	if (CachedSettings.SourceMode == ELiveLinkSourceMode::TimeSynchronized)
+	if (GetMode() == ELiveLinkSourceMode::Timecode)
 	{
-		if (TimeSyncData.IsSet())
+		const FLiveLinkSubjectFrameData& Snapshot = GetFrameSnapshot();
+		if (Snapshot.StaticData.IsValid() && Snapshot.FrameData.IsValid() && Snapshot.FrameData.GetBaseData())
 		{
-			TimeSyncData->SyncStartTime = StartData.StartFrame;
-			TimeSyncData->bHasEstablishedSync = true;
-
-			// Prevent buffers from being deleted if new data is pushed before we build snapshots.
-			LastReadTime = 0.f;
-			LastReadFrame = 0.f;
+			const FFrameNumber FrameDataInEngineFrameNumber = Snapshot.FrameData.GetBaseData()->MetaData.SceneTime.ConvertTo(FApp::GetTimecodeFrameRate()).GetFrame();
+			const FFrameNumber CurrentEngineFrameNumber = FApp::GetTimecode().ToFrameNumber(FApp::GetTimecodeFrameRate());
+			return FrameDataInEngineFrameNumber == CurrentEngineFrameNumber;
 		}
-		else
-		{
-			UE_LOG(LogLiveLink, Warning, TEXT("OnSynchronizationEstablished called with invalid TimeSyncData. Subject may have switched modes or been recreated. %s"), *SubjectKey.SubjectName.ToString());
-		}
-	}
-	else
-	{
-		TimeSyncData.Reset();
 	}
-}
-
-void FLiveLinkSubject::OnStopSynchronization()
-{
-	TimeSyncData.Reset();
+	return false;
 }
diff --git a/Engine/Plugins/Animation/LiveLink/Source/LiveLink/Private/LiveLinkSubject.h b/Engine/Plugins/Animation/LiveLink/Source/LiveLink/Private/LiveLinkSubject.h
--- a/Engine/Plugins/Animation/LiveLink/Source/LiveLink/Private/LiveLinkSubject.h
+++ b/Engine/Plugins/Animation/LiveLink/Source/LiveLink/Private/LiveLinkSubject.h
@@ -77,40 +77,31 @@ public:
 
 	ELiveLinkSourceMode GetMode() const { return CachedSettings.SourceMode; }
 	FLiveLinkSubjectTimeSyncData GetTimeSyncData();
+	bool IsTimeSynchronized() const;
 
-	void OnStartSynchronization(const struct FTimeSynchronizationOpenData& OpenData, const int32 FrameOffset);
-	void OnSynchronizationEstablished(const struct FTimeSynchronizationStartData& StartData);
-	void OnStopSynchronization();
-
-	bool IsTimeSynchronized() const { return (GetMode() == ELiveLinkSourceMode::TimeSynchronized  && TimeSyncData.IsSet()); }
 	double GetLastPushTime() const { return LastPushTime; }
 
 private:
-	int32 AddFrame_Default(const FLiveLinkWorldTime& FrameTime);
-	int32 AddFrame_Interpolated(const FLiveLinkWorldTime& FrameTime);
-	int32 AddFrame_TimeSynchronized(const FFrameTime& FrameTime, FLiveLinkTimeSynchronizationData& InSynchronizationData);
+	int32 FindNewFrame_WorldTime(const FLiveLinkWorldTime& FrameTime) const;
+	int32 FindNewFrame_SceneTime(const FQualifiedFrameTime& FrameTime) const;
+	int32 FindNewFrame_Latest() const;
 
-	template<bool bWithRollover>
-	int32 AddFrame_TimeSynchronized(const FFrameTime& FrameTime, FLiveLinkTimeSynchronizationData& InSynchronizationData);
 
-	void GetFrameAtWorldTime_Default(const double InSeconds, FLiveLinkSubjectFrameData& OutFrame);
-	void GetFrameAtWorldTime_Interpolated(const double InSeconds, FLiveLinkSubjectFrameData& OutFrame);
-	void GetFrameAtWorldTime_Closest(const double InSeconds, FLiveLinkSubjectFrameData& OutFrame);
+	// Populate OutFrame with a frame based off of the supplied time and our own offsets
+	bool GetFrameAtWorldTime(const double InSeconds, FLiveLinkSubjectFrameData& OutFrame);
+	bool GetFrameAtWorldTime_Closest(const double InSeconds, FLiveLinkSubjectFrameData& OutFrame);
+	bool GetFrameAtWorldTime_Interpolated(const double InSeconds, FLiveLinkSubjectFrameData& OutFrame);
 
-	template<bool bWithRollover>
-	void GetFrameAtSceneTime_TimeSynchronized(const FFrameTime& FrameTime, FLiveLinkTimeSynchronizationData& InSynchronizationData, FLiveLinkSubjectFrameData& OutFrame);
+	// Populate OutFrame with a frame based off of the supplied scene time.
+	bool GetFrameAtSceneTime(const FQualifiedFrameTime& InSceneTime, FLiveLinkSubjectFrameData& OutFrame);
+	bool GetFrameAtSceneTime_Closest(const FQualifiedFrameTime& InSceneTime, FLiveLinkSubjectFrameData& OutFrame);
+	bool GetFrameAtSceneTime_Interpolated(const FQualifiedFrameTime& InSceneTime, FLiveLinkSubjectFrameData& OutFrame);
 
-	template<bool bForInsert, bool bWithRollover>
-	int32 FindFrameIndex_TimeSynchronized(const FFrameTime& InFrameTime, FLiveLinkTimeSynchronizationData& InSynchronizationData);
+	// Populate OutFrame with the latest frame.
+	bool GetLatestFrame(FLiveLinkSubjectFrameData& OutFrame);
 
 	void ResetFrame(FLiveLinkSubjectFrameData& OutFrame) const;
 
-	// Populate OutFrame with a frame based off of the supplied time and our own offsets
-	void GetFrameAtWorldTime(const double InSeconds, FLiveLinkSubjectFrameData& OutFrame);
-
-	// Populate OutFrame with a frame based off of the supplied scene time.
-	void GetFrameAtSceneTime(const FQualifiedFrameTime& InSceneTime, FLiveLinkSubjectFrameData& OutFrame);
-
 protected:
 	// The role the subject was build with
 	TSubclassOf<ULiveLinkRole> Role;
@@ -125,9 +116,8 @@ protected:
 private:
 	struct FLiveLinkCachedSettings
 	{
-		ELiveLinkSourceMode SourceMode = ELiveLinkSourceMode::Default;
-		TOptional<FLiveLinkInterpolationSettings> InterpolationSettings;
-		TOptional<FLiveLinkTimeSynchronizationSettings> TimeSynchronizationSettings;
+		ELiveLinkSourceMode SourceMode = ELiveLinkSourceMode::EngineTime;
+		FLiveLinkSourceBufferManagementSettings BufferSettings;
 	};
 
 	// Static data of the subject
@@ -143,24 +133,16 @@ private:
 	FLiveLinkSubjectKey SubjectKey;
 
 	// Connection settings specified by user
-	// May only store settings relevant to the current mode (ELiveLinkSourceMode).
 	FLiveLinkCachedSettings CachedSettings;
 
-	// Allow us to track changes to the ref skeleton
-	FGuid StaticDataGuid;
-
-	TOptional<FLiveLinkTimeSynchronizationData> TimeSyncData;
-
-	// Time difference between current system time and TimeCode times 
-	double SubjectTimeOffset;
-
-	// Last time we read a frame from this subject. Used to determine whether any new incoming
-	// frames are usable
-	double LastReadTime;
-
 	// Last time a frame was pushed
-	double LastPushTime;
+	double LastPushTime = 0.0;
 
-	//Cache of the last frame we read from, Used for frame cleanup
+	//Cache of the last frame we used to build the snapshot, used to clean frames
 	int32 LastReadFrame;
+
+#if WITH_EDITORONLY_DATA
+	int32 SnapshotIndex = INDEX_NONE;
+	int32 NumberOfBufferAtSnapshot = 0;
+#endif
 };
diff --git a/Engine/Plugins/Animation/LiveLink/Source/LiveLink/Private/TimeSync/LiveLinkTimeSynchronizationSource.cpp b/Engine/Plugins/Animation/LiveLink/Source/LiveLink/Private/TimeSync/LiveLinkTimeSynchronizationSource.cpp
--- a/Engine/Plugins/Animation/LiveLink/Source/LiveLink/Private/TimeSync/LiveLinkTimeSynchronizationSource.cpp
+++ b/Engine/Plugins/Animation/LiveLink/Source/LiveLink/Private/TimeSync/LiveLinkTimeSynchronizationSource.cpp
@@ -35,55 +35,52 @@ FFrameTime ULiveLinkTimeSynchronizationSource::GetOldestSampleTime() const
 FFrameRate ULiveLinkTimeSynchronizationSource::GetFrameRate() const
 {
 	UpdateCachedState();
-	return CachedData.Settings.FrameRate;
+	return CachedData.SampleFrameRate;
 }
 
 bool ULiveLinkTimeSynchronizationSource::IsReady() const
 {
 	UpdateCachedState();
-	return LiveLinkClient && CachedData.bIsValid && (ESyncState::NotSynced == State || LastUpdateGuid == CachedData.SkeletonGuid);
+	return LiveLinkClient && CachedData.bIsValid && IsCurrentStateValid();
 }
 
 bool ULiveLinkTimeSynchronizationSource::Open(const FTimeSynchronizationOpenData& OpenData)
 {
-	UE_LOG(LogLiveLink, Log, TEXT("ULiveLinkTimeSynchronizationSource::Open %s"), *SubjectName.ToString());
-	if (ensure(LiveLinkClient != nullptr) && IsReady())
+	State = ESyncState::NotSynced;
+	SubjectKey = FLiveLinkSubjectKey();
+
+	if (LiveLinkClient == nullptr)
 	{
-		State = ESyncState::Opened;
-		LastUpdateGuid = CachedData.SkeletonGuid;
-		LiveLinkClient->OnStartSynchronization(SubjectName, OpenData, FrameOffset);
-		return true;
+		return false;
 	}
-	else
+
+	TArray<FLiveLinkSubjectKey> AllSubjects = LiveLinkClient->GetSubjects(false, false);
+	FLiveLinkSubjectKey* SubjectKeyPtr = AllSubjects.FindByPredicate([this](const FLiveLinkSubjectKey& SubjectKey) { return SubjectKey.SubjectName == SubjectName; });
+	if (SubjectKeyPtr == nullptr)
 	{
-		State = ESyncState::NotSynced;
+		UE_LOG(LogLiveLink, Error, TEXT("The susbject '%s' is not valid"), *SubjectName.ToString());
 		return false;
 	}
+
+	SubjectKey = *SubjectKeyPtr;
+
+	bool bResult = IsCurrentStateValid();
+
+	if (bResult)
+	{
+		State = ESyncState::Opened;
+	}
+	return bResult;
 }
 
 void ULiveLinkTimeSynchronizationSource::Start(const FTimeSynchronizationStartData& StartData)
 {
-	UE_LOG(LogLiveLink, Log, TEXT("ULiveLinkTimeSynchronizationSource::Start %s"), *SubjectName.ToString());
-	if (ensure(LiveLinkClient != nullptr))
-	{
-		State = ESyncState::Synced;
-		LiveLinkClient->OnSynchronizationEstablished(SubjectName, StartData);
-	}
-	else
-	{
-		State = ESyncState::NotSynced;
-	}
 }
 
 void ULiveLinkTimeSynchronizationSource::Close()
 {
-	UE_LOG(LogLiveLink, Log, TEXT("ULiveLinkTimeSynchronizationSource::Close %s"), *SubjectName.ToString());
-	if (ensure(LiveLinkClient != nullptr))
-	{
-		LiveLinkClient->OnStopSynchronization(SubjectName);
-	}
-
 	State = ESyncState::NotSynced;
+	SubjectKey = FLiveLinkSubjectKey();
 }
 
 FString ULiveLinkTimeSynchronizationSource::GetDisplayName() const
@@ -91,6 +88,42 @@ FString ULiveLinkTimeSynchronizationSource::GetDisplayName() const
 	return SubjectName.ToString();
 }
 
+bool ULiveLinkTimeSynchronizationSource::IsCurrentStateValid() const
+{
+	ensure(LiveLinkClient != nullptr);
+	if (LiveLinkClient == nullptr)
+	{
+		return false;
+	}
+
+	if (!LiveLinkClient->IsSubjectEnabled(SubjectKey, false))
+	{
+		UE_LOG(LogLiveLink, Error, TEXT("The subject '%s' is not enabled."), *SubjectName.ToString());
+		return false;
+	}
+
+	if (LiveLinkClient->IsVirtualSubject(SubjectKey))
+	{
+		UE_LOG(LogLiveLink, Error, TEXT("The subject '%s' can't be a virtual subject."), *SubjectName.ToString());
+		return false;
+	}
+
+	ULiveLinkSourceSettings* SourceSettings = LiveLinkClient->GetSourceSettings(SubjectKey.Source);
+	if (SourceSettings == nullptr)
+	{
+		UE_LOG(LogLiveLink, Error, TEXT("The subject '%s' source does not have a source settings."), *SubjectName.ToString());
+		return false;
+	}
+
+	if (SourceSettings->Mode != ELiveLinkSourceMode::Timecode)
+	{
+		UE_LOG(LogLiveLink, Error, TEXT("The subject '%s' source is not in Timecode mode."), *SubjectName.ToString());
+		return false;
+	}
+
+	return true;
+}
+
 void ULiveLinkTimeSynchronizationSource::OnModularFeatureRegistered(const FName& FeatureName, class IModularFeature* Feature)
 {
 	if (FeatureName == ILiveLinkClient::ModularFeatureName)
diff --git a/Engine/Plugins/Animation/LiveLink/Source/LiveLink/Private/TimeSync/LiveLinkTimeSynchronizationSource.h b/Engine/Plugins/Animation/LiveLink/Source/LiveLink/Private/TimeSync/LiveLinkTimeSynchronizationSource.h
--- a/Engine/Plugins/Animation/LiveLink/Source/LiveLink/Private/TimeSync/LiveLinkTimeSynchronizationSource.h
+++ b/Engine/Plugins/Animation/LiveLink/Source/LiveLink/Private/TimeSync/LiveLinkTimeSynchronizationSource.h
@@ -13,24 +13,23 @@ class ULiveLinkTimeSynchronizationSource : public UTimeSynchronizationSource
 {
 	GENERATED_BODY()
 
-private:
-
+public:
 	UPROPERTY(EditAnywhere, Category="LiveLink")
-	FName SubjectName;
+	FLiveLinkSubjectName SubjectName;
 
+private:
 	FLiveLinkClient* LiveLinkClient;
 
 	enum class ESyncState
 	{
 		NotSynced,
 		Opened,
-		Synced
 	};
 
 	mutable ESyncState State = ESyncState::NotSynced;
-	mutable int64 LastUpdateFrame;
 	mutable FLiveLinkSubjectTimeSyncData CachedData;
-	mutable FGuid LastUpdateGuid;
+	mutable int64 LastUpdateFrame;
+	FLiveLinkSubjectKey SubjectKey;
 
 public:
 
@@ -49,6 +48,7 @@ public:
 
 private:
 
+	bool IsCurrentStateValid() const;
 	void OnModularFeatureRegistered(const FName& FeatureName, class IModularFeature* Feature);
 	void OnModularFeatureUnregistered(const FName& FeatureName, class IModularFeature* Feature);
 	void UpdateCachedState() const;
diff --git a/Engine/Plugins/Animation/LiveLink/Source/LiveLink/Public/LiveLinkBlueprintLibrary.h b/Engine/Plugins/Animation/LiveLink/Source/LiveLink/Public/LiveLinkBlueprintLibrary.h
--- a/Engine/Plugins/Animation/LiveLink/Source/LiveLink/Public/LiveLinkBlueprintLibrary.h
+++ b/Engine/Plugins/Animation/LiveLink/Source/LiveLink/Public/LiveLinkBlueprintLibrary.h
@@ -119,17 +119,30 @@ public:
 	static TArray<FLiveLinkSubjectKey> GetLiveLinkSubjects(bool bIncludeDisabledSubject, bool bIncludeDisal);
 
 	/**
-	 * Whether or not a subject from the specific source is the enabled subject
-	 * Only 1 subject with the same name can be enabled
+	 * Whether or not a subject from the specific source is the enabled subject.
+	 * Only 1 subject with the same name can be enabled.
+	 * At the start of the frame, a snapshot of the enabled subjects will be made.
+	 * That snapshot dictate which subject will be used for the duration of that frame.
 	 */
 	UFUNCTION(BlueprintCallable, Category = "LiveLink")
-	bool IsSpecificLiveLinkSubjectEnabled(const FLiveLinkSubjectKey SubjectKey);
+	bool IsSpecificLiveLinkSubjectEnabled(const FLiveLinkSubjectKey SubjectKey, bool bUseSnapshot);
 
-	/** Whether or not the client has a subject with this name enabled */
+	/**
+	 * Whether or not the client has a subject with this name enabled
+	 * Only 1 subject with the same name can be enabled.
+	 * At the start of the frame, a snapshot of the enabled subjects will be made.
+	 * That snapshot dictate which subject will be used for the duration of that frame.
+	 */
 	UFUNCTION(BlueprintCallable, Category = "LiveLink")
 	bool IsLiveLinkSubjectEnabled(const FLiveLinkSubjectName SubjectName);
 
-	/** Set the subject's from a specific source to enabled, disabling the other in the process, or disabled */
+	/** 
+	 * Set the subject's from a specific source to enabled, disabling the other in the process.
+	 * Only 1 subject with the same name can be enabled.
+	 * At the start of the frame, a snapshot of the enabled subjects will be made.
+	 * That snapshot dictate which subject will be used for the duration of that frame.
+	 * SetSubjectEnabled will take effect on the next frame.
+	 */
 	UFUNCTION(BlueprintCallable, Category = "LiveLink")
 	void SetLiveLinkSubjectEnabled(const FLiveLinkSubjectKey SubjectKey, bool bEnabled);
 
diff --git a/Engine/Plugins/Animation/LiveLink/Source/LiveLink/Public/LiveLinkClient.h b/Engine/Plugins/Animation/LiveLink/Source/LiveLink/Public/LiveLinkClient.h
--- a/Engine/Plugins/Animation/LiveLink/Source/LiveLink/Public/LiveLinkClient.h
+++ b/Engine/Plugins/Animation/LiveLink/Source/LiveLink/Public/LiveLinkClient.h
@@ -28,10 +28,9 @@ DECLARE_STATS_GROUP(TEXT("Live Link"), STATGROUP_LiveLink, STATCAT_Advanced);
 struct FLiveLinkSubjectTimeSyncData
 {
 	bool bIsValid = false;
-	FGuid SkeletonGuid;
 	FFrameTime OldestSampleTime;
 	FFrameTime NewestSampleTime;
-	FLiveLinkTimeSynchronizationSettings Settings;
+	FFrameRate SampleFrameRate;
 };
 
 PRAGMA_DISABLE_DEPRECATION_WARNINGS
@@ -63,20 +62,12 @@ protected:
 PRAGMA_ENABLE_DEPRECATION_WARNINGS
 
 
-class LIVELINK_API FLiveLinkClient : public FLiveLinkClient_Base_DEPRECATED, public FTickableGameObject
+class LIVELINK_API FLiveLinkClient : public FLiveLinkClient_Base_DEPRECATED
 {
 public:
 	FLiveLinkClient();
 	virtual ~FLiveLinkClient();
 
-	//~ Begin FTickableGameObject implementation
-	virtual ETickableTickType GetTickableTickType() const override { return ETickableTickType::Always; }
-	virtual bool IsTickableWhenPaused() const override { return true; }
-	virtual bool IsTickableInEditor() const override { return true; }
-	virtual void Tick(float DeltaTime) override;
-	virtual TStatId GetStatId() const override { RETURN_QUICK_DECLARE_CYCLE_STAT(LiveLinkClient, STATGROUP_Tickables); }
-	//~ End FTickableGameObject
-
 	//~ Begin ILiveLinkClient implementation
 	virtual FGuid AddSource(TSharedPtr<ILiveLinkSource> Source) override;
 	virtual bool CreateSource(const FLiveLinkSourcePreset& SourcePreset) override;
@@ -100,7 +91,7 @@ public:
 
 	virtual bool IsSubjectValid(const FLiveLinkSubjectKey& SubjectKey) const override;
 	virtual bool IsSubjectValid(FLiveLinkSubjectName SubjectName) const override;
-	virtual bool IsSubjectEnabled(const FLiveLinkSubjectKey& SubjectKey) const override;
+	virtual bool IsSubjectEnabled(const FLiveLinkSubjectKey& SubjectKey, bool bUseSnapshot) const override;
 	virtual bool IsSubjectEnabled(FLiveLinkSubjectName SubjectName) const override;
 	virtual void SetSubjectEnabled(const FLiveLinkSubjectKey& SubjectKey, bool bEnabled) override;
 	virtual bool IsSubjectTimeSynchronized(const FLiveLinkSubjectKey& SubjectKey) const override;
@@ -126,6 +117,9 @@ public:
 	virtual void UnregisterSubjectFramesHandle(FLiveLinkSubjectName InSubjectName, FDelegateHandle InStaticDataReceivedHandle, FDelegateHandle InFrameDataReceivedHandle) override;
 	//~ End ILiveLinkClient implementation
 
+	/** The tick callback to update the pending work and clear the subject's snapshot*/
+	void Tick();
+
 	/** Remove the specified source from the live link client */
 	void RemoveSource(FGuid InEntryGuid);
 
@@ -151,16 +145,6 @@ public:
 
 	FText GetSourceMachineName(FGuid EntryGuid) const;
 	FText GetSourceStatus(FGuid EntryGuid) const;
-
-	/** Called when time synchronization is starting for a subject. */
-	void OnStartSynchronization(FLiveLinkSubjectName SubjectName, const struct FTimeSynchronizationOpenData& OpenData, const int32 FrameOffset);
-
-	/** Called when time synchronization has been established for a subject. */
-	void OnSynchronizationEstablished(FLiveLinkSubjectName SubjectName, const struct FTimeSynchronizationStartData& StartData);
-
-	/** Called when time synchronization has been stopped for a subject. */
-	void OnStopSynchronization(FLiveLinkSubjectName SubjectName);
-
 	UE_DEPRECATED(4.23, "FLiveLinkClient::GetSourceTypeForEntry is deprecated. Please use GetSourceType instead!")
 	FText GetSourceTypeForEntry(FGuid EntryGuid) const { return GetSourceType(EntryGuid); }
 	UE_DEPRECATED(4.23, "FLiveLinkClient::GetMachineNameForEntry is deprecated. Please use GetSourceMachineName instead!")
diff --git a/Engine/Plugins/Animation/LiveLink/Source/LiveLink/Public/LiveLinkSettings.h b/Engine/Plugins/Animation/LiveLink/Source/LiveLink/Public/LiveLinkSettings.h
--- a/Engine/Plugins/Animation/LiveLink/Source/LiveLink/Public/LiveLinkSettings.h
+++ b/Engine/Plugins/Animation/LiveLink/Source/LiveLink/Public/LiveLinkSettings.h
@@ -67,15 +67,15 @@ public:
 	FDirectoryPath PresetSaveDir;
 
 	/** The refresh frequency of the list of message bus provider (when discovery is requested). */
-	UPROPERTY(config, EditAnywhere, AdvancedDisplay, Category="LiveLink", meta=(ConfigRestartRequired=true))
+	UPROPERTY(config, EditAnywhere, AdvancedDisplay, Category="LiveLink", meta=(ConfigRestartRequired=true, ForceUnits=s))
 	double MessageBusPingRequestFrequency;
 
 	/** The refresh frequency of the heartbeat when a provider didn't send us an updated. */
-	UPROPERTY(config, EditAnywhere, AdvancedDisplay, Category="LiveLink", meta=(ConfigRestartRequired = true))
+	UPROPERTY(config, EditAnywhere, AdvancedDisplay, Category="LiveLink", meta=(ConfigRestartRequired=true, ForceUnits=s))
 	double MessageBusHeartbeatFrequency;
 
 	/** How long we should wait before a provider become unresponsive. */
-	UPROPERTY(config, EditAnywhere, AdvancedDisplay, Category="LiveLink", meta=(ConfigRestartRequired = true))
+	UPROPERTY(config, EditAnywhere, AdvancedDisplay, Category="LiveLink", meta=(ConfigRestartRequired=true, ForceUnits=s))
 	double MessageBusHeartbeatTimeout;
 
 	/**
@@ -84,7 +84,7 @@ public:
 	 * The subject still exist and can still be evaluated.
 	 * An invalid subject is shown as yellow in the LiveLink UI.
 	 */
-	UPROPERTY(config, EditAnywhere, Category = "LiveLink|UI")
+	UPROPERTY(config, EditAnywhere, Category = "LiveLink|UI", meta=(ForceUnits=s))
 	double TimeWithoutFrameToBeConsiderAsInvalid;
 
 	UPROPERTY(config, EditAnywhere, Category = "LiveLink|UI")
diff --git a/Engine/Plugins/Animation/LiveLink/Source/LiveLinkEditor/Private/LiveLinkClientPanel.cpp b/Engine/Plugins/Animation/LiveLink/Source/LiveLinkEditor/Private/LiveLinkClientPanel.cpp
--- a/Engine/Plugins/Animation/LiveLink/Source/LiveLinkEditor/Private/LiveLinkClientPanel.cpp
+++ b/Engine/Plugins/Animation/LiveLink/Source/LiveLinkEditor/Private/LiveLinkClientPanel.cpp
@@ -114,7 +114,7 @@ struct FLiveLinkSubjectUIEntry
 
 	bool IsSubjectEnabled() const
 	{
-		return IsSubject() ? Client->IsSubjectEnabled(SubjectKey) : false;
+		return IsSubject() ? Client->IsSubjectEnabled(SubjectKey, true) : false;
 	}
 
 	bool IsSubjectValid() const
diff --git a/Engine/Plugins/Animation/LiveLink/Source/LiveLinkEditor/Private/LiveLinkTest.cpp b/Engine/Plugins/Animation/LiveLink/Source/LiveLinkEditor/Private/LiveLinkTest.cpp
--- a/Engine/Plugins/Animation/LiveLink/Source/LiveLinkEditor/Private/LiveLinkTest.cpp
+++ b/Engine/Plugins/Animation/LiveLink/Source/LiveLinkEditor/Private/LiveLinkTest.cpp
@@ -60,10 +60,8 @@ bool FLiveLinkInterpolationTest::RunTest(const FString& Parameters)
 	FrameB.MetaData.SceneTime.Rate = FFrameRate(24, 1);
 	FrameA.MetaData.SceneTime.Time.FrameNumber = IntA;
 	FrameB.MetaData.SceneTime.Time.FrameNumber = IntB;
-	FrameA.WorldTime.Offset = 0.0;
-	FrameB.WorldTime.Offset = 0.0;
-	FrameA.WorldTime.Time = 1000.0;
-	FrameB.WorldTime.Time = 4000.0;
+	FrameA.WorldTime = FLiveLinkWorldTime(1000.0, 0.0);
+	FrameB.WorldTime = FLiveLinkWorldTime(4000.0, 0.0);
 
 	FrameA.NotInterpolated = IntA;
 	FrameB.NotInterpolated = IntB;
@@ -124,7 +122,7 @@ bool FLiveLinkInterpolationTest::RunTest(const FString& Parameters)
 
 	//TestEqual(TEXT("MetaData.StringMetaData"), FrameR.MetaData.StringMetaData, FrameA.MetaData.StringMetaData);
 	//TestEqual(TEXT("MetaData.SceneTime"), FrameR.MetaData.SceneTime, FrameA.MetaData.SceneTime);
-	TestEqual(TEXT("WorldTime"), FrameR.WorldTime.Time, 2000.0);
+	TestEqual(TEXT("WorldTime"), FrameR.WorldTime.GetOffsettedTime(), 2000.0);
 
 	TestEqual(TEXT("NotInterpolated"), FrameR.NotInterpolated, FrameA.NotInterpolated);
 	TestEqual(TEXT("SingleVector"), FrameR.SingleVector, VectorR);
diff --git a/Engine/Plugins/Animation/LiveLink/Source/LiveLinkEditor/Private/SLiveLinkDataView.cpp b/Engine/Plugins/Animation/LiveLink/Source/LiveLinkEditor/Private/SLiveLinkDataView.cpp
--- a/Engine/Plugins/Animation/LiveLink/Source/LiveLinkEditor/Private/SLiveLinkDataView.cpp
+++ b/Engine/Plugins/Animation/LiveLink/Source/LiveLinkEditor/Private/SLiveLinkDataView.cpp
@@ -179,7 +179,7 @@ void SLiveLinkDataView::Tick(const FGeometry& AllottedGeometry, const double InC
 	{
 		LastUpdateSeconds = InCurrentTime;
 
-		if (Client->IsSubjectEnabled(SubjectKey))
+		if (Client->IsSubjectEnabled(SubjectKey, true))
 		{
 			TSubclassOf<ULiveLinkRole> SubjectRole = Client->GetSubjectRole(SubjectKey.SubjectName);
 			FLiveLinkSubjectFrameData SubjectData;
diff --git a/Engine/Plugins/Animation/LiveLink/Source/LiveLinkMovieScene/Private/MovieScene/MovieSceneLiveLinkSectionTemplate.cpp b/Engine/Plugins/Animation/LiveLink/Source/LiveLinkMovieScene/Private/MovieScene/MovieSceneLiveLinkSectionTemplate.cpp
--- a/Engine/Plugins/Animation/LiveLink/Source/LiveLinkMovieScene/Private/MovieScene/MovieSceneLiveLinkSectionTemplate.cpp
+++ b/Engine/Plugins/Animation/LiveLink/Source/LiveLinkMovieScene/Private/MovieScene/MovieSceneLiveLinkSectionTemplate.cpp
@@ -31,10 +31,9 @@ static FQualifiedFrameTime ConvertFrameTimeToTimeCodeTime(const FFrameTime& Fram
 
 static FLiveLinkWorldTime ConvertFrameTimeToLiveLinkWorldTime(const FFrameTime& FrameTime, const FFrameRate& FrameRate, const FFrameTime& FrameTimeEqualToWorldFrameTime, const FLiveLinkWorldTime& LiveLinkWorldTime)
 {
-	FLiveLinkWorldTime WorldTime;
 	FFrameTime DiffFrameTime = FrameTime - FrameTimeEqualToWorldFrameTime;
 	double DiffSeconds = FrameRate.AsSeconds(DiffFrameTime);
-	return WorldTime.Time = DiffSeconds + (LiveLinkWorldTime.Time + LiveLinkWorldTime.Offset);
+	return FLiveLinkWorldTime(DiffSeconds + LiveLinkWorldTime.GetOffsettedTime(), 0.0);
 }
 
 struct FMovieSceneLiveLinkSectionTemplatePersistentData : IPersistentEvaluationData
diff --git a/Engine/Plugins/Animation/LiveLink/Source/LiveLinkMovieScene/Private/MovieScene/MovieSceneLiveLinkSource.cpp b/Engine/Plugins/Animation/LiveLink/Source/LiveLinkMovieScene/Private/MovieScene/MovieSceneLiveLinkSource.cpp
--- a/Engine/Plugins/Animation/LiveLink/Source/LiveLinkMovieScene/Private/MovieScene/MovieSceneLiveLinkSource.cpp
+++ b/Engine/Plugins/Animation/LiveLink/Source/LiveLinkMovieScene/Private/MovieScene/MovieSceneLiveLinkSource.cpp
@@ -32,7 +32,7 @@ void FMovieSceneLiveLinkSource::RemoveLiveLinkSource(TSharedPtr<FMovieSceneLiveL
 		ILiveLinkClient* LiveLinkClient = &IModularFeatures::Get().GetModularFeature<ILiveLinkClient>(ILiveLinkClient::ModularFeatureName);
 
 		//If we're still enabled, put back the previous one with our name if it was enabled
-		if (InSource->PreviousSubjectEnabled.Source.IsValid() && LiveLinkClient->IsSubjectEnabled(InSource->SubjectPreset.Key))
+		if (InSource->PreviousSubjectEnabled.Source.IsValid() && LiveLinkClient->IsSubjectEnabled(InSource->SubjectPreset.Key, false))
 		{
 			LiveLinkClient->SetSubjectEnabled(InSource->PreviousSubjectEnabled, true);
 		}
diff --git a/Engine/Plugins/Animation/LiveLink/Source/LiveLinkSequencer/Private/SequencerRecorderSections/MovieSceneLiveLinkTrackRecorder.cpp b/Engine/Plugins/Animation/LiveLink/Source/LiveLinkSequencer/Private/SequencerRecorderSections/MovieSceneLiveLinkTrackRecorder.cpp
--- a/Engine/Plugins/Animation/LiveLink/Source/LiveLinkSequencer/Private/SequencerRecorderSections/MovieSceneLiveLinkTrackRecorder.cpp
+++ b/Engine/Plugins/Animation/LiveLink/Source/LiveLinkSequencer/Private/SequencerRecorderSections/MovieSceneLiveLinkTrackRecorder.cpp
@@ -206,7 +206,7 @@ void UMovieSceneLiveLinkTrackRecorder::RecordSampleImpl(const FQualifiedFrameTim
 				}
 				else
 				{
-					const double Second = Frame.GetBaseData()->WorldTime.Time + Frame.GetBaseData()->WorldTime.Offset - SecondsDiff;
+					const double Second = Frame.GetBaseData()->WorldTime.GetOffsettedTime() - SecondsDiff;
 					FrameNumber = (Second * TickResolution).FloorToFrame();
 				}
 
diff --git a/Engine/Plugins/Media/TimecodeSynchronizer/Source/TimecodeSynchronizer/Private/TimecodeSynchronizer.cpp b/Engine/Plugins/Media/TimecodeSynchronizer/Source/TimecodeSynchronizer/Private/TimecodeSynchronizer.cpp
--- a/Engine/Plugins/Media/TimecodeSynchronizer/Source/TimecodeSynchronizer/Private/TimecodeSynchronizer.cpp
+++ b/Engine/Plugins/Media/TimecodeSynchronizer/Source/TimecodeSynchronizer/Private/TimecodeSynchronizer.cpp
@@ -5,7 +5,6 @@
 
 #include "Engine/Engine.h"
 #include "FixedFrameRateCustomTimeStep.h"
-#include "IMediaModule.h"
 #include "ITimeManagementModule.h"
 #include "Misc/App.h"
 
@@ -197,6 +196,7 @@ UTimecodeSynchronizer::UTimecodeSynchronizer()
 	, MasterSynchronizationSourceIndex(INDEX_NONE)
 	, PreRollingTimecodeMarginOfErrors(4)
 	, PreRollingTimeout(30.f)
+	, bIsTickEnabled(false)
 	, State(ESynchronizationState::None)
 	, StartPreRollingTime(0.0)
 	, bRegistered(false)
@@ -479,22 +479,15 @@ void UTimecodeSynchronizer::Unregister()
 
 void UTimecodeSynchronizer::SetTickEnabled(bool bEnabled)
 {
-	IMediaModule* MediaModule = FModuleManager::LoadModulePtr<IMediaModule>("Media");
-	if (MediaModule == nullptr)
-	{
-		UE_LOG(LogTimecodeSynchronizer, Error, TEXT("The 'Media' module couldn't be loaded"));
-		SwitchState(ESynchronizationState::Error);
-		return;
-	}
+	bIsTickEnabled = bEnabled;
+}
 
-	MediaModule->GetOnTickPreEngineCompleted().RemoveAll(this);
-	if (bEnabled)
-	{
-		MediaModule->GetOnTickPreEngineCompleted().AddUObject(this, &UTimecodeSynchronizer::Tick);
-	}
+bool UTimecodeSynchronizer::IsTickable() const
+{
+	return bIsTickEnabled;
 }
 
-void UTimecodeSynchronizer::Tick()
+void UTimecodeSynchronizer::Tick(float DeltaTime)
 {
 	UpdateSourceStates();
 	CurrentProviderFrameTime = GetProviderFrameTime();
diff --git a/Engine/Plugins/Media/TimecodeSynchronizer/Source/TimecodeSynchronizer/Public/TimecodeSynchronizer.h b/Engine/Plugins/Media/TimecodeSynchronizer/Source/TimecodeSynchronizer/Public/TimecodeSynchronizer.h
--- a/Engine/Plugins/Media/TimecodeSynchronizer/Source/TimecodeSynchronizer/Public/TimecodeSynchronizer.h
+++ b/Engine/Plugins/Media/TimecodeSynchronizer/Source/TimecodeSynchronizer/Public/TimecodeSynchronizer.h
@@ -7,6 +7,7 @@
 
 #include "Delegates/Delegate.h"
 #include "Delegates/DelegateCombinations.h"
+#include "Tickable.h"
 #include "TimeSynchronizationSource.h"
 
 #include "Misc/QualifiedFrameTime.h"
@@ -247,7 +248,7 @@ struct TStructOpsTypeTraits<FTimecodeSynchronizerActiveTimecodedInputSource> : p
  * making sure all sources are ready, determining if sync is possible, etc.
  */
 UCLASS()
-class TIMECODESYNCHRONIZER_API UTimecodeSynchronizer : public UTimecodeProvider
+class TIMECODESYNCHRONIZER_API UTimecodeSynchronizer : public UTimecodeProvider, public FTickableGameObject
 {
 	GENERATED_BODY()
 
@@ -271,6 +272,15 @@ public:
 	virtual void Shutdown(class UEngine* InEngine) override;
 	//~ End TimecodeProvider Interface
 
+	//~ Begin FTickableGameObject implementation
+	virtual ETickableTickType GetTickableTickType() const override { return ETickableTickType::Conditional; }
+	virtual bool IsTickableWhenPaused() const override { return true; }
+	virtual bool IsTickableInEditor() const override { return true; }
+	virtual bool IsTickable() const override;
+	virtual void Tick(float DeltaTime) override;
+	virtual TStatId GetStatId() const override { RETURN_QUICK_DECLARE_CYCLE_STAT(TimecodeSynchronizer, STATGROUP_Tickables); }
+	//~ End FTickableGameObject
+
 public:
 
 	/**
@@ -386,9 +396,6 @@ private:
 	/** Registers asset to MediaModule tick */
 	void SetTickEnabled(bool bEnabled);
 
-	/** Tick method of the asset */
-	void Tick();
-
 	/** Switches on current state and ticks it */
 	void Tick_Switch();
 
@@ -473,7 +480,7 @@ public:
 	bool bUsePreRollingTimeout;
 
 	/** How long to wait for all source to be ready */
-	UPROPERTY(EditAnywhere, Category="Synchronization", Meta=(EditCondition="bUsePreRollingTimeout", ClampMin="0.0"))
+	UPROPERTY(EditAnywhere, Category="Synchronization", Meta=(EditCondition="bUsePreRollingTimeout", ClampMin="0.0", ForceUnits=s))
 	float PreRollingTimeout; 
 
 public:
@@ -530,6 +537,8 @@ private:
 
 private:
 
+	bool bIsTickEnabled;
+
 	int64 LastUpdatedSources = 0;
 
 	/** The actual synchronization state */
diff --git a/Engine/Plugins/Media/TimecodeSynchronizer/Source/TimecodeSynchronizer/TimecodeSynchronizer.Build.cs b/Engine/Plugins/Media/TimecodeSynchronizer/Source/TimecodeSynchronizer/TimecodeSynchronizer.Build.cs
--- a/Engine/Plugins/Media/TimecodeSynchronizer/Source/TimecodeSynchronizer/TimecodeSynchronizer.Build.cs
+++ b/Engine/Plugins/Media/TimecodeSynchronizer/Source/TimecodeSynchronizer/TimecodeSynchronizer.Build.cs
@@ -11,7 +11,6 @@ public class TimecodeSynchronizer : ModuleRules
 				"Core",
 				"CoreUObject",
 				"Engine",
-				"Media",
 				"MediaAssets",
 				"MediaUtils",
 				"TimeManagement",
diff --git a/Engine/Source/Runtime/LiveLinkInterface/Private/LiveLinkSourceSettings.cpp b/Engine/Source/Runtime/LiveLinkInterface/Private/LiveLinkSourceSettings.cpp
--- a/Engine/Source/Runtime/LiveLinkInterface/Private/LiveLinkSourceSettings.cpp
+++ b/Engine/Source/Runtime/LiveLinkInterface/Private/LiveLinkSourceSettings.cpp
@@ -3,20 +3,33 @@
 #include "LiveLinkSourceSettings.h"
 #include "UObject/EnterpriseObjectVersion.h"
 
-PRAGMA_DISABLE_DEPRECATION_WARNINGS
 void ULiveLinkSourceSettings::Serialize(FArchive& Ar)
 {
 	Super::Serialize(Ar);
 
-	// LiveLinkSourceSettings aren't persistently stored by the engine,
-	// but they could have been elsewhere.
-
 	Ar.UsingCustomVersion(FEnterpriseObjectVersion::GUID);
-#if WITH_EDITORONLY_DATA
-	if (Ar.IsLoading() && FEnterpriseObjectVersion::LiveLinkTimeSynchronization > Ar.CustomVer(FEnterpriseObjectVersion::GUID))
+}
+
+#if WITH_EDITOR
+bool ULiveLinkSourceSettings::CanEditChange(const UProperty* InProperty) const
+{
+	if (Super::CanEditChange(InProperty))
 	{
-		Mode = InterpolationSettings.bUseInterpolation_DEPRECATED ? ELiveLinkSourceMode::Interpolated : ELiveLinkSourceMode::Default;
+		if (InProperty->GetFName() == GET_MEMBER_NAME_CHECKED(FLiveLinkSourceBufferManagementSettings, TimecodeFrameOffset)
+			|| InProperty->GetFName() == GET_MEMBER_NAME_CHECKED(FLiveLinkSourceBufferManagementSettings, TimecodeFrameRate)
+			|| InProperty->GetFName() == GET_MEMBER_NAME_CHECKED(FLiveLinkSourceBufferManagementSettings, ValidTimecodeFrame))
+		{
+			return Mode == ELiveLinkSourceMode::Timecode;
+		}
+
+		if (InProperty->GetFName() == GET_MEMBER_NAME_CHECKED(FLiveLinkSourceBufferManagementSettings, ValidEngineTime)
+			|| InProperty->GetFName() == GET_MEMBER_NAME_CHECKED(FLiveLinkSourceBufferManagementSettings, EngineTimeOffset))
+		{
+			return Mode == ELiveLinkSourceMode::EngineTime;
+		}
+
+		return true;
 	}
-#endif
+	return false;
 }
-PRAGMA_ENABLE_DEPRECATION_WARNINGS
\ No newline at end of file
+#endif //WITH_EDITOR
diff --git a/Engine/Source/Runtime/LiveLinkInterface/Private/Roles/LiveLinkAnimationRole.cpp b/Engine/Source/Runtime/LiveLinkInterface/Private/Roles/LiveLinkAnimationRole.cpp
--- a/Engine/Source/Runtime/LiveLinkInterface/Private/Roles/LiveLinkAnimationRole.cpp
+++ b/Engine/Source/Runtime/LiveLinkInterface/Private/Roles/LiveLinkAnimationRole.cpp
@@ -51,6 +51,59 @@ namespace LiveLinkAnimationBlendingUtil
 
 		Blend(PreFrame.Transforms, PostFrame.Transforms, BlendedFrame->Transforms, BlendWeight);
 	}
+
+	template<class TTimeType>
+	void Interpolate(TTimeType InTime, const FLiveLinkStaticDataStruct& InStaticData, const TArray<FLiveLinkFrameDataStruct>& InSourceFrames, FLiveLinkSubjectFrameData& OutBlendedFrame, bool bInInterpolatePropertyValues)
+	{
+		//Validate inputs
+		check(InStaticData.Cast<FLiveLinkSkeletonStaticData>());
+
+		int32 FrameDataIndexA = INDEX_NONE;
+		int32 FrameDataIndexB = INDEX_NONE;
+		if (ULiveLinkAnimationFrameInterpolateProcessor::FLiveLinkAnimationFrameInterpolateProcessorWorker::FindInterpolateIndex(InTime, InSourceFrames, FrameDataIndexA, FrameDataIndexB))
+		{
+			if (FrameDataIndexA == FrameDataIndexB)
+			{
+				// Copy over the frame directly
+				OutBlendedFrame.FrameData.InitializeWith(InSourceFrames[FrameDataIndexA]);
+			}
+			else
+			{
+				//Initialize the output frame for animation. It will be filled during blended values copied
+				OutBlendedFrame.FrameData.InitializeWith(FLiveLinkAnimationFrameData::StaticStruct(), nullptr);
+
+				const FLiveLinkFrameDataStruct& FrameDataA = InSourceFrames[FrameDataIndexA];
+				const FLiveLinkFrameDataStruct& FrameDataB = InSourceFrames[FrameDataIndexB];
+
+				const double BlendWeight = ULiveLinkAnimationFrameInterpolateProcessor::FLiveLinkAnimationFrameInterpolateProcessorWorker::GetBlendFactor(InTime, FrameDataA, FrameDataB);
+				if (FMath::IsNearlyZero(BlendWeight))
+				{
+					OutBlendedFrame.FrameData.InitializeWith(FrameDataA);
+				}
+				else if (FMath::IsNearlyEqual(1.0, BlendWeight))
+				{
+					OutBlendedFrame.FrameData.InitializeWith(FrameDataB);
+				}
+				else
+				{
+					const FLiveLinkAnimationFrameData* AnimationFrameDataPtrA = FrameDataA.Cast<FLiveLinkAnimationFrameData>();
+					const FLiveLinkAnimationFrameData* AnimationFrameDataPtrB = FrameDataB.Cast<FLiveLinkAnimationFrameData>();
+					const FLiveLinkAnimationFrameData* AnimationFrameDataPtrOutput = OutBlendedFrame.FrameData.Cast<FLiveLinkAnimationFrameData>();
+					check(AnimationFrameDataPtrA && AnimationFrameDataPtrB && AnimationFrameDataPtrOutput);
+
+					ULiveLinkAnimationFrameInterpolateProcessor::FLiveLinkAnimationFrameInterpolateProcessorWorker::FGenericInterpolateOptions InterpolationOptions;
+					InterpolationOptions.bCopyClosestFrame = false; // Do not copy all the Transforms
+					InterpolationOptions.bInterpolateInterpProperties = bInInterpolatePropertyValues;
+					ULiveLinkAnimationFrameInterpolateProcessor::FLiveLinkAnimationFrameInterpolateProcessorWorker::GenericInterpolate(BlendWeight, InterpolationOptions, FrameDataA, FrameDataB, OutBlendedFrame.FrameData);
+					LiveLinkAnimationBlendingUtil::CopyFrameDataBlended(*AnimationFrameDataPtrA, *AnimationFrameDataPtrB, BlendWeight, OutBlendedFrame);
+				}
+			}
+		}
+		else if (InSourceFrames.Num())
+		{
+			OutBlendedFrame.FrameData.InitializeWith(InSourceFrames[0].GetStruct(), InSourceFrames[0].GetBaseData());
+		}
+	}
 }
 
 
@@ -106,50 +159,14 @@ TSubclassOf<ULiveLinkRole> ULiveLinkAnimationFrameInterpolateProcessor::FLiveLin
 	return ULiveLinkAnimationRole::StaticClass();
 }
 
-void ULiveLinkAnimationFrameInterpolateProcessor::FLiveLinkAnimationFrameInterpolateProcessorWorker::Interpolate(double InTime, const FLiveLinkStaticDataStruct& InStaticData, const TArray<FLiveLinkFrameDataStruct>& InSourceFrames, FLiveLinkSubjectFrameData& OutBlendedFrame, int32& OutLastFrameIndexUsed)
+void ULiveLinkAnimationFrameInterpolateProcessor::FLiveLinkAnimationFrameInterpolateProcessorWorker::Interpolate(double InTime, const FLiveLinkStaticDataStruct& InStaticData, const TArray<FLiveLinkFrameDataStruct>& InSourceFrames, FLiveLinkSubjectFrameData& OutBlendedFrame)
 {
-	//Validate inputs
-	check(InStaticData.Cast<FLiveLinkSkeletonStaticData>());
-
-	int32 FrameDataIndexA = INDEX_NONE;
-	int32 FrameDataIndexB = INDEX_NONE;
-	if (FindInterpolateIndex(InTime, InSourceFrames, FrameDataIndexA, FrameDataIndexB))
-	{
-		if (FrameDataIndexA == FrameDataIndexB)
-		{
-			OutLastFrameIndexUsed = FrameDataIndexA;
-			// Copy over the frame directly
-			OutBlendedFrame.FrameData.InitializeWith(InSourceFrames[FrameDataIndexA]);
-		}
-		else
-		{
-			OutLastFrameIndexUsed = FrameDataIndexA;
-
-			//Initialize the output frame for animation. It will be filled during blended values copied
-			OutBlendedFrame.FrameData.InitializeWith(FLiveLinkAnimationFrameData::StaticStruct(), nullptr);
-
-			const FLiveLinkFrameDataStruct& FrameDataA = InSourceFrames[FrameDataIndexA];
-			const FLiveLinkFrameDataStruct& FrameDataB = InSourceFrames[FrameDataIndexB];
-
-			const FLiveLinkAnimationFrameData* AnimationFrameDataPtrA = FrameDataA.Cast<FLiveLinkAnimationFrameData>();
-			const FLiveLinkAnimationFrameData* AnimationFrameDataPtrB = FrameDataB.Cast<FLiveLinkAnimationFrameData>();
-			const FLiveLinkAnimationFrameData* AnimationFrameDataPtrOutput = OutBlendedFrame.FrameData.Cast<FLiveLinkAnimationFrameData>();
-			check(AnimationFrameDataPtrA && AnimationFrameDataPtrB && AnimationFrameDataPtrOutput);
-
-			const double BlendWeight = GetBlendFactor(InTime, FrameDataA, FrameDataB);
+	LiveLinkAnimationBlendingUtil::Interpolate(InTime, InStaticData, InSourceFrames, OutBlendedFrame, bInterpolatePropertyValues);
+}
 
-			FGenericInterpolateOptions InterpolationOptions;
-			InterpolationOptions.bCopyClosestFrame = false; // Do not copy all the Transforms
-			InterpolationOptions.bInterpolateInterpProperties = bInterpolatePropertyValues;
-			GenericInterpolate(BlendWeight, InterpolationOptions, FrameDataA, FrameDataB, OutBlendedFrame.FrameData);
-			LiveLinkAnimationBlendingUtil::CopyFrameDataBlended(*AnimationFrameDataPtrA, *AnimationFrameDataPtrB, BlendWeight, OutBlendedFrame);
-		}
-	}
-	else if (InSourceFrames.Num())
-	{
-		OutLastFrameIndexUsed = 0;
-		OutBlendedFrame.FrameData.InitializeWith(InSourceFrames[0].GetStruct(), InSourceFrames[0].GetBaseData());
-	}
+void ULiveLinkAnimationFrameInterpolateProcessor::FLiveLinkAnimationFrameInterpolateProcessorWorker::Interpolate(const FQualifiedFrameTime& InTime, const FLiveLinkStaticDataStruct& InStaticData, const TArray<FLiveLinkFrameDataStruct>& InSourceFrames, FLiveLinkSubjectFrameData& OutBlendedFrame)
+{
+	LiveLinkAnimationBlendingUtil::Interpolate(InTime, InStaticData, InSourceFrames, OutBlendedFrame, bInterpolatePropertyValues);
 }
 
 
diff --git a/Engine/Source/Runtime/LiveLinkInterface/Private/Roles/LiveLinkBasicRole.cpp b/Engine/Source/Runtime/LiveLinkInterface/Private/Roles/LiveLinkBasicRole.cpp
--- a/Engine/Source/Runtime/LiveLinkInterface/Private/Roles/LiveLinkBasicRole.cpp
+++ b/Engine/Source/Runtime/LiveLinkInterface/Private/Roles/LiveLinkBasicRole.cpp
@@ -216,37 +216,57 @@ namespace LiveLinkInterpolation
 			}
 		}
 	}
-}
 
-void ULiveLinkBasicFrameInterpolateProcessor::FLiveLinkBasicFrameInterpolateProcessorWorker::Interpolate(double InTime, const FLiveLinkStaticDataStruct& InStaticData, const TArray<FLiveLinkFrameDataStruct>& InSourceFrames, FLiveLinkSubjectFrameData& OutBlendedFrame, int32& OutLastFrameIndexUsed)
-{
-	int32 FrameDataIndexA = INDEX_NONE;
-	int32 FrameDataIndexB = INDEX_NONE;
-	if (FindInterpolateIndex(InTime, InSourceFrames, FrameDataIndexA, FrameDataIndexB))
+	template<class TTimeType>
+	void Interpolate(TTimeType InTime, const FLiveLinkStaticDataStruct& InStaticData, const TArray<FLiveLinkFrameDataStruct>& InSourceFrames, FLiveLinkSubjectFrameData& OutBlendedFrame, bool bInInterpolatePropertyValues)
 	{
-		if (FrameDataIndexA == FrameDataIndexB)
-		{
-			OutLastFrameIndexUsed = FrameDataIndexA;
-			// Copy over the frame directly
-			OutBlendedFrame.FrameData.InitializeWith(InSourceFrames[FrameDataIndexA]);
-		}
-		else
+		int32 FrameDataIndexA = INDEX_NONE;
+		int32 FrameDataIndexB = INDEX_NONE;
+		if (ULiveLinkBasicFrameInterpolateProcessor::FLiveLinkBasicFrameInterpolateProcessorWorker::FindInterpolateIndex(InTime, InSourceFrames, FrameDataIndexA, FrameDataIndexB))
 		{
-			OutLastFrameIndexUsed = FrameDataIndexA;
+			if (FrameDataIndexA == FrameDataIndexB)
+			{
+				// Copy over the frame directly
+				OutBlendedFrame.FrameData.InitializeWith(InSourceFrames[FrameDataIndexA]);
+			}
+			else
+			{
+				const FLiveLinkFrameDataStruct& FrameDataA = InSourceFrames[FrameDataIndexA];
+				const FLiveLinkFrameDataStruct& FrameDataB = InSourceFrames[FrameDataIndexB];
 
-			const FLiveLinkFrameDataStruct& FrameDataA = InSourceFrames[FrameDataIndexA];
-			const FLiveLinkFrameDataStruct& FrameDataB = InSourceFrames[FrameDataIndexB];
+				const double BlendFactor = ULiveLinkBasicFrameInterpolateProcessor::FLiveLinkBasicFrameInterpolateProcessorWorker::GetBlendFactor(InTime, FrameDataA, FrameDataB);
+				if (FMath::IsNearlyZero(BlendFactor))
+				{
+					OutBlendedFrame.FrameData.InitializeWith(FrameDataA);
+				}
+				else if (FMath::IsNearlyEqual(1.0, BlendFactor))
+				{
+					OutBlendedFrame.FrameData.InitializeWith(FrameDataB);
+				}
+				else
+				{
+					ULiveLinkBasicFrameInterpolateProcessor::FLiveLinkBasicFrameInterpolateProcessorWorker::FGenericInterpolateOptions InterpolationOptions;
+					InterpolationOptions.bInterpolatePropertyValues = bInInterpolatePropertyValues;
 
-			FGenericInterpolateOptions InterpolationOptions;
-			InterpolationOptions.bInterpolatePropertyValues = bInterpolatePropertyValues;
-			GenericInterpolate(GetBlendFactor(InTime, FrameDataA, FrameDataB), InterpolationOptions, FrameDataA, FrameDataB, OutBlendedFrame.FrameData);
+					ULiveLinkBasicFrameInterpolateProcessor::FLiveLinkBasicFrameInterpolateProcessorWorker::GenericInterpolate(BlendFactor, InterpolationOptions, FrameDataA, FrameDataB, OutBlendedFrame.FrameData);
+				}
+			}
+		}
+		else if (InSourceFrames.Num())
+		{
+			OutBlendedFrame.FrameData.InitializeWith(InSourceFrames[0].GetStruct(), InSourceFrames[0].GetBaseData());
 		}
 	}
-	else if (InSourceFrames.Num())
-	{
-		OutLastFrameIndexUsed = 0;
-		OutBlendedFrame.FrameData.InitializeWith(InSourceFrames[0].GetStruct(), InSourceFrames[0].GetBaseData());
-	}
+}
+
+void ULiveLinkBasicFrameInterpolateProcessor::FLiveLinkBasicFrameInterpolateProcessorWorker::Interpolate(double InTime, const FLiveLinkStaticDataStruct& InStaticData, const TArray<FLiveLinkFrameDataStruct>& InSourceFrames, FLiveLinkSubjectFrameData& OutBlendedFrame)
+{
+	LiveLinkInterpolation::Interpolate(InTime, InStaticData, InSourceFrames, OutBlendedFrame, bInterpolatePropertyValues);
+}
+
+void ULiveLinkBasicFrameInterpolateProcessor::FLiveLinkBasicFrameInterpolateProcessorWorker::Interpolate(const FQualifiedFrameTime& InTime, const FLiveLinkStaticDataStruct& InStaticData, const TArray<FLiveLinkFrameDataStruct>& InSourceFrames, FLiveLinkSubjectFrameData& OutBlendedFrame)
+{
+	LiveLinkInterpolation::Interpolate(InTime, InStaticData, InSourceFrames, OutBlendedFrame, bInterpolatePropertyValues);
 }
 
 void ULiveLinkBasicFrameInterpolateProcessor::FLiveLinkBasicFrameInterpolateProcessorWorker::GenericInterpolate(double InBlendWeight, const FGenericInterpolateOptions& Options, const FLiveLinkFrameDataStruct& FrameDataA, const FLiveLinkFrameDataStruct& FrameDataB, FLiveLinkFrameDataStruct& OutBlendedFrameData)
@@ -268,8 +288,7 @@ void ULiveLinkBasicFrameInterpolateProcessor::FLiveLinkBasicFrameInterpolateProc
 		}
 	}
 
-	OutBlendedFrameData.GetBaseData()->WorldTime.Time = FMath::Lerp(FrameDataA.GetBaseData()->WorldTime.Time, FrameDataB.GetBaseData()->WorldTime.Time, InBlendWeight);
-	OutBlendedFrameData.GetBaseData()->WorldTime.Offset = FMath::Lerp(FrameDataA.GetBaseData()->WorldTime.Offset, FrameDataB.GetBaseData()->WorldTime.Offset, InBlendWeight);
+	OutBlendedFrameData.GetBaseData()->WorldTime = FLiveLinkWorldTime(FMath::Lerp(FrameDataA.GetBaseData()->WorldTime.GetOffsettedTime(), FrameDataB.GetBaseData()->WorldTime.GetOffsettedTime(), InBlendWeight), 0.0);
 
 	if (Options.bInterpolatePropertyValues)
 	{
@@ -298,7 +317,16 @@ void ULiveLinkBasicFrameInterpolateProcessor::FLiveLinkBasicFrameInterpolateProc
 
 double ULiveLinkBasicFrameInterpolateProcessor::FLiveLinkBasicFrameInterpolateProcessorWorker::GetBlendFactor(double InTime, const FLiveLinkFrameDataStruct& FrameDataA, const FLiveLinkFrameDataStruct& FrameDataB)
 {
-	return (InTime - FrameDataA.GetBaseData()->WorldTime.Time) / (FrameDataB.GetBaseData()->WorldTime.Time - FrameDataA.GetBaseData()->WorldTime.Time);
+	const double FrameATime = FrameDataA.GetBaseData()->WorldTime.GetOffsettedTime();
+	const double FrameBTime = FrameDataB.GetBaseData()->WorldTime.GetOffsettedTime();
+	return (InTime - FrameATime) / (FrameBTime - FrameATime);
+}
+
+double ULiveLinkBasicFrameInterpolateProcessor::FLiveLinkBasicFrameInterpolateProcessorWorker::GetBlendFactor(FQualifiedFrameTime InTime, const FLiveLinkFrameDataStruct& FrameDataA, const FLiveLinkFrameDataStruct& FrameDataB)
+{
+	const double FrameASeconds = FrameDataA.GetBaseData()->MetaData.SceneTime.AsSeconds();
+	const double FrameBSeconds = FrameDataB.GetBaseData()->MetaData.SceneTime.AsSeconds();
+	return (InTime.AsSeconds() - FrameASeconds) / (FrameBSeconds - FrameASeconds);
 }
 
  bool ULiveLinkBasicFrameInterpolateProcessor::FLiveLinkBasicFrameInterpolateProcessorWorker::FindInterpolateIndex(double InTime, const TArray<FLiveLinkFrameDataStruct>& InSourceFrames, int32& OutFrameIndexA, int32& OutFrameIndexB)
@@ -311,7 +339,7 @@ double ULiveLinkBasicFrameInterpolateProcessor::FLiveLinkBasicFrameInterpolatePr
 	for (int32 FrameIndex = InSourceFrames.Num() - 1; FrameIndex >= 0; --FrameIndex)
 	{
 		const FLiveLinkFrameDataStruct& SourceFrameData = InSourceFrames[FrameIndex];
-		if (SourceFrameData.GetBaseData()->WorldTime.Time < InTime)
+		if (SourceFrameData.GetBaseData()->WorldTime.GetOffsettedTime() < InTime)
 		{
 			if (FrameIndex == InSourceFrames.Num() - 1)
 			{
@@ -332,5 +360,35 @@ double ULiveLinkBasicFrameInterpolateProcessor::FLiveLinkBasicFrameInterpolatePr
 	return false;
 }
 
+ bool ULiveLinkBasicFrameInterpolateProcessor::FLiveLinkBasicFrameInterpolateProcessorWorker::FindInterpolateIndex(FQualifiedFrameTime InTime, const TArray<FLiveLinkFrameDataStruct>& InSourceFrames, int32& OutFrameIndexA, int32& OutFrameIndexB)
+ {
+	 if (InSourceFrames.Num() == 0)
+	 {
+		 return false;
+	 }
+
+	const double InTimeInSeconds = InTime.AsSeconds();
+	for (int32 FrameIndex = InSourceFrames.Num() - 1; FrameIndex >= 0; --FrameIndex)
+	{
+		const FLiveLinkFrameDataStruct& SourceFrameData = InSourceFrames[FrameIndex];
+		if (SourceFrameData.GetBaseData()->MetaData.SceneTime.AsSeconds() < InTimeInSeconds)
+		{
+			if (FrameIndex == InSourceFrames.Num() - 1)
+			{
+				OutFrameIndexA = FrameIndex;
+				OutFrameIndexB = FrameIndex;
+				return true;
+			}
+			else
+			{
+				OutFrameIndexA = FrameIndex;
+				OutFrameIndexB = FrameIndex + 1;
+				return true;
+			}
+		}
+	}
+
+	return false;
+ }
 
 #undef LOCTEXT_NAMESPACE
diff --git a/Engine/Source/Runtime/LiveLinkInterface/Public/ILiveLinkClient.h b/Engine/Source/Runtime/LiveLinkInterface/Public/ILiveLinkClient.h
--- a/Engine/Source/Runtime/LiveLinkInterface/Public/ILiveLinkClient.h
+++ b/Engine/Source/Runtime/LiveLinkInterface/Public/ILiveLinkClient.h
@@ -167,15 +167,28 @@ public:
 	virtual bool IsSubjectValid(FLiveLinkSubjectName SubjectName) const = 0;
 
 	/**
-	 * Whether or not a subject from the specific source is the enabled subject
-	 * Only 1 subject with the same name can be enabled
+	 * Whether or not a subject from the specific source is the enabled subject.
+	 * Only 1 subject with the same name can be enabled.
+	 * At the start of the frame, a snapshot of the enabled subjects will be made.
+	 * That snapshot dictate which subject will be used for the duration of that frame.
 	 */
-	virtual bool IsSubjectEnabled(const FLiveLinkSubjectKey& SubjectKey) const = 0;
+	virtual bool IsSubjectEnabled(const FLiveLinkSubjectKey& SubjectKey, bool bUseSnaphot) const = 0;
 
-	/** Whether or not the client has a subject with this name enabled */
+	/**
+	 * Whether or not the client has a subject with this name enabled
+	 * Only 1 subject with the same name can be enabled.
+	 * At the start of the frame, a snapshot of the enabled subjects will be made.
+	 * That snapshot dictate which subject will be used for the duration of that frame.
+	 */
 	virtual bool IsSubjectEnabled(FLiveLinkSubjectName SubjectName) const = 0;
 
-	/** Set the subject's from a specific source to enabled, disabling the other in the process, or disabled */
+	/**
+	 * Set the subject's from a specific source to enabled, disabling the other in the process.
+	 * Only 1 subject with the same name can be enabled.
+	 * At the start of the frame, a snapshot of the enabled subjects will be made.
+	 * That snapshot dictate which subject will be used for the duration of that frame.
+	 * SetSubjectEnabled will take effect on the next frame.
+	 */
 	virtual void SetSubjectEnabled(const FLiveLinkSubjectKey& SubjectKey, bool bEnabled) = 0;
 	
 	/** Whether or not the subject's data, from a specific source, is time synchronized or not */
@@ -195,31 +208,33 @@ public:
 
 	/** Whether a subject support a particular role, either directly or through a translator */
 	virtual bool DoesSubjectSupportsRole(const FLiveLinkSubjectKey& SubjectKey, TSubclassOf<ULiveLinkRole> SupportedRole) const = 0;
-	
+
 	/**
 	 * Return the evaluated subject snapshot for a specific role.
-	 * A subject could have to go through a translator to output in the desired role.
+	 * The subject may go through a translator to get the desired role's frame data.
+	 * @return True if the snapshot is valid.
 	 * @note This will always return the same value for a specific frame.
+	 * @see ULiveLinkSourceSettings
 	 */
 	virtual bool EvaluateFrame_AnyThread(FLiveLinkSubjectName SubjectName, TSubclassOf<ULiveLinkRole> Role, FLiveLinkSubjectFrameData& OutFrame) = 0;
 
 	/**
 	 * Evaluates a subject for a specific role.
-	 * A subject could have to go through a translator to output in the desired role.
-	 * If it's a virtual subject EvaluateFrame_AnyThread will be return instead.
-	 * The source and subject configurations will affect which frame will be returned.
+	 * The subject may go through a translator to get the desired role's frame data.
+	 * If it's a virtual subject EvaluateFrame_AnyThread will be used instead.
+	 * @return True if a frame data was calculated.
 	 * @note This value is not cached.
-	 * @see ULiveLinkSourceSettings, ULiveLinkSubject
+	 * @see ULiveLinkSourceSettings
 	 */
 	virtual bool EvaluateFrameAtWorldTime_AnyThread(FLiveLinkSubjectName SubjectName, double WorldTime, TSubclassOf<ULiveLinkRole> DesiredRole, FLiveLinkSubjectFrameData& OutFrame) = 0;
 
 	/**
 	 * Evaluates a subject for a specific role at a Timecode.
-	 * A subject could have to go through a translator to output in the desired role.
-	 * If it's a virtual subject EvaluateFrame_AnyThread will be return instead.
-	 * The source and subject configuration will affect the result.
+	 * The subject may go through a translator to get the desired role's frame data.
+	 * If it's a virtual subject EvaluateFrame_AnyThread will be used instead.
+	 * @return True if a frame data was calculated.
 	 * @note This value is not cached.
-	 * @see ULiveLinkSourceSettings, ULiveLinkSubject
+	 * @see ULiveLinkSourceSettings
 	 */
 	virtual bool EvaluateFrameAtSceneTime_AnyThread(FLiveLinkSubjectName SubjectName, const FTimecode& SceneTime, TSubclassOf<ULiveLinkRole> DesiredRole, FLiveLinkSubjectFrameData& OutFrame) = 0;
 
diff --git a/Engine/Source/Runtime/LiveLinkInterface/Public/LiveLinkFrameInterpolationProcessor.h b/Engine/Source/Runtime/LiveLinkInterface/Public/LiveLinkFrameInterpolationProcessor.h
--- a/Engine/Source/Runtime/LiveLinkInterface/Public/LiveLinkFrameInterpolationProcessor.h
+++ b/Engine/Source/Runtime/LiveLinkInterface/Public/LiveLinkFrameInterpolationProcessor.h
@@ -22,7 +22,8 @@ class LIVELINKINTERFACE_API ILiveLinkFrameInterpolationProcessorWorker
 {
 public:
 	virtual TSubclassOf<ULiveLinkRole> GetRole() const = 0;
-	virtual void Interpolate(double InTime, const FLiveLinkStaticDataStruct& InStaticData, const TArray<FLiveLinkFrameDataStruct>& InSourceFrames, FLiveLinkSubjectFrameData& OutBlendedFrame, int32& OutLastFrameIndexUsed) = 0;
+	virtual void Interpolate(double InTime, const FLiveLinkStaticDataStruct& InStaticData, const TArray<FLiveLinkFrameDataStruct>& InSourceFrames, FLiveLinkSubjectFrameData& OutBlendedFrame) = 0;
+	virtual void Interpolate(const FQualifiedFrameTime& InTime, const FLiveLinkStaticDataStruct& InStaticData, const TArray<FLiveLinkFrameDataStruct>& InSourceFrames, FLiveLinkSubjectFrameData& OutBlendedFrame) = 0;
 };
 
 
diff --git a/Engine/Source/Runtime/LiveLinkInterface/Public/LiveLinkSourceSettings.h b/Engine/Source/Runtime/LiveLinkInterface/Public/LiveLinkSourceSettings.h
--- a/Engine/Source/Runtime/LiveLinkInterface/Public/LiveLinkSourceSettings.h
+++ b/Engine/Source/Runtime/LiveLinkInterface/Public/LiveLinkSourceSettings.h
@@ -7,6 +7,8 @@
 #include "UObject/Object.h"
 
 #include "LiveLinkSourceFactory.h"
+#include "LiveLinkTypes.h"
+#include "Misc/FrameNumber.h"
 #include "Misc/FrameRate.h"
 #include "Templates/SubclassOf.h"
 
@@ -15,83 +17,143 @@
 UENUM()
 enum class ELiveLinkSourceMode : uint8
 {
-	//The source will be run in default mode.
-	//This mode will not attempt any type of interpolation, time synchronization,
-	//or other processing.
-	Default,				
-
-	//The source will be run in interpolated mode.
-	//This mode will use FLiveLinkInterpolationSettings and is most useful
-	//when smooth animation is desired.
-	Interpolated,		
-
-	//The source will be run in time synchronized mode.
-	//This mode will use FLiveLinkTimeSynchronizationSettings and is most useful
-	//when sources need to be synchronized with multiple other external inputs
+	//The source will the latest frame available to evaluate its subjects.
+	//This mode will not attempt any type of interpolation or time synchronization.
+	Latest,
+
+	//The source will use the engine's time to evaluate its subjects.
+	//This mode is most useful when smooth animation is desired.
+	EngineTime,
+
+	//The source will use the engine's timecode to evaluate its subjects.
+	//This mode is most useful when sources need to be synchronized with 
+	//multiple other external inputs
 	//(such as video or other time synchronized sources).
-	//Don't use if the engine isn't setup with a Timecode provider.
-	TimeSynchronized,	
+	//Should not be used when the engine isn't setup with a Timecode provider.
+	Timecode,
 };
 
 USTRUCT()
-struct FLiveLinkTimeSynchronizationSettings
+struct FLiveLinkSourceBufferManagementSettings
 {
 	GENERATED_BODY()
 
-	FLiveLinkTimeSynchronizationSettings() = default;
+	/** If the frame is older than ValidTime, remove it from the buffer list (in seconds). */
+	UPROPERTY(EditAnywhere, Category = "Settings", meta=(ForceUnits=s))
+	float ValidEngineTime = 1.0f;
 
-	// The frame rate of the source.
-	// This should be the frame rate the source is "stamped" at, not necessarily the frame rate the source is sending.
-	// The source should supply this whenever possible.
-	UPROPERTY(EditAnywhere, Category = Settings)
-	FFrameRate FrameRate = {60, 1};
+	/** When evaluating with time: how far back from current time should we read the buffer (in seconds) */
+	UPROPERTY(EditAnywhere, Category = "Settings", meta=(ForceUnits=s))
+	float EngineTimeOffset;
 
-	// When evaluating: how far back from current timecode should we read the buffer (in frame number)
-	UPROPERTY(EditAnywhere, Category = Settings)
-	int32 FrameDelay = 0;
+	/** When evaluating with timecode: what is the expected frame rate of the timecode */
+	UPROPERTY(EditAnywhere, Category = "Settings")
+	FFrameRate TimecodeFrameRate = {24, 1};
+
+	/** If the frame timecode is older than ValidTimecodeFrame, remove it from the buffer list (in TimecodeFrameRate). */
+	UPROPERTY(EditAnywhere, Category = "Settings")
+	int32 ValidTimecodeFrame = 30;
+
+	/** When evaluating with timecode: how far back from current timecode should we read the buffer (in TimecodeFrameRate). */
+	UPROPERTY(EditAnywhere, Category = "Settings")
+	int32 TimecodeFrameOffset;
+
+	/** Maximum number of frame to keep in memory. */
+	UPROPERTY(EditAnywhere, Category = "Settings")
+	int32 MaxNumberOfFrameToBuffered = 10;
 };
- 
+
 USTRUCT()
-struct FLiveLinkInterpolationSettings
+struct FLiveLinkSourceDebugInfo
 {
 	GENERATED_BODY()
 
-	FLiveLinkInterpolationSettings() = default;
+	UPROPERTY(VisibleAnywhere, Category = "Debug")
+	FLiveLinkSubjectName SubjectName;
 
-#if WITH_EDITORONLY_DATA
-	//UE_DEPRECATED(4.21, "Please use ULiveLinkSourceSettings::Mode to specify how the source will behave.")
-	UPROPERTY()
-	bool bUseInterpolation_DEPRECATED = false;
-#endif
+	UPROPERTY(VisibleAnywhere, Category = "Debug")
+	int32 SnapshotIndex;
 
-	// When evaluating: how far back from current time should we read the buffer (in seconds)
-	UPROPERTY(EditAnywhere, Category = Settings)
-	float InterpolationOffset = 0.5f;
+	UPROPERTY(VisibleAnywhere, Category = "Debug")
+	int32 NumberOfBufferAtSnapshot;
 };
 
-// Base class for live link source settings (can be replaced by sources themselves) 
+/** Base class for live link source settings (can be replaced by sources themselves) */
 UCLASS()
 class LIVELINKINTERFACE_API ULiveLinkSourceSettings : public UObject
 {
 public:
 	GENERATED_BODY()
 
-	UPROPERTY(EditAnywhere, Category = "Settings")
-	ELiveLinkSourceMode Mode = ELiveLinkSourceMode::Default;
-
-	// Only used when Mode is set to Interpolated.
-	UPROPERTY(EditAnywhere, Category = "Interpolation Settings")
-	FLiveLinkInterpolationSettings InterpolationSettings;
+	/**
+	 * The the subject how to create the frame snapshot.
+	 * @note A client may evaluate manually the subject in a different mode by using EvaluateFrameAtWorldTime or EvaluateFrameAtSceneTime.
+	 */
+	UPROPERTY(EditAnywhere, Category = "Settings", meta=(DisplayName="Evaluation Mode"))
+	ELiveLinkSourceMode Mode = ELiveLinkSourceMode::EngineTime;
 
-	// Only used when Mode is set to TimeSynchronized.
-	UPROPERTY(EditAnywhere, Category = "Time Synchronization Settings")
-	FLiveLinkTimeSynchronizationSettings TimeSynchronizationSettings;
+	/** How the frame buffers are managed. */
+	UPROPERTY(EditAnywhere, Category = "Settings", meta=(ShowOnlyInnerProperties))
+	FLiveLinkSourceBufferManagementSettings BufferSettings;
 
-	UPROPERTY(EditAnywhere, Category = "Settings", AdvancedDisplay)
+	/** Connection information that is needed by the factory to recreate the source from a preset. */
+	UPROPERTY(EditAnywhere, AdvancedDisplay, Category = "Settings")
 	FString ConnectionString;
 
-	UPROPERTY(VisibleAnywhere, Category = "Settings", AdvancedDisplay)
+	/** Factory used to create the source. */
+	UPROPERTY(VisibleAnywhere, AdvancedDisplay, Category = "Settings")
 	TSubclassOf<ULiveLinkSourceFactory> Factory;
 
+#if WITH_EDITORONLY_DATA
+	UPROPERTY(VisibleAnywhere, Category = "Debug", meta=(ShowOnlyInnerProperties))
+	TArray<FLiveLinkSourceDebugInfo> SourceDebugInfos;
+#endif
+
 	virtual void Serialize(FArchive& Ar) override;
+
+#if WITH_EDITOR
+	virtual bool CanEditChange(const UProperty* InProperty) const override;
+#endif
+};
+
+USTRUCT()
+struct
+	UE_DEPRECATED(4.23, "FLiveLinkTimeSynchronizationSettings is now unused.")
+	FLiveLinkTimeSynchronizationSettings
+{
+	GENERATED_BODY()
+
+	FLiveLinkTimeSynchronizationSettings() : FrameRate(60, 1) {}
+
+	/**
+	 * The frame rate of the source.
+	 * This should be the frame rate the source is "stamped" at, not necessarily the frame rate the source is sending.
+	 * The source should supply this whenever possible.
+	 */
+	UPROPERTY(EditAnywhere, Category = Settings)
+	FFrameRate FrameRate;
+
+	/** When evaluating: how far back from current timecode should we read the buffer (in frame number) */
+	UPROPERTY(EditAnywhere, Category = Settings)
+	FFrameNumber FrameOffset;
+};
+ 
+USTRUCT()
+struct 
+	UE_DEPRECATED(4.23, "FLiveLinkInterpolationSettings is now unused.")
+	FLiveLinkInterpolationSettings
+{
+	GENERATED_BODY()
+
+	FLiveLinkInterpolationSettings() 
+		: bUseInterpolation_DEPRECATED(false)
+		, InterpolationOffset(0.5f) 
+	{}
+
+	UPROPERTY()
+	bool bUseInterpolation_DEPRECATED;
+
+	/** When interpolating: how far back from current time should we read the buffer (in seconds) */
+	UPROPERTY(EditAnywhere, Category = Settings)
+	float InterpolationOffset;
 };
diff --git a/Engine/Source/Runtime/LiveLinkInterface/Public/LiveLinkTypes.h b/Engine/Source/Runtime/LiveLinkInterface/Public/LiveLinkTypes.h
--- a/Engine/Source/Runtime/LiveLinkInterface/Public/LiveLinkTypes.h
+++ b/Engine/Source/Runtime/LiveLinkInterface/Public/LiveLinkTypes.h
@@ -85,26 +85,35 @@ struct FLiveLinkWorldTime
 public:
 	GENERATED_BODY()
 
-	// Time for this frame. Used during interpolation. If this goes backwards we will dump already stored frames. 
-	UPROPERTY(meta = (IgnoreForMemberInitializationTest))
-	double Time;
-
-	// Value calculated on create to represent the different between the source time and client time
-	UPROPERTY()
-	double Offset;
-
+public:
 	FLiveLinkWorldTime()
 		: Offset(0.0)
 	{
 		Time = FPlatformTime::Seconds();
-	};
+	}
 
 	FLiveLinkWorldTime(const double InTime)
 		: Time(InTime)
 	{
 		Offset = FPlatformTime::Seconds() - InTime;
-	};
+	}
 
+	explicit FLiveLinkWorldTime(const double InTime, const double InOffset)
+		: Time(InTime)
+		, Offset(InOffset)
+	{
+	}
+
+	double GetOffsettedTime() const { return Time + Offset; }
+
+private:
+	// Time for this frame. Used during interpolation. If this goes backwards we will dump already stored frames.
+	UPROPERTY(meta = (IgnoreForMemberInitializationTest))
+	double Time;
+
+	// Value calculated on create to represent the different between the source time and client time
+	UPROPERTY()
+	double Offset;
 };
 
 
@@ -130,13 +139,15 @@ struct LIVELINKINTERFACE_API FLiveLinkBaseFrameData
 {
 	GENERATED_BODY();
 
-	UPROPERTY()
+	/** Time in seconds the frame was created. */
+	UPROPERTY(VisibleAnywhere, Category="LiveLink")
 	FLiveLinkWorldTime WorldTime;
 
+	/** Frame's metadata. */
 	UPROPERTY(EditAnywhere, BlueprintReadWrite, Category="LiveLink")
 	FLiveLinkMetaData MetaData;
 
-	// Array of float values for each property data of the skeleton data
+	/** Values of the properties defined in the static structure. Use FLiveLinkBaseStaticData.FindPropertyValue to evaluate. */
 	UPROPERTY(EditAnywhere, BlueprintReadWrite, Category="LiveLink")
 	TArray<float> PropertyValues;
 };
@@ -152,7 +163,7 @@ struct LIVELINKINTERFACE_API FLiveLinkBaseStaticData
 {
 	GENERATED_BODY()
 
-	// Names for each curve values that will be sent for each frame
+	/** Names for each curve values that will be sent for each frame */
 	UPROPERTY(EditAnywhere, BlueprintReadWrite, Category="LiveLink")
 	TArray<FName> PropertyNames;
 
diff --git a/Engine/Source/Runtime/LiveLinkInterface/Public/Roles/LiveLinkAnimationRole.h b/Engine/Source/Runtime/LiveLinkInterface/Public/Roles/LiveLinkAnimationRole.h
--- a/Engine/Source/Runtime/LiveLinkInterface/Public/Roles/LiveLinkAnimationRole.h
+++ b/Engine/Source/Runtime/LiveLinkInterface/Public/Roles/LiveLinkAnimationRole.h
@@ -43,7 +43,8 @@ public:
 		FLiveLinkAnimationFrameInterpolateProcessorWorker(bool bInterpolatePropertyValues);
 
 		virtual TSubclassOf<ULiveLinkRole> GetRole() const override;
-		virtual void Interpolate(double InTime, const FLiveLinkStaticDataStruct& InStaticData, const TArray<FLiveLinkFrameDataStruct>& InSourceFrames, FLiveLinkSubjectFrameData& OutBlendedFrame, int32& OutLastFrameIndexUsed) override;
+		virtual void Interpolate(double InTime, const FLiveLinkStaticDataStruct& InStaticData, const TArray<FLiveLinkFrameDataStruct>& InSourceFrames, FLiveLinkSubjectFrameData& OutBlendedFrame) override;
+		virtual void Interpolate(const FQualifiedFrameTime& InTime, const FLiveLinkStaticDataStruct& InStaticData, const TArray<FLiveLinkFrameDataStruct>& InSourceFrames, FLiveLinkSubjectFrameData& OutBlendedFrame) override;
 	};
 
 public:
diff --git a/Engine/Source/Runtime/LiveLinkInterface/Public/Roles/LiveLinkBasicRole.h b/Engine/Source/Runtime/LiveLinkInterface/Public/Roles/LiveLinkBasicRole.h
--- a/Engine/Source/Runtime/LiveLinkInterface/Public/Roles/LiveLinkBasicRole.h
+++ b/Engine/Source/Runtime/LiveLinkInterface/Public/Roles/LiveLinkBasicRole.h
@@ -41,7 +41,8 @@ public:
 		FLiveLinkBasicFrameInterpolateProcessorWorker(bool bInterpolatePropertyValues);
 
 		virtual TSubclassOf<ULiveLinkRole> GetRole() const override;
-		virtual void Interpolate(double InTime, const FLiveLinkStaticDataStruct& InStaticData, const TArray<FLiveLinkFrameDataStruct>& InSourceFrames, FLiveLinkSubjectFrameData& OutBlendedFrame, int32& OutLastFrameIndexUsed) override;
+		virtual void Interpolate(double InTime, const FLiveLinkStaticDataStruct& InStaticData, const TArray<FLiveLinkFrameDataStruct>& InSourceFrames, FLiveLinkSubjectFrameData& OutBlendedFrame) override;
+		virtual void Interpolate(const FQualifiedFrameTime& InTime, const FLiveLinkStaticDataStruct& InStaticData, const TArray<FLiveLinkFrameDataStruct>& InSourceFrames, FLiveLinkSubjectFrameData& OutBlendedFrame) override;
 
 		struct FGenericInterpolateOptions
 		{
@@ -53,7 +54,9 @@ public:
 
 		static void GenericInterpolate(double InBlendFactor, const FGenericInterpolateOptions& Options, const FLiveLinkFrameDataStruct& FrameDataA, const FLiveLinkFrameDataStruct& FrameDataB, FLiveLinkFrameDataStruct& OutBlendedFrame);
 		static double GetBlendFactor(double InTime, const FLiveLinkFrameDataStruct& FrameDataA, const FLiveLinkFrameDataStruct& FrameDataB);
+		static double GetBlendFactor(FQualifiedFrameTime InTime, const FLiveLinkFrameDataStruct& FrameDataA, const FLiveLinkFrameDataStruct& FrameDataB);
 		static bool FindInterpolateIndex(double InTime, const TArray<FLiveLinkFrameDataStruct>& InSourceFrames, int32& OutFrameIndexA, int32& OutFrameIndexB);
+		static bool FindInterpolateIndex(FQualifiedFrameTime InTime, const TArray<FLiveLinkFrameDataStruct>& InSourceFrames, int32& OutFrameIndexA, int32& OutFrameIndexB);
 
 	protected:
 		bool bInterpolatePropertyValues = true;
