Fix shadowed variable in LiveLinkTimeSynchronization

diff --git a/Engine/Plugins/Animation/LiveLink/Source/LiveLink/Private/TimeSync/LiveLinkTimeSynchronizationSource.cpp b/Engine/Plugins/Animation/LiveLink/Source/LiveLink/Private/TimeSync/LiveLinkTimeSynchronizationSource.cpp
--- a/Engine/Plugins/Animation/LiveLink/Source/LiveLink/Private/TimeSync/LiveLinkTimeSynchronizationSource.cpp
+++ b/Engine/Plugins/Animation/LiveLink/Source/LiveLink/Private/TimeSync/LiveLinkTimeSynchronizationSource.cpp
@@ -55,7 +55,7 @@ bool ULiveLinkTimeSynchronizationSource::Open(const FTimeSynchronizationOpenData
 	}
 
 	TArray<FLiveLinkSubjectKey> AllSubjects = LiveLinkClient->GetSubjects(false, false);
-	FLiveLinkSubjectKey* SubjectKeyPtr = AllSubjects.FindByPredicate([this](const FLiveLinkSubjectKey& SubjectKey) { return SubjectKey.SubjectName == SubjectName; });
+	FLiveLinkSubjectKey* SubjectKeyPtr = AllSubjects.FindByPredicate([this](const FLiveLinkSubjectKey& InSubjectKey) { return InSubjectKey.SubjectName == SubjectName; });
 	if (SubjectKeyPtr == nullptr)
 	{
 		UE_LOG(LogLiveLink, Error, TEXT("The susbject '%s' is not valid"), *SubjectName.ToString());
