Add deprecation path for ILiveLinkSource::IsSourceStillValid. Add a new function IsSourceValid

diff --git a/Engine/Plugins/Animation/LiveLink/Config/BaseLiveLink.ini b/Engine/Plugins/Animation/LiveLink/Config/BaseLiveLink.ini
--- a/Engine/Plugins/Animation/LiveLink/Config/BaseLiveLink.ini
+++ b/Engine/Plugins/Animation/LiveLink/Config/BaseLiveLink.ini
@@ -9,11 +9,14 @@
 +ClassRedirects=(OldName="MovieSceneLiveLinkTrackRecorder", NewName="/Script/LiveLinkSequencer.MovieSceneLiveLinkTrackRecorder")
 
 ; LiveLink MovieScene redirects
- +ClassRedirects=(OldName="MovieSceneLiveLinkSection", NewName="/Script/LiveLinkMovieScene.MovieSceneLiveLinkSection")
- +StructRedirects=(OldName="MovieSceneLiveLinkSectionTemplate", NewName="/Script/LiveLinkMovieScene.MovieSceneLiveLinkSectionTemplate")
++ClassRedirects=(OldName="MovieSceneLiveLinkSection", NewName="/Script/LiveLinkMovieScene.MovieSceneLiveLinkSection")
++StructRedirects=(OldName="MovieSceneLiveLinkSectionTemplate", NewName="/Script/LiveLinkMovieScene.MovieSceneLiveLinkSectionTemplate")
 
 ; LiveLink Graph
 +ClassRedirects=(OldName="K2Node_EvaluateLiveLinkFrame", NewName="/Script/LiveLinkGraphNode.K2Node_EvaluateLiveLinkFrameWithSpecificRole")
 +ClassRedirects=(OldName="K2Node_EvaluateLiveLinkFrameAtWorldTime", NewName="/Script/LiveLinkGraphNode.K2Node_EvaluateLiveLinkFrameAtWorldTime")
 +ClassRedirects=(OldName="K2Node_EvaluateLiveLinkFrameAtSceneTime", NewName="/Script/LiveLinkGraphNode.K2Node_EvaluateLiveLinkFrameAtSceneTime")
-+ClassRedirects=(OldName="AnimGraphNode_LiveLinkPose", NewName="/Script/LiveLinkGraphNode.AnimGraphNode_LiveLinkPose")
\ No newline at end of file
++ClassRedirects=(OldName="AnimGraphNode_LiveLinkPose", NewName="/Script/LiveLinkGraphNode.AnimGraphNode_LiveLinkPose")
+
+; LiveLink BP redirects
++FunctionRedirects=(OldName="LiveLinkBlueprintLibrary.IsSourceStillValid",NewName="LiveLinkBlueprintLibrary.IsSourceValid")
\ No newline at end of file
diff --git a/Engine/Plugins/Animation/LiveLink/Source/LiveLink/Private/LiveLinkBlueprintLibrary.cpp b/Engine/Plugins/Animation/LiveLink/Source/LiveLink/Private/LiveLinkBlueprintLibrary.cpp
--- a/Engine/Plugins/Animation/LiveLink/Source/LiveLink/Private/LiveLinkBlueprintLibrary.cpp
+++ b/Engine/Plugins/Animation/LiveLink/Source/LiveLink/Private/LiveLinkBlueprintLibrary.cpp
@@ -92,9 +92,9 @@ void ULiveLinkBlueprintLibrary::GetChildren(UPARAM(ref) FLiveLinkTransform& Live
 	LiveLinkTransform.GetChildren(Children);
 };
 
-bool ULiveLinkBlueprintLibrary::IsSourceStillValid(UPARAM(ref) FLiveLinkSourceHandle& SourceHandle)
+bool ULiveLinkBlueprintLibrary::IsSourceValid(UPARAM(ref) FLiveLinkSourceHandle& SourceHandle)
 {
-	return SourceHandle.SourcePointer.IsValid() && SourceHandle.SourcePointer->IsSourceStillValid();
+	return SourceHandle.SourcePointer.IsValid() && SourceHandle.SourcePointer->IsSourceValid();
 };
 
 bool ULiveLinkBlueprintLibrary::RequestShutdown(UPARAM(ref) FLiveLinkSourceHandle& SourceHandle)
diff --git a/Engine/Plugins/Animation/LiveLink/Source/LiveLink/Private/LiveLinkMessageBusSource.cpp b/Engine/Plugins/Animation/LiveLink/Source/LiveLink/Private/LiveLinkMessageBusSource.cpp
--- a/Engine/Plugins/Animation/LiveLink/Source/LiveLink/Private/LiveLinkMessageBusSource.cpp
+++ b/Engine/Plugins/Animation/LiveLink/Source/LiveLink/Private/LiveLinkMessageBusSource.cpp
@@ -185,7 +185,7 @@ void FLiveLinkMessageBusSource::InternalHandleMessage(const TSharedRef<IMessageC
 	}
 }
 
-bool FLiveLinkMessageBusSource::IsSourceStillValid() const
+bool FLiveLinkMessageBusSource::IsSourceValid() const
 {
 	return ConnectionAddress.IsValid() && bIsValid;
 }
@@ -196,7 +196,7 @@ FText FLiveLinkMessageBusSource::GetSourceStatus() const
 	{
 		return NSLOCTEXT("LiveLinkMessageBusSource", "InvalidConnection", "Waiting for connection");
 	}
-	else if (IsSourceStillValid())
+	else if (IsSourceValid())
 	{
 		return NSLOCTEXT("LiveLinkMessageBusSource", "ActiveStatus", "Active");
 	}
diff --git a/Engine/Plugins/Animation/LiveLink/Source/LiveLink/Private/LiveLinkMessageBusSource.h b/Engine/Plugins/Animation/LiveLink/Source/LiveLink/Private/LiveLinkMessageBusSource.h
--- a/Engine/Plugins/Animation/LiveLink/Source/LiveLink/Private/LiveLinkMessageBusSource.h
+++ b/Engine/Plugins/Animation/LiveLink/Source/LiveLink/Private/LiveLinkMessageBusSource.h
@@ -33,7 +33,7 @@ public:
 	virtual void ReceiveClient(ILiveLinkClient* InClient, FGuid InSourceGuid) override;
 	virtual void Update() override;
 
-	virtual bool IsSourceStillValid() const override;
+	virtual bool IsSourceValid() const override;
 
 	virtual bool RequestSourceShutdown() override;
 
diff --git a/Engine/Plugins/Animation/LiveLink/Source/LiveLink/Private/LiveLinkSourceCollection.cpp b/Engine/Plugins/Animation/LiveLink/Source/LiveLink/Private/LiveLinkSourceCollection.cpp
--- a/Engine/Plugins/Animation/LiveLink/Source/LiveLink/Private/LiveLinkSourceCollection.cpp
+++ b/Engine/Plugins/Animation/LiveLink/Source/LiveLink/Private/LiveLinkSourceCollection.cpp
@@ -12,7 +12,7 @@ struct FLiveLinkVirtualSubjectSource : public ILiveLinkSource
 {
 	virtual bool CanBeDisplayedInUI() const override { return false; }
 	virtual void ReceiveClient(ILiveLinkClient* InClient, FGuid InSourceGuid) override {}
-	virtual bool IsSourceStillValid() const override { return true; }
+	virtual bool IsSourceValid() const override { return true; }
 	virtual bool RequestSourceShutdown() override { return true; }
 
 	virtual FText GetSourceType() const override { return NSLOCTEXT("TempLocTextLiveLink", "LiveLinkVirtualSubjectName", "Virtual Subjects"); }
diff --git a/Engine/Plugins/Animation/LiveLink/Source/LiveLink/Public/LiveLinkBlueprintLibrary.h b/Engine/Plugins/Animation/LiveLink/Source/LiveLink/Public/LiveLinkBlueprintLibrary.h
--- a/Engine/Plugins/Animation/LiveLink/Source/LiveLink/Public/LiveLinkBlueprintLibrary.h
+++ b/Engine/Plugins/Animation/LiveLink/Source/LiveLink/Public/LiveLinkBlueprintLibrary.h
@@ -90,7 +90,7 @@ class LIVELINK_API ULiveLinkBlueprintLibrary : public UBlueprintFunctionLibrary
 
 	// Checks whether the LiveLink Source is valid via its handle
 	UFUNCTION(BlueprintCallable, Category = "LiveLink")
-	static bool IsSourceStillValid(UPARAM(ref) FLiveLinkSourceHandle& SourceHandle);
+	static bool IsSourceValid(UPARAM(ref) FLiveLinkSourceHandle& SourceHandle);
 
 	// Requests the given LiveLink Source to shut down via its handle
 	UFUNCTION(BlueprintCallable, Category = "LiveLink")
diff --git a/Engine/Plugins/Animation/LiveLink/Source/LiveLinkMovieScene/Private/MovieScene/MovieSceneLiveLinkSectionTemplate.cpp b/Engine/Plugins/Animation/LiveLink/Source/LiveLinkMovieScene/Private/MovieScene/MovieSceneLiveLinkSectionTemplate.cpp
--- a/Engine/Plugins/Animation/LiveLink/Source/LiveLinkMovieScene/Private/MovieScene/MovieSceneLiveLinkSectionTemplate.cpp
+++ b/Engine/Plugins/Animation/LiveLink/Source/LiveLinkMovieScene/Private/MovieScene/MovieSceneLiveLinkSectionTemplate.cpp
@@ -255,7 +255,7 @@ void FMovieSceneLiveLinkSectionTemplate::FillFrameInterpolated(const FFrameTime&
 void FMovieSceneLiveLinkSectionTemplate::EvaluateSwept(const FMovieSceneEvaluationOperand& Operand, const FMovieSceneContext& Context, const TRange<FFrameNumber>& SweptRange, const FPersistentEvaluationData& PersistentData, FMovieSceneExecutionTokens& ExecutionTokens) const
 {
 	FMovieSceneLiveLinkSectionTemplatePersistentData* Data = PersistentData.FindSectionData<FMovieSceneLiveLinkSectionTemplatePersistentData>();
-	if (Data && Data->LiveLinkSource.IsValid() && Data->LiveLinkSource->IsSourceStillValid())
+	if (Data && Data->LiveLinkSource.IsValid() && Data->LiveLinkSource->IsSourceValid())
 	{
 		TArray<FLiveLinkFrameDataStruct>  LiveLinkFrameDataArray;
 		GetLiveLinkFrameArray(Context.GetTime(), SweptRange.GetLowerBoundValue(), SweptRange.GetUpperBoundValue(), LiveLinkFrameDataArray, Context.GetFrameRate());
@@ -267,7 +267,7 @@ void FMovieSceneLiveLinkSectionTemplate::EvaluateSwept(const FMovieSceneEvaluati
 void FMovieSceneLiveLinkSectionTemplate::Evaluate(const FMovieSceneEvaluationOperand& Operand, const FMovieSceneContext& Context, const FPersistentEvaluationData& PersistentData, FMovieSceneExecutionTokens& ExecutionTokens) const
 {
 	FMovieSceneLiveLinkSectionTemplatePersistentData* Data = PersistentData.FindSectionData<FMovieSceneLiveLinkSectionTemplatePersistentData>();
-	if (Data && Data->LiveLinkSource.IsValid() && Data->LiveLinkSource->IsSourceStillValid())
+	if (Data && Data->LiveLinkSource.IsValid() && Data->LiveLinkSource->IsSourceValid())
 	{
 		TArray<FLiveLinkFrameDataStruct>  LiveLinkFrameDataArray;
 		FFrameTime FrameTime = Context.GetTime();
@@ -290,7 +290,7 @@ void FMovieSceneLiveLinkSectionTemplate::TearDown(FPersistentEvaluationData& Per
 
 	if (Data && Data->LiveLinkSource.IsValid())
 	{
-		if (Data->LiveLinkSource->IsSourceStillValid())
+		if (Data->LiveLinkSource->IsSourceValid())
 		{
 			FMovieSceneLiveLinkSource::RemoveLiveLinkSource(Data->LiveLinkSource);
 		}
diff --git a/Engine/Plugins/Animation/LiveLink/Source/LiveLinkMovieScene/Private/MovieScene/MovieSceneLiveLinkSource.cpp b/Engine/Plugins/Animation/LiveLink/Source/LiveLinkMovieScene/Private/MovieScene/MovieSceneLiveLinkSource.cpp
--- a/Engine/Plugins/Animation/LiveLink/Source/LiveLinkMovieScene/Private/MovieScene/MovieSceneLiveLinkSource.cpp
+++ b/Engine/Plugins/Animation/LiveLink/Source/LiveLinkMovieScene/Private/MovieScene/MovieSceneLiveLinkSource.cpp
@@ -67,7 +67,7 @@ void FMovieSceneLiveLinkSource::ReceiveClient(ILiveLinkClient* InClient, FGuid I
 	Client->CreateSubject(SubjectPreset);
 }
 
-bool FMovieSceneLiveLinkSource::IsSourceStillValid() const
+bool FMovieSceneLiveLinkSource::IsSourceValid() const
 {
 	return Client != nullptr;
 }
diff --git a/Engine/Plugins/Animation/LiveLink/Source/LiveLinkMovieScene/Private/MovieScene/MovieSceneLiveLinkSource.h b/Engine/Plugins/Animation/LiveLink/Source/LiveLinkMovieScene/Private/MovieScene/MovieSceneLiveLinkSource.h
--- a/Engine/Plugins/Animation/LiveLink/Source/LiveLinkMovieScene/Private/MovieScene/MovieSceneLiveLinkSource.h
+++ b/Engine/Plugins/Animation/LiveLink/Source/LiveLinkMovieScene/Private/MovieScene/MovieSceneLiveLinkSource.h
@@ -22,7 +22,7 @@ public:
 
 	// ILiveLinkSource interface
 	virtual void ReceiveClient(ILiveLinkClient* InClient, FGuid InSourceGuid) override;
-	virtual bool IsSourceStillValid() const override;
+	virtual bool IsSourceValid() const override;
 	virtual bool RequestSourceShutdown() override;
 	virtual FText GetSourceMachineName() const override;
 	virtual FText GetSourceStatus() const override;
diff --git a/Engine/Plugins/Lumin/MagicLeap/Source/MagicLeapHandTracking/Private/MagicLeapHandTracking.h b/Engine/Plugins/Lumin/MagicLeap/Source/MagicLeapHandTracking/Private/MagicLeapHandTracking.h
--- a/Engine/Plugins/Lumin/MagicLeap/Source/MagicLeapHandTracking/Private/MagicLeapHandTracking.h
+++ b/Engine/Plugins/Lumin/MagicLeap/Source/MagicLeapHandTracking/Private/MagicLeapHandTracking.h
@@ -181,7 +181,7 @@ public:
 
 	// ILiveLinkSource interface
 	virtual void ReceiveClient(ILiveLinkClient* InClient, FGuid InSourceGuid) override;
-	virtual bool IsSourceStillValid() const override;
+	virtual bool IsSourceValid() const override;
 	virtual bool RequestSourceShutdown() override;
 	virtual FText GetSourceMachineName() const override;
 	virtual FText GetSourceStatus() const override;
diff --git a/Engine/Plugins/Lumin/MagicLeap/Source/MagicLeapHandTracking/Private/MagicLeapHandTrackingLiveLink.cpp b/Engine/Plugins/Lumin/MagicLeap/Source/MagicLeapHandTracking/Private/MagicLeapHandTrackingLiveLink.cpp
--- a/Engine/Plugins/Lumin/MagicLeap/Source/MagicLeapHandTracking/Private/MagicLeapHandTrackingLiveLink.cpp
+++ b/Engine/Plugins/Lumin/MagicLeap/Source/MagicLeapHandTracking/Private/MagicLeapHandTrackingLiveLink.cpp
@@ -20,7 +20,7 @@ void FMagicLeapHandTracking::ReceiveClient(ILiveLinkClient* InClient, FGuid InSo
 	bNewLiveLinkClient = true;
 }
 
-bool FMagicLeapHandTracking::IsSourceStillValid() const
+bool FMagicLeapHandTracking::IsSourceValid() const
 {
 	return LiveLinkClient != nullptr;
 }
diff --git a/Engine/Plugins/Runtime/AR/Apple/AppleARKitFaceSupport/Source/AppleARKitFaceSupport/Private/AppleARKitLiveLinkSource.cpp b/Engine/Plugins/Runtime/AR/Apple/AppleARKitFaceSupport/Source/AppleARKitFaceSupport/Private/AppleARKitLiveLinkSource.cpp
--- a/Engine/Plugins/Runtime/AR/Apple/AppleARKitFaceSupport/Source/AppleARKitFaceSupport/Private/AppleARKitLiveLinkSource.cpp
+++ b/Engine/Plugins/Runtime/AR/Apple/AppleARKitFaceSupport/Source/AppleARKitFaceSupport/Private/AppleARKitLiveLinkSource.cpp
@@ -118,7 +118,7 @@ void FAppleARKitLiveLinkSource::ReceiveClient(ILiveLinkClient* InClient, FGuid I
 	SourceGuid = InSourceGuid;
 }
 
-bool FAppleARKitLiveLinkSource::IsSourceStillValid() const
+bool FAppleARKitLiveLinkSource::IsSourceValid() const
 {
 	return Client != nullptr;
 }
diff --git a/Engine/Plugins/Runtime/AR/Apple/AppleARKitFaceSupport/Source/AppleARKitFaceSupport/Private/AppleARKitLiveLinkSource.h b/Engine/Plugins/Runtime/AR/Apple/AppleARKitFaceSupport/Source/AppleARKitFaceSupport/Private/AppleARKitLiveLinkSource.h
--- a/Engine/Plugins/Runtime/AR/Apple/AppleARKitFaceSupport/Source/AppleARKitFaceSupport/Private/AppleARKitLiveLinkSource.h
+++ b/Engine/Plugins/Runtime/AR/Apple/AppleARKitFaceSupport/Source/AppleARKitFaceSupport/Private/AppleARKitLiveLinkSource.h
@@ -880,7 +880,7 @@ public:
 private:
 	// ILiveLinkSource interface
 	virtual void ReceiveClient(ILiveLinkClient* InClient, FGuid InSourceGuid) override;
-	virtual bool IsSourceStillValid() const override;
+	virtual bool IsSourceValid() const override;
 	virtual bool RequestSourceShutdown() override;
 	virtual FText GetSourceMachineName() const override;
 	virtual FText GetSourceStatus() const override;
diff --git a/Engine/Plugins/Runtime/WindowsMixedReality/Source/WindowsMixedRealityHandTracking/Private/WindowsMixedRealityHandTracking.h b/Engine/Plugins/Runtime/WindowsMixedReality/Source/WindowsMixedRealityHandTracking/Private/WindowsMixedRealityHandTracking.h
--- a/Engine/Plugins/Runtime/WindowsMixedReality/Source/WindowsMixedRealityHandTracking/Private/WindowsMixedRealityHandTracking.h
+++ b/Engine/Plugins/Runtime/WindowsMixedReality/Source/WindowsMixedRealityHandTracking/Private/WindowsMixedRealityHandTracking.h
@@ -47,7 +47,7 @@ public:
 
 	// ILiveLinkSource interface
 	virtual void ReceiveClient(ILiveLinkClient* InClient, FGuid InSourceGuid) override;
-	virtual bool IsSourceStillValid() const override;
+	virtual bool IsSourceValid() const override;
 	virtual bool RequestSourceShutdown() override;
 	virtual FText GetSourceMachineName() const override;
 	virtual FText GetSourceStatus() const override;
diff --git a/Engine/Plugins/Runtime/WindowsMixedReality/Source/WindowsMixedRealityHandTracking/Private/WindowsMixedRealityHandTrackingLiveLink.cpp b/Engine/Plugins/Runtime/WindowsMixedReality/Source/WindowsMixedRealityHandTracking/Private/WindowsMixedRealityHandTrackingLiveLink.cpp
--- a/Engine/Plugins/Runtime/WindowsMixedReality/Source/WindowsMixedRealityHandTracking/Private/WindowsMixedRealityHandTrackingLiveLink.cpp
+++ b/Engine/Plugins/Runtime/WindowsMixedReality/Source/WindowsMixedRealityHandTracking/Private/WindowsMixedRealityHandTrackingLiveLink.cpp
@@ -28,7 +28,7 @@ void FWindowsMixedRealityHandTracking::ReceiveClient(ILiveLinkClient* InClient,
 	UpdateLiveLink();
 }
 
-bool FWindowsMixedRealityHandTracking::IsSourceStillValid() const
+bool FWindowsMixedRealityHandTracking::IsSourceValid() const
 {
 	return LiveLinkClient != nullptr;
 }
diff --git a/Engine/Source/Runtime/LiveLinkInterface/Public/ILiveLinkSource.h b/Engine/Source/Runtime/LiveLinkInterface/Public/ILiveLinkSource.h
--- a/Engine/Source/Runtime/LiveLinkInterface/Public/ILiveLinkSource.h
+++ b/Engine/Source/Runtime/LiveLinkInterface/Public/ILiveLinkSource.h
@@ -19,13 +19,15 @@ public:
 	virtual void ReceiveClient(ILiveLinkClient* InClient, FGuid InSourceGuid) = 0;
 	virtual void InitializeSettings(ULiveLinkSourceSettings* Settings) {}
 	virtual void Update() {}
+	virtual bool IsSourceValid() const = 0;
+	virtual bool RequestSourceShutdown() = 0;
+
+	UE_DEPRECATED(4.23, "ILiveLinkSource::IsSourceStillValid is deprecated. Please use IsSourceValid instead!")
+	virtual bool IsSourceStillValid() { return IsSourceValid(); }
 
 	// Can this source be displayed in the Source UI list
 	virtual bool CanBeDisplayedInUI() const { return true; }
 
-	virtual bool IsSourceStillValid() const = 0;
-	virtual bool RequestSourceShutdown() = 0;
-
 	virtual FText GetSourceType() const = 0;
 	virtual FText GetSourceMachineName() const = 0;
 	virtual FText GetSourceStatus() const = 0;
