- Fixing undo/redo on livelink tracks

diff --git a/Engine/Plugins/Animation/LiveLink/Source/LiveLinkMovieScene/Private/MovieScene/MovieSceneLiveLinkSection.cpp b/Engine/Plugins/Animation/LiveLink/Source/LiveLinkMovieScene/Private/MovieScene/MovieSceneLiveLinkSection.cpp
--- a/Engine/Plugins/Animation/LiveLink/Source/LiveLinkMovieScene/Private/MovieScene/MovieSceneLiveLinkSection.cpp
+++ b/Engine/Plugins/Animation/LiveLink/Source/LiveLinkMovieScene/Private/MovieScene/MovieSceneLiveLinkSection.cpp
@@ -166,6 +166,28 @@ void UMovieSceneLiveLinkSection::PostLoad()
 	UpdateChannelProxy();
 }
 
+bool UMovieSceneLiveLinkSection::Modify(bool bAlwaysMarkDirty /*= true*/)
+{
+	bool bWasModified = Super::Modify(bAlwaysMarkDirty);
+
+	for (UMovieSceneLiveLinkSubSection* SubSection : SubSections)
+	{
+		bWasModified |= SubSection->Modify(bAlwaysMarkDirty);
+	}
+
+	return bWasModified;
+}
+
+#if WITH_EDITOR
+void UMovieSceneLiveLinkSection::PostEditChangeProperty(struct FPropertyChangedEvent& PropertyChangedEvent)
+{
+	Super::PostEditChangeProperty(PropertyChangedEvent);	
+
+	//Make sure the channel proxy is always in sync with subsections data. They are recreated whenever we undo / redo.
+	UpdateChannelProxy();
+}
+#endif
+
 TArray<TSubclassOf<UMovieSceneLiveLinkSection>> UMovieSceneLiveLinkSection::GetMovieSectionForRole(const TSubclassOf<ULiveLinkRole>& InRoleToSupport)
 {
 	TArray<TSubclassOf<UMovieSceneLiveLinkSection>> Results;
diff --git a/Engine/Plugins/Animation/LiveLink/Source/LiveLinkMovieScene/Public/MovieScene/MovieSceneLiveLinkSection.h b/Engine/Plugins/Animation/LiveLink/Source/LiveLinkMovieScene/Public/MovieScene/MovieSceneLiveLinkSection.h
--- a/Engine/Plugins/Animation/LiveLink/Source/LiveLinkMovieScene/Public/MovieScene/MovieSceneLiveLinkSection.h
+++ b/Engine/Plugins/Animation/LiveLink/Source/LiveLinkMovieScene/Public/MovieScene/MovieSceneLiveLinkSection.h
@@ -55,13 +55,21 @@ public:
 
 	virtual FMovieSceneEvalTemplatePtr CreateSectionTemplate(const UMovieScenePropertyTrack& InTrack) const;
 
-protected:
-
-	virtual int32 GetChannelCount() const;
-
+public:
+	//~ Begin UObject interface
 	virtual void Serialize(FArchive& Ar) override;
 	virtual void PostEditImport() override;
 	virtual void PostLoad() override;
+	virtual bool Modify(bool bAlwaysMarkDirty = true) override;
+
+#if WITH_EDITOR
+	virtual void PostEditChangeProperty(struct FPropertyChangedEvent& PropertyChangedEvent);
+#endif
+	//~ End UObject interface
+
+protected:
+
+	virtual int32 GetChannelCount() const;
 
 private:
 	void ConvertPreRoleData();
