- Fix for crash when shutdowning livelink messagebus source

diff --git a/Engine/Plugins/Animation/LiveLink/Config/BaseLiveLink.ini b/Engine/Plugins/Animation/LiveLink/Config/BaseLiveLink.ini
--- a/Engine/Plugins/Animation/LiveLink/Config/BaseLiveLink.ini
+++ b/Engine/Plugins/Animation/LiveLink/Config/BaseLiveLink.ini
@@ -7,6 +7,7 @@
 ; Plugin/Module redirects
 +ClassRedirects=(OldName="MovieSceneLiveLinkTrack", NewName="/Script/LiveLinkMovieScene.MovieSceneLiveLinkTrack")
 +ClassRedirects=(OldName="MovieSceneLiveLinkTrackRecorder", NewName="/Script/LiveLinkSequencer.MovieSceneLiveLinkTrackRecorder")
++FunctionRedirects=(OldName="/Script/LiveLink.LiveLinkBlueprintLibrary.RequestShutdown", NewName="/Script/LiveLink.LiveLinkBlueprintLibrary.RemoveSource")
 
 ; LiveLink MovieScene redirects
  +ClassRedirects=(OldName="MovieSceneLiveLinkSection", NewName="/Script/LiveLinkMovieScene.MovieSceneLiveLinkSection")
diff --git a/Engine/Plugins/Animation/LiveLink/Source/LiveLink/Private/LiveLinkBlueprintLibrary.cpp b/Engine/Plugins/Animation/LiveLink/Source/LiveLink/Private/LiveLinkBlueprintLibrary.cpp
--- a/Engine/Plugins/Animation/LiveLink/Source/LiveLink/Private/LiveLinkBlueprintLibrary.cpp
+++ b/Engine/Plugins/Animation/LiveLink/Source/LiveLink/Private/LiveLinkBlueprintLibrary.cpp
@@ -97,10 +97,9 @@ bool ULiveLinkBlueprintLibrary::IsSourceStillValid(UPARAM(ref) FLiveLinkSourceHa
 	return SourceHandle.SourcePointer.IsValid() && SourceHandle.SourcePointer->IsSourceStillValid();
 };
 
-bool ULiveLinkBlueprintLibrary::RequestShutdown(UPARAM(ref) FLiveLinkSourceHandle& SourceHandle)
+bool ULiveLinkBlueprintLibrary::RemoveSource(UPARAM(ref) FLiveLinkSourceHandle& SourceHandle)
 {
-	bool bSourceShutdown = SourceHandle.SourcePointer.IsValid() && SourceHandle.SourcePointer->RequestSourceShutdown();
-	if (bSourceShutdown)
+	if (SourceHandle.SourcePointer.IsValid())
 	{
 		IModularFeatures& ModularFeatures = IModularFeatures::Get();
 
@@ -108,9 +107,11 @@ bool ULiveLinkBlueprintLibrary::RequestShutdown(UPARAM(ref) FLiveLinkSourceHandl
 		{
 			ILiveLinkClient* LiveLinkClient = &ModularFeatures.GetModularFeature<ILiveLinkClient>(ILiveLinkClient::ModularFeatureName);
 			LiveLinkClient->RemoveSource(SourceHandle.SourcePointer);
+			return true;
 		}
 	}
-	return bSourceShutdown;
+
+	return false;
 }
 
 FText ULiveLinkBlueprintLibrary::GetSourceStatus(UPARAM(ref) FLiveLinkSourceHandle& SourceHandle)
diff --git a/Engine/Plugins/Animation/LiveLink/Source/LiveLink/Private/LiveLinkMessageBusSource.cpp b/Engine/Plugins/Animation/LiveLink/Source/LiveLink/Private/LiveLinkMessageBusSource.cpp
--- a/Engine/Plugins/Animation/LiveLink/Source/LiveLink/Private/LiveLinkMessageBusSource.cpp
+++ b/Engine/Plugins/Animation/LiveLink/Source/LiveLink/Private/LiveLinkMessageBusSource.cpp
@@ -248,7 +248,10 @@ bool FLiveLinkMessageBusSource::RequestSourceShutdown()
 	HeartbeatEmitter.StopHeartbeat(ConnectionAddress, MessageEndpoint);
 
 	// Disable the Endpoint message handling since the message could keep it alive a bit.
-	MessageEndpoint->Disable();
+	if (MessageEndpoint.IsValid())
+	{
+		MessageEndpoint->Disable();
+	}
 	MessageEndpoint.Reset();
 
 	return true;
diff --git a/Engine/Plugins/Animation/LiveLink/Source/LiveLink/Public/LiveLinkBlueprintLibrary.h b/Engine/Plugins/Animation/LiveLink/Source/LiveLink/Public/LiveLinkBlueprintLibrary.h
--- a/Engine/Plugins/Animation/LiveLink/Source/LiveLink/Public/LiveLinkBlueprintLibrary.h
+++ b/Engine/Plugins/Animation/LiveLink/Source/LiveLink/Public/LiveLinkBlueprintLibrary.h
@@ -94,7 +94,7 @@ class LIVELINK_API ULiveLinkBlueprintLibrary : public UBlueprintFunctionLibrary
 
 	// Requests the given LiveLink Source to shut down via its handle
 	UFUNCTION(BlueprintCallable, Category = "LiveLink")
-	static bool RequestShutdown(UPARAM(ref) FLiveLinkSourceHandle& SourceHandle);
+	static bool RemoveSource(UPARAM(ref) FLiveLinkSourceHandle& SourceHandle);
 
 	// Get the text status of a LiveLink Source via its handle
 	UFUNCTION(BlueprintCallable, Category = "LiveLink")
