LiveLink, change Warning log into a verbose when the subject's role doesn't matches with the requested role.

diff --git a/Engine/Plugins/Animation/LiveLink/Source/LiveLink/Private/LiveLinkClient.cpp b/Engine/Plugins/Animation/LiveLink/Source/LiveLink/Private/LiveLinkClient.cpp
--- a/Engine/Plugins/Animation/LiveLink/Source/LiveLink/Private/LiveLinkClient.cpp
+++ b/Engine/Plugins/Animation/LiveLink/Source/LiveLink/Private/LiveLinkClient.cpp
@@ -811,6 +811,10 @@ bool FLiveLinkClient::EvaluateFrame_AnyThread(FLiveLinkSubjectName InSubjectName
 			return SubjectItem->GetSubject()->EvaluateFrame(InDesiredRole, OutFrame);
 		}
 	}
+	else
+	{
+		UE_LOG(LogLiveLink, Verbose, TEXT("Subject '%s' is not enabled or doesn't exist"), *InSubjectName.ToString());
+	}
 
 	return false;
 }
@@ -836,6 +840,10 @@ bool FLiveLinkClient::EvaluateFrameAtWorldTime_AnyThread(FLiveLinkSubjectName In
 			}
 		}
 	}
+	else
+	{
+		UE_LOG(LogLiveLink, Verbose, TEXT("Subject '%s' is not enabled or doesn't exist"), *InSubjectName.ToString());
+	}
 
 	return false;
 }
@@ -861,6 +869,10 @@ bool FLiveLinkClient::EvaluateFrameAtSceneTime_AnyThread(FLiveLinkSubjectName In
 			}
 		}
 	}
+	else
+	{
+		UE_LOG(LogLiveLink, Verbose, TEXT("Subject '%s' is not enabled or doesn't exist"), *InSubjectName.ToString());
+	}
 
 	return false;
 }
diff --git a/Engine/Plugins/Animation/LiveLink/Source/LiveLink/Private/LiveLinkSubject.cpp b/Engine/Plugins/Animation/LiveLink/Source/LiveLink/Private/LiveLinkSubject.cpp
--- a/Engine/Plugins/Animation/LiveLink/Source/LiveLink/Private/LiveLinkSubject.cpp
+++ b/Engine/Plugins/Animation/LiveLink/Source/LiveLink/Private/LiveLinkSubject.cpp
@@ -142,7 +142,7 @@ bool FLiveLinkSubject::EvaluateFrameAtWorldTime(double InWorldTime, TSubclassOf<
 		}
 		else
 		{
-			UE_LOG(LogLiveLink, Warning, TEXT("Can't evaluate frame for subject %s. Role %s is incompatible with current role %s"), *SubjectKey.SubjectName.ToString(), *InDesiredRole->GetName(), *Role->GetName());
+			UE_LOG(LogLiveLink, Verbose, TEXT("Can't evaluate frame for subject %s. Role %s is incompatible with current role %s"), *SubjectKey.SubjectName.ToString(), *InDesiredRole->GetName(), *Role->GetName());
 		}
 	}
 
@@ -180,7 +180,7 @@ bool FLiveLinkSubject::EvaluateFrameAtSceneTime(const FTimecode& InSceneTime, TS
 		}
 		else
 		{
-			UE_LOG(LogLiveLink, Warning, TEXT("Can't evaluate frame for subject '%s'. Role '%s' is incompatible with current role '%s'"), *SubjectKey.SubjectName.ToString(), *InDesiredRole->GetName(), *Role->GetName());
+			UE_LOG(LogLiveLink, Verbose, TEXT("Can't evaluate frame for subject '%s'. Role '%s' is incompatible with current role '%s'"), *SubjectKey.SubjectName.ToString(), *InDesiredRole->GetName(), *Role->GetName());
 		}
 	}
 
diff --git a/Engine/Source/Runtime/LiveLinkInterface/Private/ILiveLinkSubject.cpp b/Engine/Source/Runtime/LiveLinkInterface/Private/ILiveLinkSubject.cpp
--- a/Engine/Source/Runtime/LiveLinkInterface/Private/ILiveLinkSubject.cpp
+++ b/Engine/Source/Runtime/LiveLinkInterface/Private/ILiveLinkSubject.cpp
@@ -5,7 +5,7 @@
 #include "ILiveLinkClient.h"
 #include "LiveLinkFrameTranslator.h"
 
-DEFINE_LOG_CATEGORY_STATIC(LogLiveLinkSubject, Warning, Warning);
+DEFINE_LOG_CATEGORY_STATIC(LogLiveLinkSubject, Warning, All);
 
 
 bool ILiveLinkSubject::EvaluateFrame(TSubclassOf<ULiveLinkRole> InDesiredRole, FLiveLinkSubjectFrameData& OutFrame)
@@ -40,7 +40,7 @@ bool ILiveLinkSubject::EvaluateFrame(TSubclassOf<ULiveLinkRole> InDesiredRole, F
 	const bool bSuccess = Translate(this, InDesiredRole, GetFrameSnapshot().StaticData, GetFrameSnapshot().FrameData, OutFrame);
 	if (!bSuccess)
 	{
-		UE_LOG(LogLiveLinkSubject, Warning, TEXT("Can't evaluate frame for subject '%s' for incompatible role '%s. Subject has the role '%s' and no translators could work."), *GetSubjectKey().SubjectName.ToString(), *InDesiredRole->GetName(), *Role->GetName());
+		UE_LOG(LogLiveLinkSubject, Verbose, TEXT("Can't evaluate frame for subject '%s' for incompatible role '%s. Subject has the role '%s' and no translators could work."), *GetSubjectKey().SubjectName.ToString(), *InDesiredRole->GetName(), *Role->GetName());
 	}
 
 	return bSuccess;
