#LiveLink Fix crash in AnimNode_LiveLinkPose

diff --git a/Engine/Plugins/Animation/LiveLink/Source/LiveLink/Private/AnimNode_LiveLinkPose.cpp b/Engine/Plugins/Animation/LiveLink/Source/LiveLink/Private/AnimNode_LiveLinkPose.cpp
--- a/Engine/Plugins/Animation/LiveLink/Source/LiveLink/Private/AnimNode_LiveLinkPose.cpp
+++ b/Engine/Plugins/Animation/LiveLink/Source/LiveLink/Private/AnimNode_LiveLinkPose.cpp
@@ -67,7 +67,7 @@ void FAnimNode_LiveLinkPose::Evaluate_AnyThread(FPoseContext& Output)
 	FLiveLinkSubjectFrameData SubjectFrameData;
 
 	TSubclassOf<ULiveLinkRole> SubjectRole = LiveLinkClient->GetSubjectRole(LiveLinkSubjectName);
-	if (SubjectRole->IsChildOf(ULiveLinkAnimationRole::StaticClass()))
+	if (SubjectRole && SubjectRole->IsChildOf(ULiveLinkAnimationRole::StaticClass()))
 	{
 		//Process animation data if the subject is from that type
 		if (LiveLinkClient->EvaluateFrame_AnyThread(LiveLinkSubjectName, ULiveLinkAnimationRole::StaticClass(), SubjectFrameData))
