Fix deprecation path for LiveLink when the number of properties changes from frame to frame

diff --git a/Engine/Plugins/Animation/LiveLink/Source/LiveLink/Private/LiveLinkClient.cpp b/Engine/Plugins/Animation/LiveLink/Source/LiveLink/Private/LiveLinkClient.cpp
--- a/Engine/Plugins/Animation/LiveLink/Source/LiveLink/Private/LiveLinkClient.cpp
+++ b/Engine/Plugins/Animation/LiveLink/Source/LiveLink/Private/LiveLinkClient.cpp
@@ -1056,6 +1056,19 @@ void FLiveLinkClient::ReleaseLock_Deprecation()
 	CollectionAccessCriticalSection.Unlock();
 }
 
+void FLiveLinkClient::ClearFrames_Deprecation(const FLiveLinkSubjectKey& InSubjectKey)
+{
+	FScopeLock Lock(&CollectionAccessCriticalSection);
+
+	if (Collection)
+	{
+		if (FLiveLinkCollectionSubjectItem* SubjectItem = Collection->FindSubject(InSubjectKey))
+		{
+			SubjectItem->GetSubject()->ClearFrames();
+		}
+	}
+}
+
 FLiveLinkSkeletonStaticData* FLiveLinkClient::GetSubjectAnimationStaticData_Deprecation(const FLiveLinkSubjectKey& InSubjectKey)
 {
 	FScopeLock Lock(&CollectionAccessCriticalSection);
@@ -1132,6 +1145,9 @@ void FLiveLinkClient_Base_DEPRECATED::PushSubjectData(FGuid InSourceGuid, FName
 			NumberOfPropertyNames = AnimationStaticData->PropertyNames.Num();
 			if (NumberOfPropertyNames == 0 && InFrameData.CurveElements.Num() > 0)
 			{
+				UE_LOG(LogLiveLink, Warning, TEXT("Upgrade your code. Curve elements count has changed from the previous frame. That will clear the previous frames of that subject."));
+
+				ClearFrames_Deprecation(SubjectKey);
 				UpdateForAnimationStatic(AnimationStaticData->PropertyNames, InFrameData.CurveElements);
 				NumberOfPropertyNames = AnimationStaticData->PropertyNames.Num();
 			}
@@ -1153,6 +1169,10 @@ void FLiveLinkClient_Base_DEPRECATED::PushSubjectData(FGuid InSourceGuid, FName
 	{
 		NewData.PropertyValues[i] = InFrameData.CurveElements[i].CurveValue;
 	}
+	for (int32 i = MaxNumberOfProperties; i < NumberOfPropertyNames; ++i)
+	{
+		NewData.PropertyValues[i] = INFINITY;
+	}
 	PushSubjectFrameData_AnyThread(SubjectKey, MoveTemp(AnimationStruct));
 }
 
diff --git a/Engine/Plugins/Animation/LiveLink/Source/LiveLink/Public/LiveLinkClient.h b/Engine/Plugins/Animation/LiveLink/Source/LiveLink/Public/LiveLinkClient.h
--- a/Engine/Plugins/Animation/LiveLink/Source/LiveLink/Public/LiveLinkClient.h
+++ b/Engine/Plugins/Animation/LiveLink/Source/LiveLink/Public/LiveLinkClient.h
@@ -57,6 +57,7 @@ public:
 protected:
 	virtual void AquireLock_Deprecation() = 0;
 	virtual void ReleaseLock_Deprecation() = 0;
+	virtual void ClearFrames_Deprecation(const FLiveLinkSubjectKey& SubjectKey) = 0;
 	virtual FLiveLinkSkeletonStaticData* GetSubjectAnimationStaticData_Deprecation(const FLiveLinkSubjectKey& SubjectKey) = 0;
 };
 PRAGMA_ENABLE_DEPRECATION_WARNINGS
@@ -164,6 +165,7 @@ protected:
 	//~ Begin FLiveLinkClient_Base_DEPRECATED implementation
 	virtual void AquireLock_Deprecation() override;
 	virtual void ReleaseLock_Deprecation() override;
+	virtual void ClearFrames_Deprecation(const FLiveLinkSubjectKey& SubjectKey) override;
 	virtual FLiveLinkSkeletonStaticData* GetSubjectAnimationStaticData_Deprecation(const FLiveLinkSubjectKey& SubjectKey) override;
 	//~ End FLiveLinkClient_Base_DEPRECATED implementation
 
