Fix unitialized variable in LiveLinkDiscoveryManager. Was causing a crash in shipping.

diff --git a/Engine/Plugins/Animation/LiveLink/Source/LiveLink/Private/ILiveLinkModule.h b/Engine/Plugins/Animation/LiveLink/Source/LiveLink/Private/ILiveLinkModule.h
--- a/Engine/Plugins/Animation/LiveLink/Source/LiveLink/Private/ILiveLinkModule.h
+++ b/Engine/Plugins/Animation/LiveLink/Source/LiveLink/Private/ILiveLinkModule.h
@@ -7,6 +7,7 @@
 #include "Modules/ModuleManager.h"
 
 class FLiveLinkHeartbeatEmitter;
+class FLiveLinkMessageBusDiscoveryManager;
 
 /**
  * Interface for messaging modules.
@@ -32,6 +33,7 @@ public:
 	}
 
 	virtual FLiveLinkHeartbeatEmitter& GetHeartbeatEmitter() = 0;
+	virtual FLiveLinkMessageBusDiscoveryManager& GetMessageBusDiscoveryManager() = 0;
 
 public:
 
diff --git a/Engine/Plugins/Animation/LiveLink/Source/LiveLink/Private/LiveLinkMessageBusDiscoveryManager.cpp b/Engine/Plugins/Animation/LiveLink/Source/LiveLink/Private/LiveLinkMessageBusDiscoveryManager.cpp
--- a/Engine/Plugins/Animation/LiveLink/Source/LiveLink/Private/LiveLinkMessageBusDiscoveryManager.cpp
+++ b/Engine/Plugins/Animation/LiveLink/Source/LiveLink/Private/LiveLinkMessageBusDiscoveryManager.cpp
@@ -13,10 +13,9 @@
 
 #define LL_HEARTBEAT_SLEEP_TIME 1.0f
 
-FLiveLinkMessageBusDiscoveryManager* FLiveLinkMessageBusDiscoveryManager::Instance = nullptr;
-
 FLiveLinkMessageBusDiscoveryManager::FLiveLinkMessageBusDiscoveryManager()
 	: bRunning(true)
+	, Thread(nullptr)
 {
 	PingRequestCounter = 0;
 	PingRequestFrequency = GetDefault<ULiveLinkSettings>()->GetMessageBusPingRequestFrequency();
@@ -24,7 +23,11 @@ FLiveLinkMessageBusDiscoveryManager::FLiveLinkMessageBusDiscoveryManager()
 	MessageEndpoint = FMessageEndpoint::Builder(TEXT("LiveLinkMessageHeartbeatManager"))
 		.Handling<FLiveLinkPongMessage>(this, &FLiveLinkMessageBusDiscoveryManager::HandlePongMessage);
 
-	Thread = FRunnableThread::Create(this, TEXT("MessageBusHeartbeatManager"));
+	bRunning = MessageEndpoint.IsValid();
+	if (bRunning)
+	{
+		Thread = FRunnableThread::Create(this, TEXT("MessageBusHeartbeatManager"));
+	}
 }
 
 FLiveLinkMessageBusDiscoveryManager::~FLiveLinkMessageBusDiscoveryManager()
@@ -33,25 +36,20 @@ FLiveLinkMessageBusDiscoveryManager::~FLiveLinkMessageBusDiscoveryManager()
 		FScopeLock Lock(&SourcesCriticalSection);
 		
 		// Disable the Endpoint message handling since the message could keep it alive a bit.
-		MessageEndpoint->Disable();
-		MessageEndpoint.Reset();
+		if (MessageEndpoint)
+		{
+			MessageEndpoint->Disable();
+			MessageEndpoint.Reset();
+		}
 	}
 
 	Stop();
+
 	if (Thread)
 	{
 		Thread->Kill(true);
 		Thread = nullptr;
 	}
-};
-
-FLiveLinkMessageBusDiscoveryManager* FLiveLinkMessageBusDiscoveryManager::Get()
-{
-	if (Instance == nullptr)
-	{
-		Instance = new FLiveLinkMessageBusDiscoveryManager();
-	}
-	return Instance;
 }
 
 uint32 FLiveLinkMessageBusDiscoveryManager::Run()
@@ -72,7 +70,7 @@ uint32 FLiveLinkMessageBusDiscoveryManager::Run()
 		FPlatformProcess::Sleep(PingRequestFrequency);
 	}
 	return 0;
-};
+}
 
 void FLiveLinkMessageBusDiscoveryManager::Stop()
 {
@@ -96,6 +94,7 @@ void FLiveLinkMessageBusDiscoveryManager::RemoveDiscoveryMessageRequest()
 TArray<FProviderPollResultPtr> FLiveLinkMessageBusDiscoveryManager::GetDiscoveryResults() const
 {
 	FScopeLock Lock(&SourcesCriticalSection);
+
 	return LastProviderPoolResults;
 }
 
diff --git a/Engine/Plugins/Animation/LiveLink/Source/LiveLink/Private/LiveLinkMessageBusDiscoveryManager.h b/Engine/Plugins/Animation/LiveLink/Source/LiveLink/Private/LiveLinkMessageBusDiscoveryManager.h
--- a/Engine/Plugins/Animation/LiveLink/Source/LiveLink/Private/LiveLinkMessageBusDiscoveryManager.h
+++ b/Engine/Plugins/Animation/LiveLink/Source/LiveLink/Private/LiveLinkMessageBusDiscoveryManager.h
@@ -19,9 +19,7 @@ class FLiveLinkMessageBusSource;
 class FLiveLinkMessageBusDiscoveryManager : FRunnable
 {
 public:
-	// Singleton access to the Heartbeat Manager
-	static FLiveLinkMessageBusDiscoveryManager* Get();
-
+	FLiveLinkMessageBusDiscoveryManager();
 	~FLiveLinkMessageBusDiscoveryManager();
 
 	//~ Begin FRunnable interface
@@ -39,13 +37,9 @@ public:
 	bool IsRunning() const;
 
 private:
-	FLiveLinkMessageBusDiscoveryManager();
 	void HandlePongMessage(const FLiveLinkPongMessage& Message, const TSharedRef<IMessageContext, ESPMode::ThreadSafe>& Context);
 
 private:
-	// Singleton instance
-	static FLiveLinkMessageBusDiscoveryManager* Instance;
-
 	// Counter of item that request discovery message
 	TAtomic<int32> PingRequestCounter;
 
diff --git a/Engine/Plugins/Animation/LiveLink/Source/LiveLink/Private/LiveLinkMessageBusSource.cpp b/Engine/Plugins/Animation/LiveLink/Source/LiveLink/Private/LiveLinkMessageBusSource.cpp
--- a/Engine/Plugins/Animation/LiveLink/Source/LiveLink/Private/LiveLinkMessageBusSource.cpp
+++ b/Engine/Plugins/Animation/LiveLink/Source/LiveLink/Private/LiveLinkMessageBusSource.cpp
@@ -42,7 +42,7 @@ void FLiveLinkMessageBusSource::ReceiveClient(ILiveLinkClient* InClient, FGuid I
 	}
 	else
 	{
-		FLiveLinkMessageBusDiscoveryManager::Get()->AddDiscoveryMessageRequest();
+		ILiveLinkModule::Get().GetMessageBusDiscoveryManager().AddDiscoveryMessageRequest();
 		bIsValid = false;
 	}
 
@@ -53,13 +53,14 @@ void FLiveLinkMessageBusSource::Update()
 {
 	if (!ConnectionAddress.IsValid())
 	{
-		for (const FProviderPollResultPtr Result : FLiveLinkMessageBusDiscoveryManager::Get()->GetDiscoveryResults())
+		FLiveLinkMessageBusDiscoveryManager& DiscoveryManager = ILiveLinkModule::Get().GetMessageBusDiscoveryManager();
+		for (const FProviderPollResultPtr Result : DiscoveryManager.GetDiscoveryResults())
 		{
 			if (Client->GetSourceType(SourceGuid).ToString() == Result->Name)
 			{
 				ConnectionAddress = Result->Address;
 				SourceMachineName = FText::FromString(Result->MachineName);
-				FLiveLinkMessageBusDiscoveryManager::Get()->RemoveDiscoveryMessageRequest();
+				DiscoveryManager.RemoveDiscoveryMessageRequest();
 				SendConnectMessage();
 				UpdateConnectionLastActive();
 				break;
@@ -238,10 +239,10 @@ void FLiveLinkMessageBusSource::SendConnectMessage()
 
 bool FLiveLinkMessageBusSource::RequestSourceShutdown()
 {
-	FLiveLinkMessageBusDiscoveryManager* DiscoveryManager = FLiveLinkMessageBusDiscoveryManager::Get();
-	if (DiscoveryManager->IsRunning() && !ConnectionAddress.IsValid())
+	FLiveLinkMessageBusDiscoveryManager& DiscoveryManager = ILiveLinkModule::Get().GetMessageBusDiscoveryManager();
+	if (DiscoveryManager.IsRunning() && !ConnectionAddress.IsValid())
 	{
-		DiscoveryManager->RemoveDiscoveryMessageRequest();
+		DiscoveryManager.RemoveDiscoveryMessageRequest();
 	}
 
 	FLiveLinkHeartbeatEmitter& HeartbeatEmitter = ILiveLinkModule::Get().GetHeartbeatEmitter();
diff --git a/Engine/Plugins/Animation/LiveLink/Source/LiveLink/Private/LiveLinkModule.cpp b/Engine/Plugins/Animation/LiveLink/Source/LiveLink/Private/LiveLinkModule.cpp
--- a/Engine/Plugins/Animation/LiveLink/Source/LiveLink/Private/LiveLinkModule.cpp
+++ b/Engine/Plugins/Animation/LiveLink/Source/LiveLink/Private/LiveLinkModule.cpp
@@ -29,6 +29,7 @@ public:
 		: LiveLinkClient()
 		, LiveLinkMotionController(LiveLinkClient)
 		, HeartbeatEmitter(MakeUnique<FLiveLinkHeartbeatEmitter>())
+		, DiscoveryManager(MakeUnique<FLiveLinkMessageBusDiscoveryManager>())
 	{}
 
 	// IModuleInterface interface
@@ -37,16 +38,14 @@ public:
 	{
 		IModularFeatures::Get().RegisterModularFeature(FLiveLinkClient::ModularFeatureName, &LiveLinkClient);
 		LiveLinkMotionController.RegisterController();
-		// Create a HeartbeatManager Instance
-		FLiveLinkMessageBusDiscoveryManager::Get();
 	}
 
 	virtual void ShutdownModule() override
 	{
 		HeartbeatEmitter->Exit();
+		DiscoveryManager->Stop();
 		LiveLinkMotionController.UnregisterController();
 		IModularFeatures::Get().UnregisterModularFeature(FLiveLinkClient::ModularFeatureName, &LiveLinkClient);
-		delete FLiveLinkMessageBusDiscoveryManager::Get();
 	}
 
 	virtual bool SupportsDynamicReloading() override
@@ -59,8 +58,14 @@ public:
 		return *HeartbeatEmitter;
 	}
 
+	virtual FLiveLinkMessageBusDiscoveryManager& GetMessageBusDiscoveryManager() override
+	{
+		return *DiscoveryManager;
+	}
+
 private:
 	TUniquePtr<FLiveLinkHeartbeatEmitter> HeartbeatEmitter;
+	TUniquePtr<FLiveLinkMessageBusDiscoveryManager> DiscoveryManager;
 };
 
 IMPLEMENT_MODULE(FLiveLinkModule, LiveLink);
diff --git a/Engine/Plugins/Animation/LiveLink/Source/LiveLink/Private/SLiveLinkMessageBusSourceFactory.cpp b/Engine/Plugins/Animation/LiveLink/Source/LiveLink/Private/SLiveLinkMessageBusSourceFactory.cpp
--- a/Engine/Plugins/Animation/LiveLink/Source/LiveLink/Private/SLiveLinkMessageBusSourceFactory.cpp
+++ b/Engine/Plugins/Animation/LiveLink/Source/LiveLink/Private/SLiveLinkMessageBusSourceFactory.cpp
@@ -3,6 +3,7 @@
 #include "SLiveLinkMessageBusSourceFactory.h"
 
 #include "ILiveLinkClient.h"
+#include "ILiveLinkModule.h"
 #include "LiveLinkMessages.h"
 #include "LiveLinkMessageBusFinder.h"
 #include "LiveLinkMessageBusDiscoveryManager.h"
@@ -65,7 +66,10 @@ private:
 
 SLiveLinkMessageBusSourceFactory::~SLiveLinkMessageBusSourceFactory()
 {
-	FLiveLinkMessageBusDiscoveryManager::Get()->RemoveDiscoveryMessageRequest();
+	if (ILiveLinkModule* ModulePtr = FModuleManager::GetModulePtr<ILiveLinkModule>("LiveLink"))
+	{
+		ModulePtr->GetMessageBusDiscoveryManager().RemoveDiscoveryMessageRequest();
+	}
 }
 
 void SLiveLinkMessageBusSourceFactory::Construct(const FArguments& Args)
@@ -74,7 +78,7 @@ void SLiveLinkMessageBusSourceFactory::Construct(const FArguments& Args)
 	OnSourceSelected = Args._OnSourceSelected;
 	LastUIUpdateSeconds = 0;
 
-	FLiveLinkMessageBusDiscoveryManager::Get()->AddDiscoveryMessageRequest();
+	ILiveLinkModule::Get().GetMessageBusDiscoveryManager().AddDiscoveryMessageRequest();
 
 	ChildSlot
 	[
@@ -111,7 +115,8 @@ void SLiveLinkMessageBusSourceFactory::Tick(const FGeometry& AllottedGeometry, c
 	if (FApp::GetCurrentTime() - LastUIUpdateSeconds > 0.5)
 	{
 		PollData.Reset();
-		PollData.Append(FLiveLinkMessageBusDiscoveryManager::Get()->GetDiscoveryResults());
+		PollData.Append(ILiveLinkModule::Get().GetMessageBusDiscoveryManager().GetDiscoveryResults());
+		PollData.Sort([](const FProviderPollResultPtr& A, const FProviderPollResultPtr& B) { return A->Name < B->Name; });
 		ListView->RequestListRefresh();
 	}
 }
