LiveLink: Guard against null StaticData and SubjectPreset.Role

diff --git a/Engine/Plugins/Animation/LiveLink/Source/LiveLinkMovieScene/Private/MovieScene/MovieSceneLiveLinkSectionTemplate.cpp b/Engine/Plugins/Animation/LiveLink/Source/LiveLinkMovieScene/Private/MovieScene/MovieSceneLiveLinkSectionTemplate.cpp
--- a/Engine/Plugins/Animation/LiveLink/Source/LiveLinkMovieScene/Private/MovieScene/MovieSceneLiveLinkSectionTemplate.cpp
+++ b/Engine/Plugins/Animation/LiveLink/Source/LiveLinkMovieScene/Private/MovieScene/MovieSceneLiveLinkSectionTemplate.cpp
@@ -255,7 +255,7 @@ void FMovieSceneLiveLinkSectionTemplate::FillFrameInterpolated(const FFrameTime&
 void FMovieSceneLiveLinkSectionTemplate::EvaluateSwept(const FMovieSceneEvaluationOperand& Operand, const FMovieSceneContext& Context, const TRange<FFrameNumber>& SweptRange, const FPersistentEvaluationData& PersistentData, FMovieSceneExecutionTokens& ExecutionTokens) const
 {
 	FMovieSceneLiveLinkSectionTemplatePersistentData* Data = PersistentData.FindSectionData<FMovieSceneLiveLinkSectionTemplatePersistentData>();
-	if (Data && Data->LiveLinkSource.IsValid() && Data->LiveLinkSource->IsSourceStillValid())
+	if (Data && Data->LiveLinkSource.IsValid() && Data->LiveLinkSource->IsSourceStillValid() && SubjectPreset.Role)
 	{
 		TArray<FLiveLinkFrameDataStruct>  LiveLinkFrameDataArray;
 		GetLiveLinkFrameArray(Context.GetTime(), SweptRange.GetLowerBoundValue(), SweptRange.GetUpperBoundValue(), LiveLinkFrameDataArray, Context.GetFrameRate());
@@ -267,7 +267,7 @@ void FMovieSceneLiveLinkSectionTemplate::EvaluateSwept(const FMovieSceneEvaluati
 void FMovieSceneLiveLinkSectionTemplate::Evaluate(const FMovieSceneEvaluationOperand& Operand, const FMovieSceneContext& Context, const FPersistentEvaluationData& PersistentData, FMovieSceneExecutionTokens& ExecutionTokens) const
 {
 	FMovieSceneLiveLinkSectionTemplatePersistentData* Data = PersistentData.FindSectionData<FMovieSceneLiveLinkSectionTemplatePersistentData>();
-	if (Data && Data->LiveLinkSource.IsValid() && Data->LiveLinkSource->IsSourceStillValid())
+	if (Data && Data->LiveLinkSource.IsValid() && Data->LiveLinkSource->IsSourceStillValid() && SubjectPreset.Role)
 	{
 		TArray<FLiveLinkFrameDataStruct>  LiveLinkFrameDataArray;
 		FFrameTime FrameTime = Context.GetTime();
@@ -279,9 +279,12 @@ void FMovieSceneLiveLinkSectionTemplate::Evaluate(const FMovieSceneEvaluationOpe
 
 void FMovieSceneLiveLinkSectionTemplate::Setup(FPersistentEvaluationData& PersistentData, IMovieScenePlayer& Player) const
 {
-	FMovieSceneLiveLinkSectionTemplatePersistentData& Data = PersistentData.GetOrAddSectionData<FMovieSceneLiveLinkSectionTemplatePersistentData>();
-	Data.LiveLinkSource = FMovieSceneLiveLinkSource::CreateLiveLinkSource(SubjectPreset);
-	Data.LiveLinkSource->PublishLiveLinkStaticData(*StaticData);
+	if (StaticData)
+	{
+		FMovieSceneLiveLinkSectionTemplatePersistentData& Data = PersistentData.GetOrAddSectionData<FMovieSceneLiveLinkSectionTemplatePersistentData>();
+		Data.LiveLinkSource = FMovieSceneLiveLinkSource::CreateLiveLinkSource(SubjectPreset);
+		Data.LiveLinkSource->PublishLiveLinkStaticData(*StaticData);
+	}
 }
 
 void FMovieSceneLiveLinkSectionTemplate::TearDown(FPersistentEvaluationData& PersistentData, IMovieScenePlayer& Player) const
