Merging fix for crash in livelink deprecation path
- Fix for crash when activating / deactivating subjects in source. A discrepancy was possible between static and frame data

diff --git a/Engine/Plugins/Animation/LiveLink/Source/LiveLink/Private/LiveLinkClient.cpp b/Engine/Plugins/Animation/LiveLink/Source/LiveLink/Private/LiveLinkClient.cpp
--- a/Engine/Plugins/Animation/LiveLink/Source/LiveLink/Private/LiveLinkClient.cpp
+++ b/Engine/Plugins/Animation/LiveLink/Source/LiveLink/Private/LiveLinkClient.cpp
@@ -370,6 +370,16 @@ void FLiveLinkClient::PushSubjectStaticData_Internal(FPendingSubjectStatic&& Sub
 			LiveLinkSubject->ClearFrames();
 		}
 	}
+
+	//Clear any pending frame for that subject. This will enforce one frame delay between reception of static data and frame data but will ensure both matches, especially in the case of deprecated path
+	for (int32 Index = SubjectFrameToPush.Num() - 1; Index >= 0; --Index)
+	{
+		FPendingSubjectFrame& PendingFrame = SubjectFrameToPush[Index];
+		if (PendingFrame.SubjectKey == SubjectStaticData.SubjectKey)
+		{
+			SubjectFrameToPush.RemoveAtSwap(Index);
+		}
+	}
 	
 	if(LiveLinkSubject == nullptr)
 	{
