Live Link - Fixing undo/redo on livelink tracks

diff --git a/Engine/Plugins/Animation/LiveLink/Source/LiveLinkMovieScene/Public/MovieScene/MovieSceneLiveLinkSection.h b/Engine/Plugins/Animation/LiveLink/Source/LiveLinkMovieScene/Public/MovieScene/MovieSceneLiveLinkSection.h
--- a/Engine/Plugins/Animation/LiveLink/Source/LiveLinkMovieScene/Public/MovieScene/MovieSceneLiveLinkSection.h
+++ b/Engine/Plugins/Animation/LiveLink/Source/LiveLinkMovieScene/Public/MovieScene/MovieSceneLiveLinkSection.h
@@ -60,10 +60,10 @@ public:
 	virtual void Serialize(FArchive& Ar) override;
 	virtual void PostEditImport() override;
 	virtual void PostLoad() override;
-	virtual bool Modify(bool bAlwaysMarkDirty = true) override;
 
 #if WITH_EDITOR
-	virtual void PostEditChangeProperty(struct FPropertyChangedEvent& PropertyChangedEvent);
+	virtual bool Modify(bool bAlwaysMarkDirty = true) override;
+	virtual void PostEditChangeProperty(struct FPropertyChangedEvent& PropertyChangedEvent) override;
 #endif
 	//~ End UObject interface
 
