Move LiveLinkTimeSynchronizationSource to the public folder.

diff --git a/Engine/Plugins/Animation/LiveLink/Source/LiveLink/Private/LiveLinkTimeSynchronizationSource.cpp b/Engine/Plugins/Animation/LiveLink/Source/LiveLink/Private/LiveLinkTimeSynchronizationSource.cpp
new file mode 100644
--- /dev/null
+++ b/Engine/Plugins/Animation/LiveLink/Source/LiveLink/Private/LiveLinkTimeSynchronizationSource.cpp
@@ -0,0 +1,150 @@
+// Copyright 1998-2019 Epic Games, Inc. All Rights Reserved.
+
+#include "LiveLinkTimeSynchronizationSource.h"
+#include "LiveLinkClient.h"
+#include "Features/IModularFeatures.h"
+#include "Math/NumericLimits.h"
+
+ULiveLinkTimeSynchronizationSource::ULiveLinkTimeSynchronizationSource()
+{
+	if (!HasAnyFlags(RF_ArchetypeObject | RF_ClassDefaultObject))
+	{
+		IModularFeatures& ModularFeatures = IModularFeatures::Get();
+		ModularFeatures.OnModularFeatureRegistered().AddUObject(this, &ThisClass::OnModularFeatureRegistered);
+		ModularFeatures.OnModularFeatureUnregistered().AddUObject(this, &ThisClass::OnModularFeatureUnregistered);
+
+		if (ModularFeatures.IsModularFeatureAvailable(ILiveLinkClient::ModularFeatureName))
+		{
+			LiveLinkClient = &ModularFeatures.GetModularFeature<FLiveLinkClient>(ILiveLinkClient::ModularFeatureName);
+		}
+	}
+}
+
+FFrameTime ULiveLinkTimeSynchronizationSource::GetNewestSampleTime() const
+{
+	UpdateCachedState();
+	return CachedData.NewestSampleTime + FrameOffset;
+}
+
+FFrameTime ULiveLinkTimeSynchronizationSource::GetOldestSampleTime() const
+{
+	UpdateCachedState();
+	return CachedData.OldestSampleTime + FrameOffset;
+}
+
+FFrameRate ULiveLinkTimeSynchronizationSource::GetFrameRate() const
+{
+	UpdateCachedState();
+	return CachedData.SampleFrameRate;
+}
+
+bool ULiveLinkTimeSynchronizationSource::IsReady() const
+{
+	UpdateCachedState();
+	return LiveLinkClient && CachedData.bIsValid && IsCurrentStateValid();
+}
+
+bool ULiveLinkTimeSynchronizationSource::Open(const FTimeSynchronizationOpenData& OpenData)
+{
+	State = ESyncState::NotSynced;
+	SubjectKey = FLiveLinkSubjectKey();
+
+	if (LiveLinkClient == nullptr)
+	{
+		return false;
+	}
+
+	TArray<FLiveLinkSubjectKey> AllSubjects = LiveLinkClient->GetSubjects(false, false);
+	FLiveLinkSubjectKey* SubjectKeyPtr = AllSubjects.FindByPredicate([this](const FLiveLinkSubjectKey& InSubjectKey) { return InSubjectKey.SubjectName == SubjectName; });
+	if (SubjectKeyPtr == nullptr)
+	{
+		UE_LOG(LogLiveLink, Error, TEXT("The susbject '%s' is not valid"), *SubjectName.ToString());
+		return false;
+	}
+
+	SubjectKey = *SubjectKeyPtr;
+
+	bool bResult = IsCurrentStateValid();
+
+	if (bResult)
+	{
+		State = ESyncState::Opened;
+	}
+	return bResult;
+}
+
+void ULiveLinkTimeSynchronizationSource::Start(const FTimeSynchronizationStartData& StartData)
+{
+}
+
+void ULiveLinkTimeSynchronizationSource::Close()
+{
+	State = ESyncState::NotSynced;
+	SubjectKey = FLiveLinkSubjectKey();
+}
+
+FString ULiveLinkTimeSynchronizationSource::GetDisplayName() const
+{
+	return SubjectName.ToString();
+}
+
+bool ULiveLinkTimeSynchronizationSource::IsCurrentStateValid() const
+{
+	ensure(LiveLinkClient != nullptr);
+	if (LiveLinkClient == nullptr)
+	{
+		return false;
+	}
+
+	if (!LiveLinkClient->IsSubjectEnabled(SubjectKey, false))
+	{
+		UE_LOG(LogLiveLink, Error, TEXT("The subject '%s' is not enabled."), *SubjectName.ToString());
+		return false;
+	}
+
+	if (LiveLinkClient->IsVirtualSubject(SubjectKey))
+	{
+		UE_LOG(LogLiveLink, Error, TEXT("The subject '%s' can't be a virtual subject."), *SubjectName.ToString());
+		return false;
+	}
+
+	ULiveLinkSourceSettings* SourceSettings = LiveLinkClient->GetSourceSettings(SubjectKey.Source);
+	if (SourceSettings == nullptr)
+	{
+		UE_LOG(LogLiveLink, Error, TEXT("The subject '%s' source does not have a source settings."), *SubjectName.ToString());
+		return false;
+	}
+
+	if (SourceSettings->Mode != ELiveLinkSourceMode::Timecode)
+	{
+		UE_LOG(LogLiveLink, Error, TEXT("The subject '%s' source is not in Timecode mode."), *SubjectName.ToString());
+		return false;
+	}
+
+	return true;
+}
+
+void ULiveLinkTimeSynchronizationSource::OnModularFeatureRegistered(const FName& FeatureName, class IModularFeature* Feature)
+{
+	if (FeatureName == ILiveLinkClient::ModularFeatureName)
+	{
+		LiveLinkClient = static_cast<FLiveLinkClient*>(Feature);
+	}
+}
+
+void ULiveLinkTimeSynchronizationSource::OnModularFeatureUnregistered(const FName& FeatureName, class IModularFeature* Feature)
+{
+	if (FeatureName == ILiveLinkClient::ModularFeatureName && (LiveLinkClient != nullptr) && ensure(Feature == LiveLinkClient))
+	{
+		LiveLinkClient = nullptr;
+	}
+}
+
+void ULiveLinkTimeSynchronizationSource::UpdateCachedState() const
+{
+	if (LastUpdateFrame != GFrameCounter && LiveLinkClient != nullptr)
+	{
+		LastUpdateFrame = GFrameCounter;
+		CachedData = LiveLinkClient->GetTimeSyncData(SubjectName);
+	}
+}
\ No newline at end of file
diff --git a/Engine/Plugins/Animation/LiveLink/Source/LiveLink/Private/TimeSync/LiveLinkTimeSynchronizationSource.cpp b/Engine/Plugins/Animation/LiveLink/Source/LiveLink/Private/TimeSync/LiveLinkTimeSynchronizationSource.cpp
deleted file mode 100644
diff --git a/Engine/Plugins/Animation/LiveLink/Source/LiveLink/Private/TimeSync/LiveLinkTimeSynchronizationSource.h b/Engine/Plugins/Animation/LiveLink/Source/LiveLink/Private/TimeSync/LiveLinkTimeSynchronizationSource.h
deleted file mode 100644
diff --git a/Engine/Plugins/Animation/LiveLink/Source/LiveLink/Public/LiveLinkTimeSynchronizationSource.h b/Engine/Plugins/Animation/LiveLink/Source/LiveLink/Public/LiveLinkTimeSynchronizationSource.h
new file mode 100644
--- /dev/null
+++ b/Engine/Plugins/Animation/LiveLink/Source/LiveLink/Public/LiveLinkTimeSynchronizationSource.h
@@ -0,0 +1,55 @@
+// Copyright 1998-2019 Epic Games, Inc. All Rights Reserved.
+
+#pragma once
+
+#include "CoreMinimal.h"
+#include "TimeSynchronizationSource.h"
+#include "LiveLinkClient.h"
+#include "Misc/Guid.h"
+#include "LiveLinkTimeSynchronizationSource.generated.h"
+
+UCLASS(EditInlineNew)
+class LIVELINK_API ULiveLinkTimeSynchronizationSource : public UTimeSynchronizationSource
+{
+	GENERATED_BODY()
+
+public:
+	UPROPERTY(EditAnywhere, Category="LiveLink")
+	FLiveLinkSubjectName SubjectName;
+
+private:
+	FLiveLinkClient* LiveLinkClient;
+
+	enum class ESyncState
+	{
+		NotSynced,
+		Opened,
+	};
+
+	mutable ESyncState State = ESyncState::NotSynced;
+	mutable FLiveLinkSubjectTimeSyncData CachedData;
+	mutable int64 LastUpdateFrame;
+	FLiveLinkSubjectKey SubjectKey;
+
+public:
+
+	ULiveLinkTimeSynchronizationSource();
+
+	//~ Begin TimeSynchronizationSource API
+	virtual FFrameTime GetNewestSampleTime() const override;
+	virtual FFrameTime GetOldestSampleTime() const override;
+	virtual FFrameRate GetFrameRate() const override;
+	virtual bool IsReady() const override;
+	virtual bool Open(const FTimeSynchronizationOpenData& OpenData) override;
+	virtual void Start(const FTimeSynchronizationStartData& StartData) override;
+	virtual void Close() override;
+	virtual FString GetDisplayName() const override;
+	//~ End TimeSynchronizationSource API
+
+private:
+
+	bool IsCurrentStateValid() const;
+	void OnModularFeatureRegistered(const FName& FeatureName, class IModularFeature* Feature);
+	void OnModularFeatureUnregistered(const FName& FeatureName, class IModularFeature* Feature);
+	void UpdateCachedState() const;
+};
\ No newline at end of file
