Crash in LiveLink animation pose when the subject doesn't exist

diff --git a/Engine/Plugins/Animation/LiveLink/Source/LiveLink/Private/AnimNode_LiveLinkPose.cpp b/Engine/Plugins/Animation/LiveLink/Source/LiveLink/Private/AnimNode_LiveLinkPose.cpp
--- a/Engine/Plugins/Animation/LiveLink/Source/LiveLink/Private/AnimNode_LiveLinkPose.cpp
+++ b/Engine/Plugins/Animation/LiveLink/Source/LiveLink/Private/AnimNode_LiveLinkPose.cpp
@@ -67,33 +67,36 @@ void FAnimNode_LiveLinkPose::Evaluate_AnyThread(FPoseContext& Output)
 	FLiveLinkSubjectFrameData SubjectFrameData;
 
 	TSubclassOf<ULiveLinkRole> SubjectRole = LiveLinkClient->GetSubjectRole(LiveLinkSubjectName);
-	if (SubjectRole && SubjectRole->IsChildOf(ULiveLinkAnimationRole::StaticClass()))
+	if (SubjectRole)
 	{
-		//Process animation data if the subject is from that type
-		if (LiveLinkClient->EvaluateFrame_AnyThread(LiveLinkSubjectName, ULiveLinkAnimationRole::StaticClass(), SubjectFrameData))
+		if (SubjectRole->IsChildOf(ULiveLinkAnimationRole::StaticClass()))
 		{
-			FLiveLinkSkeletonStaticData* SkeletonData = SubjectFrameData.StaticData.Cast<FLiveLinkSkeletonStaticData>();
-			FLiveLinkAnimationFrameData* FrameData = SubjectFrameData.FrameData.Cast<FLiveLinkAnimationFrameData>();
-			check(SkeletonData);
-			check(FrameData);
-
-			CurrentRetargetAsset->BuildPoseFromAnimationData(CachedDeltaTime, SkeletonData, FrameData, Output.Pose);
-			CurrentRetargetAsset->BuildPoseAndCurveFromBaseData(CachedDeltaTime, SkeletonData, FrameData, Output.Pose, Output.Curve);
-			CachedDeltaTime = 0.f; // Reset so that if we evaluate again we don't "create" time inside of the retargeter
+			//Process animation data if the subject is from that type
+			if (LiveLinkClient->EvaluateFrame_AnyThread(LiveLinkSubjectName, ULiveLinkAnimationRole::StaticClass(), SubjectFrameData))
+			{
+				FLiveLinkSkeletonStaticData* SkeletonData = SubjectFrameData.StaticData.Cast<FLiveLinkSkeletonStaticData>();
+				FLiveLinkAnimationFrameData* FrameData = SubjectFrameData.FrameData.Cast<FLiveLinkAnimationFrameData>();
+				check(SkeletonData);
+				check(FrameData);
+
+				CurrentRetargetAsset->BuildPoseFromAnimationData(CachedDeltaTime, SkeletonData, FrameData, Output.Pose);
+				CurrentRetargetAsset->BuildPoseAndCurveFromBaseData(CachedDeltaTime, SkeletonData, FrameData, Output.Pose, Output.Curve);
+				CachedDeltaTime = 0.f; // Reset so that if we evaluate again we don't "create" time inside of the retargeter
+			}
 		}
-	}
-	else
-	{
-		//Otherwise, fetch basic data that contains property / curve data
-		if (LiveLinkClient->EvaluateFrame_AnyThread(LiveLinkSubjectName, ULiveLinkBasicRole::StaticClass(), SubjectFrameData))
+		else
 		{
-			FLiveLinkBaseStaticData* BaseStaticData = SubjectFrameData.StaticData.Cast<FLiveLinkBaseStaticData>();
-			FLiveLinkBaseFrameData* BaseFrameData = SubjectFrameData.FrameData.Cast<FLiveLinkBaseFrameData>();
-			check(BaseStaticData);
-			check(BaseFrameData);
-
-			CurrentRetargetAsset->BuildPoseAndCurveFromBaseData(CachedDeltaTime, BaseStaticData, BaseFrameData, Output.Pose, Output.Curve);
-			CachedDeltaTime = 0.f; // Reset so that if we evaluate again we don't "create" time inside of the retargeter
+			//Otherwise, fetch basic data that contains property / curve data
+			if (LiveLinkClient->EvaluateFrame_AnyThread(LiveLinkSubjectName, ULiveLinkBasicRole::StaticClass(), SubjectFrameData))
+			{
+				FLiveLinkBaseStaticData* BaseStaticData = SubjectFrameData.StaticData.Cast<FLiveLinkBaseStaticData>();
+				FLiveLinkBaseFrameData* BaseFrameData = SubjectFrameData.FrameData.Cast<FLiveLinkBaseFrameData>();
+				check(BaseStaticData);
+				check(BaseFrameData);
+
+				CurrentRetargetAsset->BuildPoseAndCurveFromBaseData(CachedDeltaTime, BaseStaticData, BaseFrameData, Output.Pose, Output.Curve);
+				CachedDeltaTime = 0.f; // Reset so that if we evaluate again we don't "create" time inside of the retargeter
+			}
 		}
 	}
 }
diff --git a/Engine/Plugins/Animation/LiveLink/Source/LiveLinkEditor/Private/SLiveLinkDataView.cpp b/Engine/Plugins/Animation/LiveLink/Source/LiveLinkEditor/Private/SLiveLinkDataView.cpp
--- a/Engine/Plugins/Animation/LiveLink/Source/LiveLinkEditor/Private/SLiveLinkDataView.cpp
+++ b/Engine/Plugins/Animation/LiveLink/Source/LiveLinkEditor/Private/SLiveLinkDataView.cpp
@@ -182,20 +182,23 @@ void SLiveLinkDataView::Tick(const FGeometry& AllottedGeometry, const double InC
 		if (Client->IsSubjectEnabled(SubjectKey, true))
 		{
 			TSubclassOf<ULiveLinkRole> SubjectRole = Client->GetSubjectRole(SubjectKey.SubjectName);
-			FLiveLinkSubjectFrameData SubjectData;
-			if (Client->EvaluateFrame_AnyThread(SubjectKey.SubjectName, SubjectRole, SubjectData))
+			if (SubjectRole)
 			{
-				if (DetailType == EDetailType::StaticData)
+				FLiveLinkSubjectFrameData SubjectData;
+				if (Client->EvaluateFrame_AnyThread(SubjectKey.SubjectName, SubjectRole, SubjectData))
 				{
-					TSharedPtr<FStructOnScope> StructOnScope = MakeShared<FStructOnScope>(SubjectData.StaticData.GetStruct());
-					CastChecked<UScriptStruct>(StructOnScope->GetStruct())->CopyScriptStruct(StructOnScope->GetStructMemory(), SubjectData.StaticData.GetBaseData());
-					StructureDetailsView->SetStructureData(StructOnScope);
-				}
-				else
-				{
-					TSharedPtr<FStructOnScope> StructOnScope = MakeShared<FStructOnScope>(SubjectData.FrameData.GetStruct());
-					CastChecked<UScriptStruct>(StructOnScope->GetStruct())->CopyScriptStruct(StructOnScope->GetStructMemory(), SubjectData.FrameData.GetBaseData());
-					StructureDetailsView->SetStructureData(StructOnScope);
+					if (DetailType == EDetailType::StaticData)
+					{
+						TSharedPtr<FStructOnScope> StructOnScope = MakeShared<FStructOnScope>(SubjectData.StaticData.GetStruct());
+						CastChecked<UScriptStruct>(StructOnScope->GetStruct())->CopyScriptStruct(StructOnScope->GetStructMemory(), SubjectData.StaticData.GetBaseData());
+						StructureDetailsView->SetStructureData(StructOnScope);
+					}
+					else
+					{
+						TSharedPtr<FStructOnScope> StructOnScope = MakeShared<FStructOnScope>(SubjectData.FrameData.GetStruct());
+						CastChecked<UScriptStruct>(StructOnScope->GetStruct())->CopyScriptStruct(StructOnScope->GetStructMemory(), SubjectData.FrameData.GetBaseData());
+						StructureDetailsView->SetStructureData(StructOnScope);
+					}
 				}
 			}
 		}
