Fix for LiveLink recorded track having the wrong length and showing no data from partner branch (CL #8101469)

diff --git a/Engine/Plugins/Animation/LiveLink/Source/LiveLinkSequencer/Private/SequencerRecorderSections/MovieSceneLiveLinkTrackRecorder.cpp b/Engine/Plugins/Animation/LiveLink/Source/LiveLinkSequencer/Private/SequencerRecorderSections/MovieSceneLiveLinkTrackRecorder.cpp
--- a/Engine/Plugins/Animation/LiveLink/Source/LiveLinkSequencer/Private/SequencerRecorderSections/MovieSceneLiveLinkTrackRecorder.cpp
+++ b/Engine/Plugins/Animation/LiveLink/Source/LiveLinkSequencer/Private/SequencerRecorderSections/MovieSceneLiveLinkTrackRecorder.cpp
@@ -168,7 +168,13 @@ void UMovieSceneLiveLinkTrackRecorder::FinalizeTrackImpl()
 		FKeyDataOptimizationParams Params;
 		Params.bAutoSetInterpolation = true;
 		MovieSceneSection->FinalizeSection(bReduceKeys, Params);
-		
+
+		TOptional<TRange<FFrameNumber> > DefaultSectionLength = MovieSceneSection->GetAutoSizeRange();
+		if (DefaultSectionLength.IsSet())
+		{
+			MovieSceneSection->SetRange(DefaultSectionLength.GetValue());
+		}
+
 		MovieSceneSection->SetIsActive(true);
 	}
 }
