Merging #9442709 from Dev-Enterprise

- Subject Role is now split into a hierarchy and a controller is assigned for each role class of the hierarchy.
- User can now select which controller to use for the given role
- Project settings was added to set a default controller per role class
- Detail Customization was added to LiveLinkComponent to show the controller map, handle modifications to the component and better support MU workflow

diff --git a/Engine/Plugins/Animation/LiveLink/Source/LiveLinkComponents/Private/Controllers/LiveLinkCameraController.cpp b/Engine/Plugins/Animation/LiveLink/Source/LiveLinkComponents/Private/Controllers/LiveLinkCameraController.cpp
--- a/Engine/Plugins/Animation/LiveLink/Source/LiveLinkComponents/Private/Controllers/LiveLinkCameraController.cpp
+++ b/Engine/Plugins/Animation/LiveLink/Source/LiveLinkComponents/Private/Controllers/LiveLinkCameraController.cpp
@@ -2,76 +2,96 @@
 
 #include "Controllers/LiveLinkCameraController.h"
 
-#include "ILiveLinkClient.h"
-#include "Roles/LiveLinkCameraRole.h"
-#include "Roles/LiveLinkCameraTypes.h"
 
 #include "Camera/CameraComponent.h"
 #include "CineCameraComponent.h"
+#include "Controllers/LiveLinkTransformController.h"
 #include "Features/IModularFeatures.h"
 #include "GameFramework/Actor.h"
+#include "ILiveLinkClient.h"
+#include "LiveLinkComponentController.h"
+#include "Roles/LiveLinkCameraRole.h"
+#include "Roles/LiveLinkCameraTypes.h"
+#include "UObject/EnterpriseObjectVersion.h"
 
 #if WITH_EDITOR
 #include "Kismet2/ComponentEditorUtils.h"
 #endif
 
 
-void ULiveLinkCameraController::OnEvaluateRegistered()
+void ULiveLinkCameraController::Tick(float DeltaTime, const FLiveLinkSubjectFrameData& SubjectData)
 {
-	AActor* OuterActor = GetOuterActor();
-	TransformData.CheckForError(OuterActor ? OuterActor->GetFName() : NAME_None, Cast<USceneComponent>(ComponentToControl.GetComponent(OuterActor)));
-}
-
+	const FLiveLinkCameraStaticData* StaticData = SubjectData.StaticData.Cast<FLiveLinkCameraStaticData>();
+	const FLiveLinkCameraFrameData* FrameData = SubjectData.FrameData.Cast<FLiveLinkCameraFrameData>();
 
-void ULiveLinkCameraController::Tick(float DeltaTime, const FLiveLinkSubjectRepresentation& SubjectRepresentation)
-{
-	if (UCameraComponent* CameraComponent = Cast<UCameraComponent>(ComponentToControl.GetComponent(GetOuterActor())))
+	if (StaticData && FrameData)
 	{
-		ILiveLinkClient& LiveLinkClient = IModularFeatures::Get().GetModularFeature<ILiveLinkClient>(ILiveLinkClient::ModularFeatureName);
-
-		FLiveLinkSubjectFrameData SubjectData;
-		if (LiveLinkClient.EvaluateFrame_AnyThread(SubjectRepresentation.Subject, SubjectRepresentation.Role, SubjectData))
+		if (UCameraComponent* CameraComponent = Cast<UCameraComponent>(AttachedComponent))
 		{
-			FLiveLinkCameraStaticData* StaticData = SubjectData.StaticData.Cast<FLiveLinkCameraStaticData>();
-			FLiveLinkCameraFrameData* FrameData = SubjectData.FrameData.Cast<FLiveLinkCameraFrameData>();
+			if (StaticData->bIsFieldOfViewSupported) { CameraComponent->SetFieldOfView(FrameData->FieldOfView); }
+			if (StaticData->bIsAspectRatioSupported) { CameraComponent->SetAspectRatio(FrameData->AspectRatio); }
+			if (StaticData->bIsProjectionModeSupported) { CameraComponent->SetProjectionMode(FrameData->ProjectionMode == ELiveLinkCameraProjectionMode::Perspective ? ECameraProjectionMode::Perspective : ECameraProjectionMode::Orthographic); }
 
-			if (StaticData && FrameData)
+			if (UCineCameraComponent* CineCameraComponent = Cast<UCineCameraComponent>(CameraComponent))
 			{
-				TransformData.ApplyTransform(CameraComponent, FrameData->Transform);
-
-				if (StaticData->bIsFieldOfViewSupported) { CameraComponent->SetFieldOfView(FrameData->FieldOfView); } 
-				if (StaticData->bIsAspectRatioSupported) { CameraComponent->SetAspectRatio(FrameData->AspectRatio); }
-				if (StaticData->bIsProjectionModeSupported) { CameraComponent->SetProjectionMode(FrameData->ProjectionMode == ELiveLinkCameraProjectionMode::Perspective ? ECameraProjectionMode::Perspective : ECameraProjectionMode::Orthographic); }
-
-				if (UCineCameraComponent* CineCameraComponent = Cast<UCineCameraComponent>(CameraComponent))
-				{
-					if (StaticData->bIsFocalLengthSupported) { CineCameraComponent->CurrentFocalLength = FrameData->FocalLength; }
-					if (StaticData->bIsApertureSupported) { CineCameraComponent->CurrentAperture = FrameData->Aperture; }
-					if (StaticData->FilmBackWidth > 0.0f) { CineCameraComponent->FilmbackSettings.SensorWidth = StaticData->FilmBackWidth; }
-					if (StaticData->FilmBackHeight > 0.0f) { CineCameraComponent->FilmbackSettings.SensorHeight = StaticData->FilmBackHeight; }
-					if (StaticData->bIsFocusDistanceSupported) { CineCameraComponent->FocusSettings.ManualFocusDistance = FrameData->FocusDistance; }
-				}
+				if (StaticData->bIsFocalLengthSupported) { CineCameraComponent->CurrentFocalLength = FrameData->FocalLength; }
+				if (StaticData->bIsApertureSupported) { CineCameraComponent->CurrentAperture = FrameData->Aperture; }
+				if (StaticData->FilmBackWidth > 0.0f) { CineCameraComponent->FilmbackSettings.SensorWidth = StaticData->FilmBackWidth; }
+				if (StaticData->FilmBackHeight > 0.0f) { CineCameraComponent->FilmbackSettings.SensorHeight = StaticData->FilmBackHeight; }
+				if (StaticData->bIsFocusDistanceSupported) { CineCameraComponent->FocusSettings.ManualFocusDistance = FrameData->FocusDistance; }
 			}
 		}
 	}
 }
 
-
 bool ULiveLinkCameraController::IsRoleSupported(const TSubclassOf<ULiveLinkRole>& RoleToSupport)
 {
-	return RoleToSupport->IsChildOf(ULiveLinkCameraRole::StaticClass());
+	return RoleToSupport == ULiveLinkCameraRole::StaticClass();
 }
 
+TSubclassOf<UActorComponent> ULiveLinkCameraController::GetDesiredComponentClass() const
+{
+	return UCameraComponent::StaticClass();
+}
 
-#if WITH_EDITOR
-void ULiveLinkCameraController::InitializeInEditor()
+void ULiveLinkCameraController::PostLoad()
 {
-	if (AActor* Actor = GetOuterActor())
+	Super::PostLoad();
+
+#if WITH_EDITOR
+	const int32 Version = GetLinkerCustomVersion(FEnterpriseObjectVersion::GUID);
+	if (Version < FEnterpriseObjectVersion::LiveLinkControllerSplitPerRole)
 	{
-		if (UCameraComponent* CameraComponent = Actor->FindComponentByClass<UCameraComponent>())
+		AActor* MyActor = GetOuterActor();
+		if (MyActor)
 		{
-			ComponentToControl = FComponentEditorUtils::MakeComponentReference(Actor, CameraComponent);
+			//Make sure all UObjects we use in our post load have been postloaded
+			MyActor->ConditionalPostLoad();
+
+			ULiveLinkComponentController* LiveLinkComponent = Cast<ULiveLinkComponentController>(MyActor->GetComponentByClass(ULiveLinkComponentController::StaticClass()));
+			if (LiveLinkComponent)
+			{
+				LiveLinkComponent->ConditionalPostLoad();
+
+				//If the transform controller that was created to drive the TransformRole is the built in one, set its data structure with the one that we had internally
+				if (LiveLinkComponent->ControllerMap.Contains(ULiveLinkTransformRole::StaticClass()))
+				{
+					ULiveLinkTransformController* TransformController = Cast<ULiveLinkTransformController>(LiveLinkComponent->ControllerMap[ULiveLinkTransformRole::StaticClass()]);
+					if (TransformController)
+					{
+						TransformController->ConditionalPostLoad();
+						TransformController->TransformData = TransformData_DEPRECATED;
+					}
+				}
+
+				//if Subjects role direct controller is us, set the component to control to what we had
+				if (LiveLinkComponent->SubjectRepresentation.Role == ULiveLinkCameraRole::StaticClass())
+				{
+					LiveLinkComponent->ComponentToControl = ComponentToControl_DEPRECATED;
+				}
+			}
 		}
 	}
+#endif //WITH_EDITOR
 }
-#endif
+
diff --git a/Engine/Plugins/Animation/LiveLink/Source/LiveLinkComponents/Private/Controllers/LiveLinkLightController.cpp b/Engine/Plugins/Animation/LiveLink/Source/LiveLinkComponents/Private/Controllers/LiveLinkLightController.cpp
--- a/Engine/Plugins/Animation/LiveLink/Source/LiveLinkComponents/Private/Controllers/LiveLinkLightController.cpp
+++ b/Engine/Plugins/Animation/LiveLink/Source/LiveLinkComponents/Private/Controllers/LiveLinkLightController.cpp
@@ -2,81 +2,99 @@
 
 #include "Controllers/LiveLinkLightController.h"
 
-#include "ILiveLinkClient.h"
-#include "Roles/LiveLinkLightRole.h"
-#include "Roles/LiveLinkLightTypes.h"
-
 #include "Components/LightComponent.h"
 #include "Components/SpotLightComponent.h"
-#include "Features/IModularFeatures.h"
+#include "Controllers/LiveLinkTransformController.h"
 #include "GameFramework/Actor.h"
+#include "ILiveLinkClient.h"
+#include "LiveLinkComponentController.h"
+#include "Roles/LiveLinkLightRole.h"
+#include "Roles/LiveLinkLightTypes.h"
+#include "UObject/EnterpriseObjectVersion.h"
 
 #if WITH_EDITOR
 #include "Kismet2/ComponentEditorUtils.h"
 #endif
 
 
-void ULiveLinkLightController::OnEvaluateRegistered()
+void ULiveLinkLightController::Tick(float DeltaTime, const FLiveLinkSubjectFrameData& SubjectData)
 {
-	AActor* OuterActor = GetOuterActor();
-	TransformData.CheckForError(OuterActor ? OuterActor->GetFName() : NAME_None, Cast<USceneComponent>(ComponentToControl.GetComponent(OuterActor)));
-}
-
+	const FLiveLinkLightStaticData* StaticData = SubjectData.StaticData.Cast<FLiveLinkLightStaticData>();
+	const FLiveLinkLightFrameData* FrameData = SubjectData.FrameData.Cast<FLiveLinkLightFrameData>();
 
-void ULiveLinkLightController::Tick(float DeltaTime, const FLiveLinkSubjectRepresentation& SubjectRepresentation)
-{
-	if (ULightComponent* LightComponent = Cast<ULightComponent>(ComponentToControl.GetComponent(GetOuterActor())))
+	if (StaticData && FrameData)
 	{
-		ILiveLinkClient& LiveLinkClient = IModularFeatures::Get().GetModularFeature<ILiveLinkClient>(ILiveLinkClient::ModularFeatureName);
-
-		FLiveLinkSubjectFrameData SubjectData;
-		if (LiveLinkClient.EvaluateFrame_AnyThread(SubjectRepresentation.Subject, SubjectRepresentation.Role, SubjectData))
+		if (ULightComponent* LightComponent = Cast<ULightComponent>(AttachedComponent))
 		{
-			FLiveLinkLightStaticData* StaticData = SubjectData.StaticData.Cast<FLiveLinkLightStaticData>();
-			FLiveLinkLightFrameData* FrameData = SubjectData.FrameData.Cast<FLiveLinkLightFrameData>();
+			if (StaticData->bIsTemperatureSupported) { LightComponent->SetTemperature(FrameData->Temperature); }
+			if (StaticData->bIsIntensitySupported) { LightComponent->SetIntensity(FrameData->Intensity); }
+			if (StaticData->bIsLightColorSupported) { LightComponent->SetLightColor(FrameData->LightColor); }
 
-			if (StaticData && FrameData)
+			if (UPointLightComponent* PointLightComponent = Cast<UPointLightComponent>(LightComponent))
 			{
-				TransformData.ApplyTransform(LightComponent, FrameData->Transform);
+				if (StaticData->bIsAttenuationRadiusSupported) { PointLightComponent->SetAttenuationRadius(FrameData->AttenuationRadius); }
+				if (StaticData->bIsSourceRadiusSupported) { PointLightComponent->SetSourceRadius(FrameData->SourceRadius); }
+				if (StaticData->bIsSoftSourceRadiusSupported) { PointLightComponent->SetSoftSourceRadius(FrameData->SoftSourceRadius); }
+				if (StaticData->bIsSourceLenghtSupported) { PointLightComponent->SetSourceLength(FrameData->SourceLength); }
 
-				if (StaticData->bIsTemperatureSupported) { LightComponent->SetTemperature(FrameData->Temperature); }
-				if (StaticData->bIsIntensitySupported) { LightComponent->SetIntensity(FrameData->Intensity); }
-				if (StaticData->bIsLightColorSupported) { LightComponent->SetLightColor(FrameData->LightColor); }
-
-				if (UPointLightComponent* PointLightComponent = Cast<UPointLightComponent>(LightComponent))
+				if (USpotLightComponent* SpotlightComponent = Cast<USpotLightComponent>(LightComponent))
 				{
-					if (StaticData->bIsAttenuationRadiusSupported) { PointLightComponent->SetAttenuationRadius(FrameData->AttenuationRadius); }
-					if (StaticData->bIsSourceRadiusSupported) { PointLightComponent->SetSourceRadius(FrameData->SourceRadius); }
-					if (StaticData->bIsSoftSourceRadiusSupported) { PointLightComponent->SetSoftSourceRadius(FrameData->SoftSourceRadius); }
-					if (StaticData->bIsSourceLenghtSupported) { PointLightComponent->SetSourceLength(FrameData->SourceLength); }
-
-					if (USpotLightComponent* SpotlightComponent = Cast<USpotLightComponent>(LightComponent))
-					{
-						if (StaticData->bIsInnerConeAngleSupported) { SpotlightComponent->SetInnerConeAngle(FrameData->InnerConeAngle); }
-						if (StaticData->bIsOuterConeAngleSupported) { SpotlightComponent->SetOuterConeAngle(FrameData->OuterConeAngle); }
-					}
+					if (StaticData->bIsInnerConeAngleSupported) { SpotlightComponent->SetInnerConeAngle(FrameData->InnerConeAngle); }
+					if (StaticData->bIsOuterConeAngleSupported) { SpotlightComponent->SetOuterConeAngle(FrameData->OuterConeAngle); }
 				}
 			}
 		}
 	}
 }
 
-
 bool ULiveLinkLightController::IsRoleSupported(const TSubclassOf<ULiveLinkRole>& RoleToSupport)
 {
-	return RoleToSupport->IsChildOf(ULiveLinkLightRole::StaticClass());
+	return RoleToSupport == ULiveLinkLightRole::StaticClass();
+}
+
+TSubclassOf<UActorComponent> ULiveLinkLightController::GetDesiredComponentClass() const
+{
+	return ULightComponent::StaticClass();
 }
 
+void ULiveLinkLightController::PostLoad()
+{
+	Super::PostLoad();
 
 #if WITH_EDITOR
-void ULiveLinkLightController::InitializeInEditor()
-{
-	if (AActor* Actor = GetOuterActor())
+	const int32 Version = GetLinkerCustomVersion(FEnterpriseObjectVersion::GUID);
+	if (Version < FEnterpriseObjectVersion::LiveLinkControllerSplitPerRole)
 	{
-		if (ULightComponent* LightComponent = Actor->FindComponentByClass<ULightComponent>())
+		AActor* MyActor = GetOuterActor();
+		if (MyActor)
 		{
-			ComponentToControl = FComponentEditorUtils::MakeComponentReference(Actor, LightComponent);
+			//Make sure all UObjects we use in our post load have been postloaded
+			MyActor->ConditionalPostLoad();
+
+			ULiveLinkComponentController* LiveLinkComponent = Cast<ULiveLinkComponentController>(MyActor->GetComponentByClass(ULiveLinkComponentController::StaticClass()));
+			if (LiveLinkComponent)
+			{
+				LiveLinkComponent->ConditionalPostLoad();
+
+				//If the transform controller that was created to drive the TransformRole is the built in one, set its data structure with the one that we had internally
+				if (LiveLinkComponent->ControllerMap.Contains(ULiveLinkTransformRole::StaticClass()))
+				{
+					ULiveLinkTransformController* TransformController = Cast<ULiveLinkTransformController>(LiveLinkComponent->ControllerMap[ULiveLinkTransformRole::StaticClass()]);
+					if (TransformController)
+					{
+						TransformController->ConditionalPostLoad();
+						TransformController->TransformData = TransformData_DEPRECATED;
+					}
+				}
+
+				//if Subjects role direct controller is us, set the component to control to what we had
+				if (LiveLinkComponent->SubjectRepresentation.Role == ULiveLinkLightRole::StaticClass())
+				{
+					LiveLinkComponent->ComponentToControl = ComponentToControl_DEPRECATED;
+				}
+			}
 		}
 	}
+#endif //WITH_EDITOR
 }
-#endif
+
diff --git a/Engine/Plugins/Animation/LiveLink/Source/LiveLinkComponents/Private/Controllers/LiveLinkTransformController.cpp b/Engine/Plugins/Animation/LiveLink/Source/LiveLinkComponents/Private/Controllers/LiveLinkTransformController.cpp
--- a/Engine/Plugins/Animation/LiveLink/Source/LiveLinkComponents/Private/Controllers/LiveLinkTransformController.cpp
+++ b/Engine/Plugins/Animation/LiveLink/Source/LiveLinkComponents/Private/Controllers/LiveLinkTransformController.cpp
@@ -1,15 +1,16 @@
 // Copyright 1998-2019 Epic Games, Inc. All Rights Reserved.
 
 #include "Controllers/LiveLinkTransformController.h"
-#include "LiveLinkComponentPrivate.h"
 
+#include "Components/SceneComponent.h"
+#include "GameFramework/Actor.h"
 #include "ILiveLinkClient.h"
+#include "LiveLinkComponentController.h"
+#include "LiveLinkCustomVersion.h"
+#include "LiveLinkComponentPrivate.h"
 #include "Roles/LiveLinkTransformRole.h"
 #include "Roles/LiveLinkTransformTypes.h"
-
-#include "Components/SceneComponent.h"
-#include "Features/IModularFeatures.h"
-#include "GameFramework/Actor.h"
+#include "UObject/EnterpriseObjectVersion.h"
 
 #if WITH_EDITOR
 #include "Framework/Notifications/NotificationManager.h"
@@ -17,7 +18,7 @@
 #include "Widgets/Notifications/SNotificationList.h"
 #endif
 
-#define LOCTEXT_NAMESPACE "LiveLinkController"
+#define LOCTEXT_NAMESPACE "LiveLinkTransformController"
 
 void FLiveLinkTransformControllerData::ApplyTransform(USceneComponent* SceneComponent, const FTransform& Transform) const
 {
@@ -39,7 +40,6 @@ void FLiveLinkTransformControllerData::ApplyTransform(USceneComponent* SceneComp
 	}
 }
 
-
 void FLiveLinkTransformControllerData::CheckForError(FName OwnerName, USceneComponent* SceneComponent) const
 {
 	if (SceneComponent == nullptr)
@@ -66,48 +66,70 @@ void FLiveLinkTransformControllerData::CheckForError(FName OwnerName, USceneComp
 void ULiveLinkTransformController::OnEvaluateRegistered()
 {
 	AActor* OuterActor = GetOuterActor();
-	TransformData.CheckForError(OuterActor ? OuterActor->GetFName() : NAME_None, Cast<USceneComponent>(ComponentToControl.GetComponent(OuterActor)));
+	TransformData.CheckForError(OuterActor ? OuterActor->GetFName() : NAME_None, Cast<USceneComponent>(AttachedComponent));
 }
 
-
-void ULiveLinkTransformController::Tick(float DeltaTime, const FLiveLinkSubjectRepresentation& SubjectRepresentation)
+void ULiveLinkTransformController::Tick(float DeltaTime, const FLiveLinkSubjectFrameData& SubjectData)
 {
-	if (USceneComponent* SceneComponent = Cast<USceneComponent>(ComponentToControl.GetComponent(GetOuterActor())))
-	{
-		ILiveLinkClient& LiveLinkClient = IModularFeatures::Get().GetModularFeature<ILiveLinkClient>(ILiveLinkClient::ModularFeatureName);
+	const FLiveLinkTransformStaticData* StaticData = SubjectData.StaticData.Cast<FLiveLinkTransformStaticData>();
+	const FLiveLinkTransformFrameData* FrameData = SubjectData.FrameData.Cast<FLiveLinkTransformFrameData>();
 
-		FLiveLinkSubjectFrameData SubjectData;
-		if (LiveLinkClient.EvaluateFrame_AnyThread(SubjectRepresentation.Subject, SubjectRepresentation.Role, SubjectData))
+	if (StaticData && FrameData)
+	{
+		if (USceneComponent* SceneComponent = Cast<USceneComponent>(AttachedComponent))
 		{
-			FLiveLinkTransformStaticData* StaticData = SubjectData.StaticData.Cast<FLiveLinkTransformStaticData>();
-			FLiveLinkTransformFrameData* FrameData = SubjectData.FrameData.Cast<FLiveLinkTransformFrameData>();
-
-			if (StaticData && FrameData)
-			{
-				TransformData.ApplyTransform(SceneComponent, FrameData->Transform);
-			}
+			TransformData.ApplyTransform(SceneComponent, FrameData->Transform);
 		}
 	}
 }
 
-
 bool ULiveLinkTransformController::IsRoleSupported(const TSubclassOf<ULiveLinkRole>& RoleToSupport)
 {
-	return RoleToSupport->IsChildOf(ULiveLinkTransformRole::StaticClass());
+	return RoleToSupport == ULiveLinkTransformRole::StaticClass();
 }
 
+TSubclassOf<UActorComponent> ULiveLinkTransformController::GetDesiredComponentClass() const
+{
+	return USceneComponent::StaticClass();
+}
 
-#if WITH_EDITOR
-void ULiveLinkTransformController::InitializeInEditor()
+void ULiveLinkTransformController::SetAttachedComponent(UActorComponent* ActorComponent)
+{
+	Super::SetAttachedComponent(ActorComponent);
+
+	AActor* OuterActor = GetOuterActor();
+	TransformData.CheckForError(OuterActor ? OuterActor->GetFName() : NAME_None, Cast<USceneComponent>(AttachedComponent));
+}
+
+void ULiveLinkTransformController::PostLoad()
 {
-	if (AActor* Actor = GetOuterActor())
+	Super::PostLoad();
+
+#if WITH_EDITOR
+	const int32 Version = GetLinkerCustomVersion(FEnterpriseObjectVersion::GUID);
+	if (Version < FEnterpriseObjectVersion::LiveLinkControllerSplitPerRole)
 	{
-		if (USceneComponent* SceneComponent = Actor->FindComponentByClass<USceneComponent>())
+		AActor* MyActor = GetOuterActor();
+		if (MyActor)
 		{
-			ComponentToControl = FComponentEditorUtils::MakeComponentReference(Actor, SceneComponent);
+			//Make sure all UObjects we use in our post load have been postloaded
+			MyActor->ConditionalPostLoad();
+
+			ULiveLinkComponentController* LiveLinkComponent = Cast<ULiveLinkComponentController>(MyActor->GetComponentByClass(ULiveLinkComponentController::StaticClass()));
+			if (LiveLinkComponent)
+			{
+				LiveLinkComponent->ConditionalPostLoad();
+
+				//if Subjects role direct controller is us, set the component to control to what we had
+				if (LiveLinkComponent->SubjectRepresentation.Role == ULiveLinkTransformRole::StaticClass())
+				{
+					LiveLinkComponent->ComponentToControl = ComponentToControl_DEPRECATED;
+				}
+			}
 		}
+		
 	}
+#endif //WITH_EDITOR
 }
-#endif
 
 #undef LOCTEXT_NAMESPACE
diff --git a/Engine/Plugins/Animation/LiveLink/Source/LiveLinkComponents/Private/LiveLinkComponentController.cpp b/Engine/Plugins/Animation/LiveLink/Source/LiveLinkComponents/Private/LiveLinkComponentController.cpp
--- a/Engine/Plugins/Animation/LiveLink/Source/LiveLinkComponents/Private/LiveLinkComponentController.cpp
+++ b/Engine/Plugins/Animation/LiveLink/Source/LiveLinkComponents/Private/LiveLinkComponentController.cpp
@@ -1,12 +1,16 @@
 // Copyright 1998-2019 Epic Games, Inc. All Rights Reserved.
 
 #include "LiveLinkComponentController.h"
-#include "LiveLinkComponentPrivate.h"
 
+#include "LiveLinkComponentSettings.h"
+#include "LiveLinkComponentPrivate.h"
 #include "LiveLinkControllerBase.h"
+#include "UObject/EnterpriseObjectVersion.h"
+#include "UObject/UObjectIterator.h"
 
 #if WITH_EDITOR
 #include "Framework/Notifications/NotificationManager.h"
+#include "Kismet2/ComponentEditorUtils.h"
 #include "Widgets/Notifications/SNotificationList.h"
 #endif
 
@@ -23,24 +27,103 @@ ULiveLinkComponentController::ULiveLinkComponentController()
 }
 
 
-void ULiveLinkComponentController::OnRegister()
+void ULiveLinkComponentController::OnSubjectRoleChanged()
 {
+	if (SubjectRepresentation.Role == nullptr)
+	{
+		ControllerMap.Empty();
+	}
+	else
+	{
+		TArray<TSubclassOf<ULiveLinkRole>> SelectedRoleHierarchy = GetSelectedRoleHierarchyClasses(SubjectRepresentation.Role);
+		ControllerMap.Empty(SelectedRoleHierarchy.Num());
+		for (const TSubclassOf<ULiveLinkRole>& RoleClass : SelectedRoleHierarchy)
+		{
+			if (RoleClass)
+			{
+				//Add each role class of the hierarchy in the map and assign a controller, if any, to each of them
+				ControllerMap.FindOrAdd(RoleClass);
+
+				TSubclassOf<ULiveLinkControllerBase> SelectedControllerClass = GetControllerClassForRoleClass(RoleClass);
+				SetControllerClassForRole(RoleClass, SelectedControllerClass);
+			}
+		}
+	}
+}
+
+void ULiveLinkComponentController::SetControllerClassForRole(TSubclassOf<ULiveLinkRole> RoleClass, TSubclassOf<ULiveLinkControllerBase> DesiredControllerClass)
+{
+	if (ControllerMap.Contains(RoleClass))
+	{
+		ULiveLinkControllerBase*& CurrentController = ControllerMap.FindOrAdd(RoleClass);
+		if (CurrentController == nullptr || CurrentController->GetClass() != DesiredControllerClass)
+		{
+			if (DesiredControllerClass != nullptr)
+			{
+
+				const EObjectFlags ControllerObjectFlags = GetMaskedFlags(RF_Public | RF_Transactional | RF_ArchetypeObject);
+				CurrentController = NewObject<ULiveLinkControllerBase>(this, DesiredControllerClass, NAME_None, ControllerObjectFlags);
+
+#if WITH_EDITOR
+				//For the controller directly associated with the subject role, set the component to control to the desired component this controller wants
+				if (RoleClass == SubjectRepresentation.Role)
+				{
+					TSubclassOf<UActorComponent> DesiredComponent = CurrentController->GetDesiredComponentClass();
+					if (AActor* Actor = GetOwner())
+					{
+						if (UActorComponent* ActorComponent = Actor->GetComponentByClass(DesiredComponent))
+						{
+							ComponentToControl = FComponentEditorUtils::MakeComponentReference(Actor, ActorComponent);
+						}
+					}
+				}
+				
+				CurrentController->InitializeInEditor();
+#endif
+			}
+			else
+			{
+				CurrentController = nullptr;
+			}
+		}
+	}
+
+	//Mark ourselves as dirty to update each controller's on next tick
 	bIsDirty = true;
-	Super::OnRegister();
 }
 
+void ULiveLinkComponentController::OnRegister()
+{
+	Super::OnRegister();
+
+	bIsDirty = true;
+}
 
 void ULiveLinkComponentController::TickComponent(float DeltaTime, ELevelTick TickType, FActorComponentTickFunction* ThisTickFunction)
 {
-	if (Controller)
+	ILiveLinkClient& LiveLinkClient = IModularFeatures::Get().GetModularFeature<ILiveLinkClient>(ILiveLinkClient::ModularFeatureName);
+
+	//Evaluate subject frame once and pass the data to our controllers
+	FLiveLinkSubjectFrameData SubjectData;
+	const bool bHasValidData = LiveLinkClient.EvaluateFrame_AnyThread(SubjectRepresentation.Subject, SubjectRepresentation.Role, SubjectData);
+
+	//Go through each controllers and initialize them if we're dirty and tick them if there's valid data to process
+	for (TTuple<TSubclassOf<ULiveLinkRole>, ULiveLinkControllerBase*>& ControllerEntry : ControllerMap)
 	{
-		if (bIsDirty)
+		ULiveLinkControllerBase* Controller = ControllerEntry.Value;
+		if (Controller)
 		{
-			Controller->OnEvaluateRegistered();
-			bIsDirty = false;
+			if (bIsDirty)
+			{
+				Controller->SetAttachedComponent(ComponentToControl.GetComponent(GetOwner()));
+				Controller->OnEvaluateRegistered();
+			}
+			
+			if (bHasValidData)
+			{
+				Controller->Tick(DeltaTime, SubjectData);
+			}
 		}
-
-		Controller->Tick(DeltaTime, SubjectRepresentation);
 	}
 
 	if (OnLiveLinkUpdated.IsBound())
@@ -48,57 +131,150 @@ void ULiveLinkComponentController::TickComponent(float DeltaTime, ELevelTick Tic
 		FEditorScriptExecutionGuard ScriptGuard;
 		OnLiveLinkUpdated.Broadcast(DeltaTime);
 	}
+	
+	bIsDirty = false;
 
 	Super::TickComponent(DeltaTime, TickType, ThisTickFunction);
 }
 
 
+void ULiveLinkComponentController::Serialize(FArchive& Ar)
+{
+	Super::Serialize(Ar);
+
+#if WITH_EDITOR
+	Ar.UsingCustomVersion(FEnterpriseObjectVersion::GUID);
+
+	if (Ar.IsLoading())
+	{
+		if (Ar.CustomVer(FEnterpriseObjectVersion::GUID) < FEnterpriseObjectVersion::LiveLinkControllerSplitPerRole)
+		{
+			ConvertOldControllerSystem();
+		}
+	}
+#endif
+}
+
 #if WITH_EDITOR
 
 void ULiveLinkComponentController::PostEditChangeProperty(struct FPropertyChangedEvent& PropertyChangedEvent)
 {
-	//Detect changes to the SubjectRepresentation blindly when one of our property has changed. In MultiUser, this will be called with an empty property event
-	bool bCreateAnewController = false;
-	if (SubjectRepresentation.Role.Get() == nullptr)
+	const FName PropertyName = PropertyChangedEvent.GetPropertyName();
+	if (PropertyName == GET_MEMBER_NAME_CHECKED(ULiveLinkComponentController, bUpdateInEditor))
 	{
-		Controller = nullptr;
+		bTickInEditor = bUpdateInEditor;
 	}
-	else if (Controller)
+
+	Super::PostEditChangeProperty(PropertyChangedEvent);
+}
+
+void ULiveLinkComponentController::ConvertOldControllerSystem()
+{
+	if (Controller_DEPRECATED)
 	{
-		if (!Controller->IsRoleSupported(SubjectRepresentation.Role))
+		TArray<TSubclassOf<ULiveLinkRole>> SelectedRoleHierarchy = GetSelectedRoleHierarchyClasses(SubjectRepresentation.Role);
+		ControllerMap.Empty(SelectedRoleHierarchy.Num());
+		for (const TSubclassOf<ULiveLinkRole>& RoleClass : SelectedRoleHierarchy)
 		{
-			Controller = nullptr;
-			bCreateAnewController = true;
+			if (RoleClass)
+			{
+				ControllerMap.FindOrAdd(RoleClass);
+
+				//Set the previous controller on the Subject Role entry and create new controllers for parent role classes
+				if (RoleClass == SubjectRepresentation.Role)
+				{
+					ControllerMap[RoleClass] = Controller_DEPRECATED;
+				}
+				else
+				{
+					//Verify in project settings if there is a controller associated with this component type. If not, pick the first one we find
+					TSubclassOf<ULiveLinkControllerBase> SelectedControllerClass = GetControllerClassForRoleClass(RoleClass);
+					SetControllerClassForRole(RoleClass, SelectedControllerClass);
+				}
+			}
 		}
 	}
-	else
+
+	Controller_DEPRECATED = nullptr;
+}
+
+bool ULiveLinkComponentController::IsControllerMapOutdated() const
+{
+	TArray<TSubclassOf<ULiveLinkRole>> SelectedRoleHierarchy = GetSelectedRoleHierarchyClasses(SubjectRepresentation.Role);
+	bool bNeedUpdate = ControllerMap.Num() != SelectedRoleHierarchy.Num();
+	
+	if (!bNeedUpdate)
 	{
-		bCreateAnewController = true;
+		//Check if all map matches class hierarchy
+		for (const TSubclassOf<ULiveLinkRole>& RoleClass : SelectedRoleHierarchy)
+		{
+			const ULiveLinkControllerBase* const* FoundController = ControllerMap.Find(RoleClass);
+			bNeedUpdate = FoundController == nullptr;
+			if (!bNeedUpdate)
+			{
+				const ULiveLinkControllerBase* FoundControllerPtr = *FoundController;
+				TSubclassOf<ULiveLinkControllerBase> DesiredControllerClass = GetControllerClassForRoleClass(RoleClass);
+				if (FoundControllerPtr == nullptr)
+				{
+					bNeedUpdate = DesiredControllerClass != nullptr;
+				}
+				else
+				{
+					bNeedUpdate = FoundControllerPtr->GetClass() != DesiredControllerClass;
+				}
+			}
+
+			if (bNeedUpdate)
+			{
+				break;
+			}
+		}
 	}
 
-	if (bCreateAnewController)
+	return bNeedUpdate;
+}
+
+#endif //WITH_EDITOR
+
+TArray<TSubclassOf<ULiveLinkRole>> ULiveLinkComponentController::GetSelectedRoleHierarchyClasses(const TSubclassOf<ULiveLinkRole> InCurrentRoleClass) const
+{
+	TArray<TSubclassOf<ULiveLinkRole>> ClassHierarchy;
+
+	for (TObjectIterator<UClass> It; It; ++It)
 	{
-		TSubclassOf<ULiveLinkControllerBase> NewControllerClass = ULiveLinkControllerBase::GetControllerForRole(SubjectRepresentation.Role);
-		if (NewControllerClass.Get())
+		if (!It->HasAnyClassFlags(CLASS_Abstract | CLASS_Deprecated))
 		{
-			const EObjectFlags ControllerObjectFlags = GetMaskedFlags(RF_Public | RF_Transactional | RF_ArchetypeObject);
-			Controller = NewObject<ULiveLinkControllerBase>(this, NewControllerClass.Get(), NAME_None, ControllerObjectFlags);
-			Controller->InitializeInEditor();
+			if (InCurrentRoleClass->IsChildOf(*It))
+			{
+				ClassHierarchy.AddUnique(*It);
+			}
 		}
-		else
+	}
+
+	return MoveTemp(ClassHierarchy);
+}
+
+TSubclassOf<ULiveLinkControllerBase> ULiveLinkComponentController::GetControllerClassForRoleClass(const TSubclassOf<ULiveLinkRole> RoleClass) const
+{
+	//Verify in project settings if there is a controller associated with this component type. If not, pick the first one we find that supports that role
+	TSubclassOf<ULiveLinkControllerBase> SelectedControllerClass = nullptr;
+	const TSubclassOf<ULiveLinkControllerBase>* ControllerClass = GetDefault<ULiveLinkComponentSettings>()->DefaultControllerForRole.Find(RoleClass);
+	if (ControllerClass == nullptr || ControllerClass->Get() == nullptr)
+	{
+		TArray<TSubclassOf<ULiveLinkControllerBase>> NewControllerClasses = ULiveLinkControllerBase::GetControllersForRole(RoleClass);
+		if (NewControllerClasses.Num() > 0)
 		{
-			UE_LOG(LogLiveLinkComponents, Warning, TEXT("No controller was found for role '%s'."), *SubjectRepresentation.Role->GetName());
-			FNotificationInfo NotificationInfo(LOCTEXT("NoFoundController", "No controller was found for the role."));
-			NotificationInfo.ExpireDuration = 2.0f;
-			FSlateNotificationManager::Get().AddNotification(NotificationInfo);
+			SelectedControllerClass = NewControllerClasses[0];
 		}
 	}
-	
-	bTickInEditor = bUpdateInEditor;
-	bIsDirty = true;
+	else
+	{
+		SelectedControllerClass = *ControllerClass;
+	}
 
-	Super::PostEditChangeProperty(PropertyChangedEvent);
+	return SelectedControllerClass;
 }
-#endif //WITH_EDITOR
+
+
 
 #undef LOCTEXT_NAMESPACE
\ No newline at end of file
diff --git a/Engine/Plugins/Animation/LiveLink/Source/LiveLinkComponents/Private/LiveLinkControllerBase.cpp b/Engine/Plugins/Animation/LiveLink/Source/LiveLinkComponents/Private/LiveLinkControllerBase.cpp
--- a/Engine/Plugins/Animation/LiveLink/Source/LiveLinkComponents/Private/LiveLinkControllerBase.cpp
+++ b/Engine/Plugins/Animation/LiveLink/Source/LiveLinkComponents/Private/LiveLinkControllerBase.cpp
@@ -4,6 +4,11 @@
 #include "GameFramework/Actor.h"
 #include "UObject/UObjectIterator.h"
 
+void ULiveLinkControllerBase::SetAttachedComponent(UActorComponent* ActorComponent)
+{
+	AttachedComponent = ActorComponent;
+}
+
 AActor* ULiveLinkControllerBase::GetOuterActor() const
 {
 	for (UObject* Outer = GetOuter(); Outer != nullptr; Outer = Outer->GetOuter())
@@ -21,18 +26,32 @@ AActor* ULiveLinkControllerBase::GetOuterActor() const
 	return nullptr;
 }
 
-
 TSubclassOf<ULiveLinkControllerBase> ULiveLinkControllerBase::GetControllerForRole(const TSubclassOf<ULiveLinkRole>& RoleToSupport)
 {
+	TArray<TSubclassOf<ULiveLinkControllerBase>> Controllers = GetControllersForRole(RoleToSupport);
+	if (Controllers.Num() > 0)
+	{
+		return Controllers[0];
+	}
+
+	return TSubclassOf<ULiveLinkControllerBase>();
+}
+
+TArray<TSubclassOf<ULiveLinkControllerBase>> ULiveLinkControllerBase::GetControllersForRole(const TSubclassOf<ULiveLinkRole>& RoleToSupport)
+{
+	TArray<TSubclassOf<ULiveLinkControllerBase>> Controllers;
+
 	for (TObjectIterator<UClass> Itt; Itt; ++Itt)
 	{
 		if (Itt->IsChildOf(ULiveLinkControllerBase::StaticClass()) && !Itt->HasAnyClassFlags(CLASS_Abstract | CLASS_Deprecated))
 		{
 			if (Itt->GetDefaultObject<ULiveLinkControllerBase>()->IsRoleSupported(RoleToSupport))
 			{
-				return TSubclassOf<ULiveLinkControllerBase>(*Itt);
+				Controllers.Add(TSubclassOf<ULiveLinkControllerBase>(*Itt));
 			}
 		}
 	}
-	return TSubclassOf<ULiveLinkControllerBase>();
+
+	return MoveTemp(Controllers);
 }
+
diff --git a/Engine/Plugins/Animation/LiveLink/Source/LiveLinkComponents/Public/Controllers/LiveLinkCameraController.h b/Engine/Plugins/Animation/LiveLink/Source/LiveLinkComponents/Public/Controllers/LiveLinkCameraController.h
--- a/Engine/Plugins/Animation/LiveLink/Source/LiveLinkComponents/Public/Controllers/LiveLinkCameraController.h
+++ b/Engine/Plugins/Animation/LiveLink/Source/LiveLinkComponents/Public/Controllers/LiveLinkCameraController.h
@@ -16,18 +16,22 @@ class LIVELINKCOMPONENTS_API ULiveLinkCameraController : public ULiveLinkControl
 	GENERATED_BODY()
 
 public:
-	UPROPERTY(EditAnywhere, Category="LiveLink", meta=(UseComponentPicker, AllowedClasses="CameraComponent"))
-	FComponentReference ComponentToControl;
+#if WITH_EDITORONLY_DATA
+	UPROPERTY()
+	FComponentReference ComponentToControl_DEPRECATED;
 
-	UPROPERTY(EditAnywhere, Category = "LiveLink", meta = (ShowOnlyInnerProperties))
-	FLiveLinkTransformControllerData TransformData;
+	UPROPERTY()
+	FLiveLinkTransformControllerData TransformData_DEPRECATED;
+#endif
 
 public:
-	virtual void OnEvaluateRegistered() override;
-	virtual void Tick(float DeltaTime, const FLiveLinkSubjectRepresentation& SubjectRepresentation) override;
+	//~ Begin ULiveLinkControllerBase interface
+	virtual void Tick(float DeltaTime, const FLiveLinkSubjectFrameData& SubjectData) override;
 	virtual bool IsRoleSupported(const TSubclassOf<ULiveLinkRole>& RoleToSupport) override;
+	virtual TSubclassOf<UActorComponent> GetDesiredComponentClass() const override;
+	//~ End ULiveLinkControllerBase interface
 
-#if WITH_EDITOR
-	virtual void InitializeInEditor() override;
-#endif
+	//~ Begin UObject interface
+	virtual void PostLoad() override;
+	//~ End UObject interface
 };
\ No newline at end of file
diff --git a/Engine/Plugins/Animation/LiveLink/Source/LiveLinkComponents/Public/Controllers/LiveLinkLightController.h b/Engine/Plugins/Animation/LiveLink/Source/LiveLinkComponents/Public/Controllers/LiveLinkLightController.h
--- a/Engine/Plugins/Animation/LiveLink/Source/LiveLinkComponents/Public/Controllers/LiveLinkLightController.h
+++ b/Engine/Plugins/Animation/LiveLink/Source/LiveLinkComponents/Public/Controllers/LiveLinkLightController.h
@@ -18,18 +18,23 @@ class LIVELINKCOMPONENTS_API ULiveLinkLightController : public ULiveLinkControll
 	GENERATED_BODY()
 
 public:
-	UPROPERTY(EditAnywhere, Category="LiveLink", meta=(UseComponentPicker, AllowedClasses="LightComponent"))
-	FComponentReference ComponentToControl;
+#if WITH_EDITORONLY_DATA
+	UPROPERTY()
+	FComponentReference ComponentToControl_DEPRECATED;
 
-	UPROPERTY(EditAnywhere, Category = "LiveLink", meta=(ShowOnlyInnerProperties))
-	FLiveLinkTransformControllerData TransformData;
+	UPROPERTY()
+	FLiveLinkTransformControllerData TransformData_DEPRECATED;
+#endif
 
 public:
-	virtual void OnEvaluateRegistered() override;
-	virtual void Tick(float DeltaTime, const FLiveLinkSubjectRepresentation& SubjectRepresentation) override;
+
+	//~ Begin ULiveLinkControllerBase interface
+	virtual void Tick(float DeltaTime, const FLiveLinkSubjectFrameData& SubjectData) override;
 	virtual bool IsRoleSupported(const TSubclassOf<ULiveLinkRole>& RoleToSupport) override;
+	virtual TSubclassOf<UActorComponent> GetDesiredComponentClass() const override;
+	//~ End ULiveLinkControllerBase interface
 
-#if WITH_EDITOR
-	virtual void InitializeInEditor() override;
-#endif
+	//~ Begin UObject interface
+	virtual void PostLoad() override;
+	//~ End UObject interface
 };
\ No newline at end of file
diff --git a/Engine/Plugins/Animation/LiveLink/Source/LiveLinkComponents/Public/Controllers/LiveLinkTransformController.h b/Engine/Plugins/Animation/LiveLink/Source/LiveLinkComponents/Public/Controllers/LiveLinkTransformController.h
--- a/Engine/Plugins/Animation/LiveLink/Source/LiveLinkComponents/Public/Controllers/LiveLinkTransformController.h
+++ b/Engine/Plugins/Animation/LiveLink/Source/LiveLinkComponents/Public/Controllers/LiveLinkTransformController.h
@@ -49,19 +49,24 @@ class LIVELINKCOMPONENTS_API ULiveLinkTransformController : public ULiveLinkCont
 	GENERATED_BODY()
 
 public:
-	UPROPERTY(EditAnywhere, Category="LiveLink", meta=(UseComponentPicker, AllowedClasses="SceneComponent"))
-	FComponentReference ComponentToControl;
+#if WITH_EDITORONLY_DATA
+	UPROPERTY()
+	FComponentReference ComponentToControl_DEPRECATED;
+#endif //WITH_EDITORONLY_DATA
 	
 	UPROPERTY(EditAnywhere, Category="LiveLink", meta=(ShowOnlyInnerProperties))
 	FLiveLinkTransformControllerData TransformData;
 
 public:
+	//~Begin ULiveLinkControllerBase interface
 	virtual void OnEvaluateRegistered() override;
-	virtual void Tick(float DeltaTime, const FLiveLinkSubjectRepresentation& SubjectRepresentation) override;
+	virtual void Tick(float DeltaTime, const FLiveLinkSubjectFrameData& SubjectData) override;
 	virtual bool IsRoleSupported(const TSubclassOf<ULiveLinkRole>& RoleToSupport) override;
+	virtual TSubclassOf<UActorComponent> GetDesiredComponentClass() const override;
+	virtual void SetAttachedComponent(UActorComponent* ActorComponent) override;
+	//~End ULiveLinkControllerBase interface
 
-public:
-#if WITH_EDITOR
-	virtual void InitializeInEditor() override;
-#endif
+	//~ Begin UObject interface
+	virtual void PostLoad() override;
+	//~ End UObject interface
 };
\ No newline at end of file
diff --git a/Engine/Plugins/Animation/LiveLink/Source/LiveLinkComponents/Public/LiveLinkComponentController.h b/Engine/Plugins/Animation/LiveLink/Source/LiveLinkComponents/Public/LiveLinkComponentController.h
--- a/Engine/Plugins/Animation/LiveLink/Source/LiveLinkComponents/Public/LiveLinkComponentController.h
+++ b/Engine/Plugins/Animation/LiveLink/Source/LiveLinkComponents/Public/LiveLinkComponentController.h
@@ -6,7 +6,7 @@
 #include "ILiveLinkClient.h"
 #include "LiveLinkComponentController.generated.h"
 
-DECLARE_DYNAMIC_MULTICAST_DELEGATE_OneParam(FLiveLinkTickSignature, float, DeltaTime);
+DECLARE_DYNAMIC_MULTICAST_DELEGATE_OneParam(FLiveLinkTickDelegate, float, DeltaTime);
 
 class ULiveLinkControllerBase;
 
@@ -22,30 +22,66 @@ public:
 	UPROPERTY(BlueprintReadOnly, EditAnywhere, Category="LiveLink")
 	FLiveLinkSubjectRepresentation SubjectRepresentation;
 
-	UPROPERTY(BlueprintReadOnly, VisibleAnywhere, Category="LiveLink", Instanced, NoClear, meta=(ShowOnlyInnerProperties))
-	ULiveLinkControllerBase* Controller;
+#if WITH_EDITORONLY_DATA
+	UPROPERTY(Instanced, NoClear)
+	ULiveLinkControllerBase* Controller_DEPRECATED;
+#endif
+
+	/** Instanced controllers used to control the desired role */
+	UPROPERTY(BlueprintReadOnly, EditAnywhere, Category = "LiveLink", Instanced, NoClear, meta = (ShowOnlyInnerProperties))
+	TMap<TSubclassOf<ULiveLinkRole>, ULiveLinkControllerBase*> ControllerMap;
 
 	UPROPERTY(EditAnywhere, Category="LiveLink", AdvancedDisplay)
 	bool bUpdateInEditor;
-
-public:
-	virtual void OnRegister() override;
 	
-	// Called every frame
-	virtual void TickComponent(float DeltaTime, ELevelTick TickType, FActorComponentTickFunction* ThisTickFunction) override;
-
 	// This Event is triggered any time new LiveLink data is available, including in the editor
 	UPROPERTY(BlueprintAssignable, Category = "LiveLink")
-	FLiveLinkTickSignature OnLiveLinkUpdated;
+	FLiveLinkTickDelegate OnLiveLinkUpdated;
+
+	UPROPERTY(EditInstanceOnly, Category = "LiveLink", meta = (UseComponentPicker, AllowedClasses = "ActorComponent", DisallowedClasses = "LiveLinkComponentController"))
+	FComponentReference ComponentToControl;
+
+protected:
+	// Keep track when component gets registered or controller map gets changed
+	bool bIsDirty;
 
 public:
+	
+	/** Creates an instance of the desired controller class for a specified Role class */
+	void SetControllerClassForRole(TSubclassOf<ULiveLinkRole> RoleClass, TSubclassOf<ULiveLinkControllerBase> DesiredControllerClass);
+
+	/** Creates an instance of the desired controller class for a specified Role class */
+	FLiveLinkSubjectRepresentation GetSubjectRepresentation() const { return SubjectRepresentation; }
+
+	/** Sets SubjectRepresentation and if required, will update the Controllers map associated to the Role */
+	void SetSubjectRepresentation(const FLiveLinkSubjectRepresentation& InSubjectRepresentation);
+	
+	/** Returns true if ControllerMap needs to be updated for the current Role. Useful for customization or C++ modification to the Role */
+	bool IsControllerMapOutdated() const;
+	
+	/** Used to notify that the subject role has changed. Mainly from Customization or C++ modification to the subject's Role */
+	void OnSubjectRoleChanged();
+
+	//~ Begin UActorComponent interface
+	virtual void OnRegister() override;
+	virtual void TickComponent(float DeltaTime, ELevelTick TickType, FActorComponentTickFunction* ThisTickFunction) override;
+	//~ End UActorComponent interface
+
 	//~ UObject interface
+	virtual void Serialize(FArchive& Ar) override;
 #if WITH_EDITOR
 	virtual void PostEditChangeProperty(struct FPropertyChangedEvent& PropertyChangedEvent) override;
 #endif //WITH_EDITOR
 	//~ End UObject interface
 
-private:
-	// Record whether we have been recently registered
-	bool bIsDirty;
+protected:
+
+	/** Returns an array representing the class hierarchy of the given class */
+	TArray<TSubclassOf<ULiveLinkRole>> GetSelectedRoleHierarchyClasses(const TSubclassOf<ULiveLinkRole> InCurrentRoleClass) const;
+	TSubclassOf<ULiveLinkControllerBase> GetControllerClassForRoleClass(const TSubclassOf<ULiveLinkRole> RoleClass) const;
+
+#if WITH_EDITOR
+	/** Called during loading to convert old data to new scheme. */
+	void ConvertOldControllerSystem();
+#endif //WITH_EDITOR
 };
\ No newline at end of file
diff --git a/Engine/Plugins/Animation/LiveLink/Source/LiveLinkComponents/Public/LiveLinkComponentSettings.h b/Engine/Plugins/Animation/LiveLink/Source/LiveLinkComponents/Public/LiveLinkComponentSettings.h
new file mode 100644
--- /dev/null
+++ b/Engine/Plugins/Animation/LiveLink/Source/LiveLinkComponents/Public/LiveLinkComponentSettings.h
@@ -0,0 +1,28 @@
+// Copyright 1998-2019 Epic Games, Inc. All Rights Reserved.
+
+#pragma once
+
+#include "UObject/Object.h"
+#include "UObject/ObjectMacros.h"
+
+#include "LiveLinkRole.h"
+#include "LiveLinkControllerBase.h"
+#include "Templates/SubclassOf.h"
+
+#include "LiveLinkComponentSettings.generated.h"
+
+
+/**
+ * Settings for LiveLink.
+ */
+UCLASS(config=Game, defaultconfig)
+class LIVELINKCOMPONENTS_API ULiveLinkComponentSettings : public UObject
+{
+	GENERATED_BODY()
+
+public:
+
+	/** Default Controller class to use for the specified role */
+	UPROPERTY(config, EditAnywhere, Category = "LiveLink", meta = (AllowAbstract = "false"))
+	TMap<TSubclassOf<ULiveLinkRole>, TSubclassOf<ULiveLinkControllerBase>> DefaultControllerForRole;
+};
diff --git a/Engine/Plugins/Animation/LiveLink/Source/LiveLinkComponents/Public/LiveLinkControllerBase.h b/Engine/Plugins/Animation/LiveLink/Source/LiveLinkComponents/Public/LiveLinkControllerBase.h
--- a/Engine/Plugins/Animation/LiveLink/Source/LiveLinkComponents/Public/LiveLinkControllerBase.h
+++ b/Engine/Plugins/Animation/LiveLink/Source/LiveLinkComponents/Public/LiveLinkControllerBase.h
@@ -25,11 +25,13 @@ public:
 	/** Initialize the controller at the first tick of his owner component. */
 	virtual void OnEvaluateRegistered() { }
 
+	UE_DEPRECATED(4.25, "This function is deprecated. Use Tick function that received evaluated data instead.")
+	virtual void Tick(float DeltaTime, const FLiveLinkSubjectRepresentation& SubjectRepresentation) { }
+
 	/**
-	 * Function called every frame.
-	 * Will only tick the when ShouldTickController returned true.
+	 * Function called every frame with the data evaluated by the component.
 	 */
-	virtual void Tick(float DeltaTime, const FLiveLinkSubjectRepresentation& SubjectRepresentation) { }
+	virtual void Tick(float DeltaTime, const FLiveLinkSubjectFrameData& SubjectData) { }
 
 	/**
 	 * Can it support a specific role.
@@ -37,6 +39,16 @@ public:
 	 */
 	virtual bool IsRoleSupported(const TSubclassOf<ULiveLinkRole>& RoleToSupport) { return false; }
 
+	/**
+	 * Returns the component class that this controller wants to control
+	 */
+	virtual TSubclassOf<UActorComponent> GetDesiredComponentClass() const { return UActorComponent::StaticClass(); }
+
+	/**
+	 * Sets the component this controller is driving
+	 */
+	virtual void SetAttachedComponent(UActorComponent* ActorComponent);
+
 #if WITH_EDITOR
 	virtual void InitializeInEditor() {}
 #endif
@@ -45,5 +57,16 @@ protected:
 	AActor* GetOuterActor() const;
 
 public:
+
+	UE_DEPRECATED(4.25, "This function is deprecated. Use GetControllersForRole instead and use first element to have the same result.")
 	static TSubclassOf<ULiveLinkControllerBase> GetControllerForRole(const TSubclassOf<ULiveLinkRole>& RoleToSupport);
-};
\ No newline at end of file
+
+	/**
+	 * Returns the list of ULiveLinkControllerBase classes that support the given role
+	 */
+	static TArray<TSubclassOf<ULiveLinkControllerBase>> GetControllersForRole(const TSubclassOf<ULiveLinkRole>& RoleToSupport);
+
+protected:
+	TWeakObjectPtr<UActorComponent> AttachedComponent;
+};
+
diff --git a/Engine/Plugins/Animation/LiveLink/Source/LiveLinkEditor/LiveLinkEditor.Build.cs b/Engine/Plugins/Animation/LiveLink/Source/LiveLinkEditor/LiveLinkEditor.Build.cs
--- a/Engine/Plugins/Animation/LiveLink/Source/LiveLinkEditor/LiveLinkEditor.Build.cs
+++ b/Engine/Plugins/Animation/LiveLink/Source/LiveLinkEditor/LiveLinkEditor.Build.cs
@@ -30,6 +30,7 @@ namespace UnrealBuildTool.Rules
 					"InputCore",
 					"KismetCompiler",
 					"GraphEditor",
+					"LiveLinkComponents",
 					"LiveLinkGraphNode",
 					"LiveLinkMessageBusFramework",
 					"LiveLinkMovieScene",
diff --git a/Engine/Plugins/Animation/LiveLink/Source/LiveLinkEditor/Private/LiveLinkComponentDetailCustomization.cpp b/Engine/Plugins/Animation/LiveLink/Source/LiveLinkEditor/Private/LiveLinkComponentDetailCustomization.cpp
new file mode 100644
--- /dev/null
+++ b/Engine/Plugins/Animation/LiveLink/Source/LiveLinkEditor/Private/LiveLinkComponentDetailCustomization.cpp
@@ -0,0 +1,413 @@
+// Copyright 1998-2019 Epic Games, Inc. All Rights Reserved.
+
+#include "LiveLinkComponentDetailCustomization.h"
+
+#include "DetailCategoryBuilder.h"
+#include "DetailLayoutBuilder.h"
+#include "DetailWidgetRow.h"
+#include "EditorFontGlyphs.h"
+#include "IDetailGroup.h"
+#include "Framework/MultiBox/MultiBoxBuilder.h"
+#include "LiveLinkControllerBase.h"
+#include "LiveLinkEditorPrivate.h"
+#include "ScopedTransaction.h"
+
+#include "Widgets/Input/SComboButton.h"
+
+#define LOCTEXT_NAMESPACE "LiveLinkComponentDetailsCustomization"
+
+
+void FLiveLinkComponentDetailCustomization::CustomizeDetails(IDetailLayoutBuilder& DetailBuilder)
+{
+	DetailLayout = &DetailBuilder;
+
+	TArray<TWeakObjectPtr<UObject>> SelectedObjects = DetailBuilder.GetSelectedObjects();
+
+	//Hide everything when more than one are selected
+	if (SelectedObjects.Num() != 1)
+	{
+		for (TWeakObjectPtr<UObject>& SelectedObject : SelectedObjects)
+		{
+			if (ULiveLinkComponentController* SelectedPtr = Cast<ULiveLinkComponentController>(SelectedObject.Get()))
+			{
+				TSharedRef<IPropertyHandle> ControllersProperty = DetailBuilder.GetProperty(GET_MEMBER_NAME_CHECKED(ULiveLinkComponentController, ControllerMap));
+				ControllersProperty->MarkHiddenByCustomization();
+			}
+		}
+
+		return;
+	}
+
+	if (ULiveLinkComponentController* SelectedPtr = Cast<ULiveLinkComponentController>(SelectedObjects[0].Get()))
+	{
+		//Register callback when LiveLinkSubjectRepresentation selection has changed to refresh the UI and update controller
+		TSharedRef<IPropertyHandle> SubjectRepresentationPropertyRef = DetailBuilder.GetProperty(GET_MEMBER_NAME_CHECKED(ULiveLinkComponentController, SubjectRepresentation));
+		SubjectRepresentationPropertyRef->MarkHiddenByCustomization();
+		DetailLayout->AddPropertyToCategory(SubjectRepresentationPropertyRef);
+
+		//Hide the Map default UI
+		TSharedRef<IPropertyHandle> ControllersProperty = DetailBuilder.GetProperty(GET_MEMBER_NAME_CHECKED(ULiveLinkComponentController, ControllerMap));
+		ControllersProperty->MarkHiddenByCustomization();
+
+		//Get hook to the controller map. If that fails, early exit
+		TSharedPtr<IPropertyHandleMap> MapHandle = ControllersProperty->AsMap();
+		if (MapHandle.IsValid())
+		{
+			EditedObject = Cast<ULiveLinkComponentController>(SelectedObjects[0].Get());
+
+			SubjectRepresentationPropertyRef->SetOnPropertyValueChanged(FSimpleDelegate::CreateSP(this, &FLiveLinkComponentDetailCustomization::OnSubjectRepresentationPropertyChanged));
+			SubjectRepresentationPropertyRef->SetOnChildPropertyValueChanged(FSimpleDelegate::CreateSP(this, &FLiveLinkComponentDetailCustomization::OnSubjectRepresentationPropertyChanged));
+
+			//Listen to controller map modifications to refresh UI when a change comes through MU
+			FSimpleDelegate RefreshDelegate = FSimpleDelegate::CreateSP(this, &FLiveLinkComponentDetailCustomization::ForceRefreshDetails);
+			MapHandle->SetOnNumElementsChanged(RefreshDelegate);
+
+			//Loop for each entry in the map. 
+			//Fetch the LiveLinkRole name (key) and display its name
+			//Add a menu to select a ControllerClass for it
+			//If a Controller is picked, display its properties
+			uint32 NumEntry = 0;
+			ControllersProperty->GetNumChildren(NumEntry);
+			for (uint32 EntryIndex = 0; EntryIndex < NumEntry; ++EntryIndex)
+			{
+				TSharedPtr<IPropertyHandle> EntryHandle = ControllersProperty->GetChildHandle(EntryIndex);
+				if (!EntryHandle.IsValid())
+				{
+					continue;
+				}
+
+				//Map has a TSubClassof Key type
+				TSubclassOf<ULiveLinkRole> LiveLinkRoleClass;
+				UObject* LiveLinkRoleClassObj = nullptr;
+				TSharedPtr<IPropertyHandle> KeyHandle = EntryHandle->GetKeyHandle();
+
+				//Make sure we were able to query the UClass and that it's not null. 
+				check(KeyHandle->GetValue(LiveLinkRoleClassObj) == FPropertyAccess::Success)
+				LiveLinkRoleClass = Cast<UClass>(LiveLinkRoleClassObj);
+				if (LiveLinkRoleClass == nullptr)
+				{
+					continue;
+				}
+
+				FText ControllerName = FText::Format(LOCTEXT("No Controller", "{0}"), FText::FromName(NAME_None));
+
+				//Map value is a pointer to the Controller. Can be null so it could be empty
+				uint32 ControllerValueCount = 0;
+				EntryHandle->GetNumChildren(ControllerValueCount);
+				TSharedPtr<IPropertyHandle> ValueHandle;
+				if (ControllerValueCount > 0)
+				{
+					ValueHandle = EntryHandle->GetChildHandle(0);
+
+					UObject* ControllerPtr = nullptr;
+					EntryHandle->GetValue(ControllerPtr);
+					if (ControllerPtr != nullptr)
+					{
+						ControllerName = ControllerPtr->GetClass()->GetDisplayNameText();
+					}
+				}
+
+				//Since we're displaying properties of another object, add it as external to the current one being edited. 
+				//Add each Map entry as in the row that is generated for that external objects : Name = Key (Component name) + Value = Dropdown menu with available controllers
+				TArray<UObject*> ExternalObjects;
+				if (EditedObject->ControllerMap.Contains(LiveLinkRoleClass))
+				{
+					ExternalObjects.Add(EditedObject->ControllerMap[LiveLinkRoleClass]);
+				}
+
+				IDetailPropertyRow* ThisRoleRow = DetailBuilder.EditCategory(TEXT("Role Controllers")).AddExternalObjects(ExternalObjects);
+				if (ThisRoleRow)
+				{
+					ThisRoleRow->CustomWidget()
+					.NameContent()
+					[
+						BuildControllerNameWidget(ControllersProperty, LiveLinkRoleClass)
+					]
+					.ValueContent()
+					[
+						BuildControllerValueWidget(KeyHandle, LiveLinkRoleClass, ControllerName)
+					];
+				}
+			}
+
+			// Start by looking if data is dirty as we enter. Can happen when component lives in a blueprint
+			OnSubjectRepresentationPropertyChanged();
+		}
+	}
+}
+
+void FLiveLinkComponentDetailCustomization::OnSubjectRepresentationPropertyChanged()
+{
+	if (ULiveLinkComponentController* EditedObjectPtr = EditedObject.Get())
+	{
+		//Verify if Role has changed
+		if (EditedObjectPtr->IsControllerMapOutdated())
+		{
+			const FScopedTransaction Transaction(LOCTEXT("OnChangedSubjectRepresentation", "Subject Representation Changed"));
+			EditedObjectPtr->Modify();
+
+			EditedObjectPtr->OnSubjectRoleChanged();
+		}
+	}
+}
+
+TSharedRef<SWidget> FLiveLinkComponentDetailCustomization::HandleControllerComboButton(TSharedPtr<IPropertyHandle> KeyHandle) const
+{
+	// Generate menu
+	FMenuBuilder MenuBuilder(true, nullptr);
+	MenuBuilder.BeginSection("SupportedControllers", LOCTEXT("SupportedControllers", "Controllers"));
+	{
+		if (KeyHandle.IsValid())
+		{
+			//Make sure we were able to query the UClass and that it's not null. 
+			TSubclassOf<ULiveLinkRole> RoleClass;
+			UObject* RoleClassObj = nullptr;
+			KeyHandle->GetValue(RoleClassObj);
+			RoleClass = Cast<UClass>(RoleClassObj);
+			if(RoleClass)
+			{
+				TArray<TSubclassOf<ULiveLinkControllerBase>> NewControllerClasses = ULiveLinkControllerBase::GetControllersForRole(RoleClass);
+				if (NewControllerClasses.Num() > 0)
+				{
+					//Always add a None entry
+					MenuBuilder.AddMenuEntry(
+						FText::FromName(NAME_None),
+						FText::FromName(NAME_None),
+						FSlateIcon(),
+						FUIAction(
+							FExecuteAction::CreateSP(this, &FLiveLinkComponentDetailCustomization::HandleControllerSelection, RoleClass, TWeakObjectPtr<UClass>()),
+							FCanExecuteAction(),
+							FIsActionChecked::CreateSP(this, &FLiveLinkComponentDetailCustomization::IsControllerItemSelected, FName(), RoleClass)
+						),
+						NAME_None,
+						EUserInterfaceActionType::RadioButton
+					);
+
+					for (const TSubclassOf<ULiveLinkControllerBase>& ControllerClass : NewControllerClasses)
+					{
+						if (UClass* ControllerClassPtr = ControllerClass.Get())
+						{
+							MenuBuilder.AddMenuEntry(
+								FText::Format(LOCTEXT("Controller Label", "{0}"), ControllerClassPtr->GetDisplayNameText()),
+								FText::Format(LOCTEXT("Controller ToolTip", "{0}"), ControllerClassPtr->GetDisplayNameText()),
+								FSlateIcon(),
+								FUIAction(
+									FExecuteAction::CreateSP(this, &FLiveLinkComponentDetailCustomization::HandleControllerSelection, RoleClass, MakeWeakObjectPtr(ControllerClassPtr)),
+									FCanExecuteAction(),
+									FIsActionChecked::CreateSP(this, &FLiveLinkComponentDetailCustomization::IsControllerItemSelected, ControllerClassPtr->GetFName(), RoleClass)
+								),
+								NAME_None,
+								EUserInterfaceActionType::RadioButton
+							);
+						}
+
+					}
+				}
+				else
+				{
+					MenuBuilder.AddWidget(SNullWidget::NullWidget, LOCTEXT("NoControllersFound", "No Controllers were found for this role"), false, false);
+				}
+			}
+			else
+			{
+				MenuBuilder.AddWidget(SNullWidget::NullWidget, LOCTEXT("InvalidRoleClass", "Role is invalid. Can't find controllers for it"), false, false);
+			}
+		}
+		else
+		{
+			MenuBuilder.AddWidget(SNullWidget::NullWidget, LOCTEXT("InvalidComponent", "Invalid component class handle received"), false, false);
+		}
+	}
+	MenuBuilder.EndSection();
+
+	return MenuBuilder.MakeWidget();
+}
+
+void FLiveLinkComponentDetailCustomization::HandleControllerSelection(TSubclassOf<ULiveLinkRole> RoleClass, TWeakObjectPtr<UClass> SelectedControllerClass) const
+{
+	if (ULiveLinkComponentController* EditedObjectPtr = EditedObject.Get())
+	{
+		const FScopedTransaction Transaction(LOCTEXT("OnChangedController", "Property Controller Selection"));
+		EditedObjectPtr->Modify();
+
+		UClass* SelectedControllerClassPtr = SelectedControllerClass.Get();
+		EditedObjectPtr->SetControllerClassForRole(RoleClass, SelectedControllerClassPtr);
+	}
+}
+
+bool FLiveLinkComponentDetailCustomization::IsControllerItemSelected(FName Item, TSubclassOf<ULiveLinkRole> RoleClass) const
+{
+	if (RoleClass == nullptr || RoleClass.Get() == nullptr)
+	{
+		return false;
+	}
+
+	if (ULiveLinkComponentController* EditedObjectPtr = EditedObject.Get())
+	{
+
+		ULiveLinkControllerBase** CurrentClass = EditedObjectPtr->ControllerMap.Find(RoleClass);
+		if (CurrentClass == nullptr || *CurrentClass == nullptr)
+		{
+			return Item.IsNone();
+		}
+		else
+		{
+			return (*CurrentClass)->GetClass()->GetFName() == Item;
+		}
+	}
+
+	return false;
+}
+
+FSlateColor FLiveLinkComponentDetailCustomization::HandleControllerStatusColorAndOpacity(TSubclassOf<ULiveLinkRole> RoleClassEntry) const
+{
+	FSlateColor Result = FLinearColor::Green;
+	bool bIsValid = true;
+
+	if (ULiveLinkComponentController* EditorObjectPtr = EditedObject.Get())
+	{
+		ULiveLinkControllerBase** AssociatedControllerPtr = EditorObjectPtr->ControllerMap.Find(RoleClassEntry);
+		if (AssociatedControllerPtr && *AssociatedControllerPtr)
+		{
+			if (UActorComponent* SelectedComponent = EditorObjectPtr->ComponentToControl.GetComponent(EditorObjectPtr->GetOwner()))
+			{
+				const ULiveLinkControllerBase* AssociatedController = *AssociatedControllerPtr;
+				if (!SelectedComponent->IsA(AssociatedController->GetDesiredComponentClass()))
+				{
+					//Controller exists for RoleClass and desired component is not the kind it wants
+					bIsValid = false;
+				}
+			}
+			else
+			{
+				//Component is not valid
+				bIsValid = false;
+			}
+		}
+	}
+
+	if (!bIsValid)
+	{
+		Result = FLinearColor::Red;
+	}
+
+	return Result;
+}
+
+FText FLiveLinkComponentDetailCustomization::HandleControllerStatusText(TSubclassOf<ULiveLinkRole> RoleClassEntry) const
+{
+	FText Result;
+	bool bIsValid = true;
+
+	if (ULiveLinkComponentController* EditorObjectPtr = EditedObject.Get())
+	{
+		ULiveLinkControllerBase** AssociatedControllerPtr = EditorObjectPtr->ControllerMap.Find(RoleClassEntry);
+		if (AssociatedControllerPtr && *AssociatedControllerPtr)
+		{
+			if (UActorComponent* SelectedComponent = EditorObjectPtr->ComponentToControl.GetComponent(EditorObjectPtr->GetOwner()))
+			{
+				const ULiveLinkControllerBase* AssociatedController = *AssociatedControllerPtr;
+				if (!SelectedComponent->IsA(AssociatedController->GetDesiredComponentClass()))
+				{
+					//Controller exists for RoleClass and desired component is not the kind it wants
+					bIsValid = false;
+				}
+			}
+			else
+			{
+				//Component is not valid
+				bIsValid = false;
+			}
+		}
+	}
+
+	if (!bIsValid)
+	{
+		Result = FEditorFontGlyphs::Ban;
+	}
+
+	return Result;
+}
+
+FText FLiveLinkComponentDetailCustomization::HandleControllerStatusToolTipText(TSubclassOf<ULiveLinkRole> RoleClassEntry) const
+{
+	FText Result;
+	bool bIsValid = true;
+
+	if (ULiveLinkComponentController* EditorObjectPtr = EditedObject.Get())
+	{
+		ULiveLinkControllerBase** AssociatedControllerPtr = EditorObjectPtr->ControllerMap.Find(RoleClassEntry);
+		if (AssociatedControllerPtr && *AssociatedControllerPtr)
+		{
+			if (UActorComponent* SelectedComponent = EditorObjectPtr->ComponentToControl.GetComponent(EditorObjectPtr->GetOwner()))
+			{
+				const ULiveLinkControllerBase* AssociatedController = *AssociatedControllerPtr;
+				if (!SelectedComponent->IsA(AssociatedController->GetDesiredComponentClass()))
+				{
+					//Controller exists for RoleClass and desired component is not the kind it wants
+					bIsValid = false;
+				}
+			}
+			else
+			{
+				//Component is not valid
+				bIsValid = false;
+			}
+		}
+	}
+
+	if (!bIsValid)
+	{
+		Result = LOCTEXT("ControllerToolTip", "Controller can't control selected component");
+	}
+
+	return Result;
+}
+
+TSharedRef<SWidget> FLiveLinkComponentDetailCustomization::BuildControllerNameWidget(TSharedPtr<IPropertyHandle> ControllersProperty, TSubclassOf<ULiveLinkRole> RoleClass) const
+{
+	return ControllersProperty->CreatePropertyNameWidget(RoleClass.Get()->GetDisplayNameText());
+}
+
+TSharedRef<SWidget> FLiveLinkComponentDetailCustomization::BuildControllerValueWidget(TSharedPtr<IPropertyHandle> RoleKeyPropertyHandle, TSubclassOf<ULiveLinkRole> RoleClass, const FText& ControllerName) const
+{
+	TSharedRef<SWidget> ValueWidget = 
+		SNew(SHorizontalBox)
+		+ SHorizontalBox::Slot()
+		.AutoWidth()
+		.Padding(0, 0, 0, 0)
+		.HAlign(HAlign_Left)
+		.VAlign(VAlign_Center)
+		[
+			SNew(STextBlock)
+			.Text(ControllerName)
+		]
+		+ SHorizontalBox::Slot()
+		.AutoWidth()
+		.Padding(FMargin(4.0f, 0.0f, 0.0f, 0.0f))
+		[
+			SNew(SComboButton)
+			.OnGetMenuContent(this, &FLiveLinkComponentDetailCustomization::HandleControllerComboButton, RoleKeyPropertyHandle)
+			.ContentPadding(FMargin(4.0, 2.0))
+		]
+		+ SHorizontalBox::Slot()
+		.AutoWidth()
+		.Padding(FMargin(4.0f, 0.0f, 0.0f, 0.0f))
+		.HAlign(EHorizontalAlignment::HAlign_Left)
+		[
+			SNew(STextBlock)
+			.Font(FEditorStyle::Get().GetFontStyle("FontAwesome.11"))
+			.Text(this, &FLiveLinkComponentDetailCustomization::HandleControllerStatusText, RoleClass)
+			.ColorAndOpacity(this, &FLiveLinkComponentDetailCustomization::HandleControllerStatusColorAndOpacity, RoleClass)
+			.ToolTipText(this, &FLiveLinkComponentDetailCustomization::HandleControllerStatusToolTipText, RoleClass)
+		];
+
+	return ValueWidget;
+}
+
+void FLiveLinkComponentDetailCustomization::ForceRefreshDetails()
+{
+	DetailLayout->ForceRefreshDetails();
+}
+
+#undef LOCTEXT_NAMESPACE
diff --git a/Engine/Plugins/Animation/LiveLink/Source/LiveLinkEditor/Private/LiveLinkComponentDetailCustomization.h b/Engine/Plugins/Animation/LiveLink/Source/LiveLinkEditor/Private/LiveLinkComponentDetailCustomization.h
new file mode 100644
--- /dev/null
+++ b/Engine/Plugins/Animation/LiveLink/Source/LiveLinkEditor/Private/LiveLinkComponentDetailCustomization.h
@@ -0,0 +1,51 @@
+// Copyright 1998-2019 Epic Games, Inc. All Rights Reserved.
+
+#pragma once
+
+#include "IDetailCustomization.h"
+
+#include "CoreMinimal.h"
+
+#include "LiveLinkComponentController.h"
+
+
+class IDetailLayoutBuilder;
+class IPropertyHandle;
+class SWidget;
+
+
+/**
+* Customizes a ULiveLinkComponentController details
+*/
+class FLiveLinkComponentDetailCustomization : public IDetailCustomization
+{
+public:
+
+	static TSharedRef<IDetailCustomization> MakeInstance()
+	{
+		return MakeShared<FLiveLinkComponentDetailCustomization>();
+	}
+
+	// IDetailCustomization interface
+	virtual void CustomizeDetails(class IDetailLayoutBuilder& DetailBuilder) override;
+	// End IDetailCustomization interface
+
+protected:
+	void OnSubjectRepresentationPropertyChanged();
+	TSharedRef<SWidget> HandleControllerComboButton(TSharedPtr<IPropertyHandle> KeyHandle) const;
+	void HandleControllerSelection(TSubclassOf<ULiveLinkRole> RoleClass, TWeakObjectPtr<UClass> SelectedControllerClass) const;
+	bool IsControllerItemSelected(FName Item, TSubclassOf<ULiveLinkRole> RoleClass) const;
+	FSlateColor HandleControllerStatusColorAndOpacity(TSubclassOf<ULiveLinkRole> RoleClassEntry) const;
+	FText HandleControllerStatusText(TSubclassOf<ULiveLinkRole> RoleClassEntry) const;
+	FText HandleControllerStatusToolTipText(TSubclassOf<ULiveLinkRole> RoleClassEntry) const;
+	TSharedRef<SWidget> BuildControllerNameWidget(TSharedPtr<IPropertyHandle> ControllersProperty, TSubclassOf<ULiveLinkRole> RoleClass) const;
+	TSharedRef<SWidget> BuildControllerValueWidget(TSharedPtr<IPropertyHandle> RoleKeyPropertyHandle, TSubclassOf<ULiveLinkRole> RoleClass, const FText& ControllerName) const;
+	void ForceRefreshDetails();
+
+protected:
+	/** LiveLinkComponent on which we're acting */
+	TWeakObjectPtr<ULiveLinkComponentController> EditedObject;
+
+	/** Keep a reference to force refresh the layout */
+	IDetailLayoutBuilder* DetailLayout = nullptr;
+};
diff --git a/Engine/Plugins/Animation/LiveLink/Source/LiveLinkEditor/Private/LiveLinkEditorModule.cpp b/Engine/Plugins/Animation/LiveLink/Source/LiveLinkEditor/Private/LiveLinkEditorModule.cpp
--- a/Engine/Plugins/Animation/LiveLink/Source/LiveLinkEditor/Private/LiveLinkEditorModule.cpp
+++ b/Engine/Plugins/Animation/LiveLink/Source/LiveLinkEditor/Private/LiveLinkEditorModule.cpp
@@ -20,6 +20,9 @@
 #include "WorkspaceMenuStructure.h"
 #include "WorkspaceMenuStructureModule.h"
 
+#include "LiveLinkComponentController.h"
+#include "LiveLinkComponentDetailCustomization.h"
+#include "LiveLinkComponentSettings.h"
 #include "LiveLinkClient.h"
 #include "LiveLinkClientPanel.h"
 #include "LiveLinkClientCommands.h"
@@ -207,6 +210,12 @@ private:
 				LOCTEXT("LiveLinkDescription", "Configure the Live Link plugin."),
 				GetMutableDefault<ULiveLinkSettings>()
 			);
+
+			SettingsModule->RegisterSettings("Project", "Plugins", "LiveLinkComponent",
+				LOCTEXT("LiveLinkComponentSettingsName", "Live Link Component"),
+				LOCTEXT("LiveLinkComponentDescription", "Configure the Live Link Component."),
+				GetMutableDefault<ULiveLinkComponentSettings>()
+			);
 		}
 	}
 
@@ -216,6 +225,7 @@ private:
 		ISettingsModule* SettingsModule = FModuleManager::GetModulePtr<ISettingsModule>("Settings");
 		if (SettingsModule != nullptr)
 		{
+			SettingsModule->UnregisterSettings("Project", "Plugins", "LiveLinkComponent");
 			SettingsModule->UnregisterSettings("Project", "Plugins", "LiveLink");
 		}
 	}
@@ -227,6 +237,7 @@ private:
 		PropertyEditorModule.RegisterCustomPropertyTypeLayout(FLiveLinkSubjectRepresentation::StaticStruct()->GetFName(), FOnGetPropertyTypeCustomizationInstance::CreateStatic(&FLiveLinkSubjectRepresentationDetailCustomization::MakeInstance));
 		PropertyEditorModule.RegisterCustomPropertyTypeLayout(FLiveLinkSubjectName::StaticStruct()->GetFName(), FOnGetPropertyTypeCustomizationInstance::CreateStatic(&FLiveLinkSubjectNameDetailCustomization::MakeInstance));
 		PropertyEditorModule.RegisterCustomClassLayout(ULiveLinkSourceSettings::StaticClass()->GetFName(), FOnGetDetailCustomizationInstance::CreateStatic(&FLiveLinkSourceSettingsDetailCustomization::MakeInstance));
+		PropertyEditorModule.RegisterCustomClassLayout(ULiveLinkComponentController::StaticClass()->GetFName(), FOnGetDetailCustomizationInstance::CreateStatic(&FLiveLinkComponentDetailCustomization::MakeInstance));
 
 		LiveLinkGraphPanelPinFactory = MakeShared<FLiveLinkGraphPanelPinFactory>();
 		FEdGraphUtilities::RegisterVisualPinFactory(LiveLinkGraphPanelPinFactory);
@@ -240,6 +251,7 @@ private:
 			FPropertyEditorModule* PropertyEditorModule = FModuleManager::Get().GetModulePtr<FPropertyEditorModule>("PropertyEditor");
 			if (PropertyEditorModule)
 			{
+				PropertyEditorModule->UnregisterCustomClassLayout(ULiveLinkComponentController::StaticClass()->GetFName());
 				PropertyEditorModule->UnregisterCustomClassLayout(ULiveLinkSourceSettings::StaticClass()->GetFName());
 				PropertyEditorModule->UnregisterCustomPropertyTypeLayout(FLiveLinkSubjectName::StaticStruct()->GetFName());
 				PropertyEditorModule->UnregisterCustomPropertyTypeLayout(FLiveLinkSubjectRepresentation::StaticStruct()->GetFName());
diff --git a/Engine/Source/Runtime/Core/Public/UObject/EnterpriseObjectVersion.h b/Engine/Source/Runtime/Core/Public/UObject/EnterpriseObjectVersion.h
--- a/Engine/Source/Runtime/Core/Public/UObject/EnterpriseObjectVersion.h
+++ b/Engine/Source/Runtime/Core/Public/UObject/EnterpriseObjectVersion.h
@@ -35,6 +35,9 @@ struct CORE_API FEnterpriseObjectVersion
 		
 		// MeshDescriptionBulkData contains a bGuidIsHash so we can benefit from DDC caching.
 		MeshDescriptionBulkDataGuidIsHash,
+		
+		// Splitting controller associated with a livelink component for each LiveLinkRole class hierarchy
+		LiveLinkControllerSplitPerRole,
 
 		// -----<new versions can be added above this line>-------------------------------------------------
 		VersionPlusOne,
