Fix LiveLink switch pre processor

diff --git a/Engine/Plugins/Animation/LiveLink/Source/LiveLink/Private/PreProcessor/LiveLinkAxisSwitchPreProcessor.cpp b/Engine/Plugins/Animation/LiveLink/Source/LiveLink/Private/PreProcessor/LiveLinkAxisSwitchPreProcessor.cpp
--- a/Engine/Plugins/Animation/LiveLink/Source/LiveLink/Private/PreProcessor/LiveLinkAxisSwitchPreProcessor.cpp
+++ b/Engine/Plugins/Animation/LiveLink/Source/LiveLink/Private/PreProcessor/LiveLinkAxisSwitchPreProcessor.cpp
@@ -18,7 +18,7 @@ namespace
 
 	float AxisSign(ELiveLinkAxis InAxis)
 	{
-		return static_cast<uint8>(InAxis) < 3 ? 1.f : 0.f;
+		return static_cast<uint8>(InAxis) < 3 ? 1.f : -1.f;
 	}
 
 	void SwitchTransform(FTransform& Transform, ELiveLinkAxis AxisX, ELiveLinkAxis AxisY, ELiveLinkAxis AxisZ)
diff --git a/Engine/Plugins/Animation/LiveLink/Source/LiveLink/Public/PreProcessor/LiveLinkAxisSwitchPreProcessor.h b/Engine/Plugins/Animation/LiveLink/Source/LiveLink/Public/PreProcessor/LiveLinkAxisSwitchPreProcessor.h
--- a/Engine/Plugins/Animation/LiveLink/Source/LiveLink/Public/PreProcessor/LiveLinkAxisSwitchPreProcessor.h
+++ b/Engine/Plugins/Animation/LiveLink/Source/LiveLink/Public/PreProcessor/LiveLinkAxisSwitchPreProcessor.h
@@ -22,7 +22,7 @@ enum class ELiveLinkAxis : uint8
  * @note For example the Z-Axis of an incoming transform can be set to the (optionally negated) Y-Axis of the transform in UE.
  * @note This implies that translation, rotation and scale will be affected by switching an axis.
  */
-UCLASS(meta = (DisplayName = "Axis Switch"))
+UCLASS(meta = (DisplayName = "Transform Axis Switch"))
 class LIVELINK_API ULiveLinkTransformAxisSwitchPreProcessor : public ULiveLinkFramePreProcessor
 {
 	GENERATED_BODY()
@@ -67,7 +67,7 @@ protected:
  * @note For example the Z-Axis of an incoming transform can be set to the (optionally negated) Y-Axis of the transform in UE.
  * @note This implies that translation, rotation and scale will be affected by switching an axis.
  */
-UCLASS(meta = (DisplayName = "Axis Switch"))
+UCLASS(meta = (DisplayName = "Animation Axis Switch"))
 class LIVELINK_API ULiveLinkAnimationAxisSwitchPreProcessor : public ULiveLinkTransformAxisSwitchPreProcessor
 {
 	GENERATED_BODY()
