Fix for incorrect background color of LiveLinkClientPanel when docked

diff --git a/Engine/Plugins/Animation/LiveLink/Source/LiveLinkEditor/Private/LiveLinkClientPanel.cpp b/Engine/Plugins/Animation/LiveLink/Source/LiveLinkEditor/Private/LiveLinkClientPanel.cpp
--- a/Engine/Plugins/Animation/LiveLink/Source/LiveLinkEditor/Private/LiveLinkClientPanel.cpp
+++ b/Engine/Plugins/Animation/LiveLink/Source/LiveLinkEditor/Private/LiveLinkClientPanel.cpp
@@ -466,151 +466,153 @@ void SLiveLinkClientPanel::Construct(const FArguments& Args, FLiveLinkClient* In
 	
 	ChildSlot
 	[
-		SNew(SVerticalBox)
-		+SVerticalBox::Slot()
-		.AutoHeight()
-		.Padding(FMargin(0.0f, 4.0f, 0.0f, 0.0f))
+		SNew(SBorder)
+		.BorderImage(FEditorStyle::GetBrush("MessageLog.ListBorder")) // set panel background color to same color as message log at the bottom
 		[
-			SNew(SLiveLinkClientPanelToolbar, Client)
-		]
-		+ SVerticalBox::Slot()
-		.FillHeight(1.f)
-		.Padding(FMargin(0.0f, 4.0f, 0.0f, 0.0f))
-		[
-			SNew(SSplitter)
-			.Orientation(EOrientation::Orient_Vertical)
-			+ SSplitter::Slot()
-			.Value(0.8f)
+			SNew(SVerticalBox)
+			+SVerticalBox::Slot()
+			.AutoHeight()
+			.Padding(FMargin(0.0f, 4.0f, 0.0f, 0.0f))
+			[
+				SNew(SLiveLinkClientPanelToolbar, Client)
+			]
+			+SVerticalBox::Slot()
+			.FillHeight(1.f)
+			.Padding(FMargin(0.0f, 4.0f, 0.0f, 0.0f))
 			[
 				SNew(SSplitter)
-				.Orientation(EOrientation::Orient_Horizontal)
+				.Orientation(EOrientation::Orient_Vertical)
 				+SSplitter::Slot()
-				.Value(0.5f)
+				.Value(0.8f)
 				[
 					SNew(SSplitter)
-					.Orientation(EOrientation::Orient_Vertical)
+					.Orientation(EOrientation::Orient_Horizontal)
 					+SSplitter::Slot()
-					.Value(0.25f)
+					.Value(0.5f)
 					[
-						SNew(SBorder)
-						.BorderImage(FEditorStyle::GetBrush("ToolPanel.GroupBorder"))
-						.Padding(FMargin(4.0f, 4.0f))
+						SNew(SSplitter)
+						.Orientation(EOrientation::Orient_Vertical)
+						+SSplitter::Slot()
+						.Value(0.25f)
 						[
-							SAssignNew(SourceListView, SLiveLinkSourceListView, Client)
-							.ListItemsSource(&SourceData)
-							.SelectionMode(ESelectionMode::Single)
-							.OnGenerateRow(this, &SLiveLinkClientPanel::MakeSourceListViewWidget)
-							.OnContextMenuOpening(this, &SLiveLinkClientPanel::OnSourceConstructContextMenu)
-							.OnSelectionChanged(this, &SLiveLinkClientPanel::OnSourceListSelectionChanged)
-							.HeaderRow
-							(
-								SNew(SHeaderRow)
-								+ SHeaderRow::Column(SourceListUI::TypeColumnName)
-								.FillWidth(25.f)
-								.DefaultLabel(LOCTEXT("TypeColumnHeaderName", "Source Type"))
-								+ SHeaderRow::Column(SourceListUI::MachineColumnName)
-								.FillWidth(25.f)
-								.DefaultLabel(LOCTEXT("MachineColumnHeaderName", "Source Machine"))
-								+ SHeaderRow::Column(SourceListUI::StatusColumnName)
-								.FillWidth(50.f)
-								.DefaultLabel(LOCTEXT("StatusColumnHeaderName", "Status"))
-								+ SHeaderRow::Column(SourceListUI::ActionsColumnName)
-								.ManualWidth(20.f)
-								.DefaultLabel(LOCTEXT("ActionsColumnHeaderName", ""))
-							)
+							SNew(SBorder)
+							.BorderImage(FEditorStyle::GetBrush("ToolPanel.GroupBorder"))
+							.Padding(FMargin(4.0f, 4.0f))
+							[
+								SAssignNew(SourceListView, SLiveLinkSourceListView, Client)
+								.ListItemsSource(&SourceData)
+								.SelectionMode(ESelectionMode::Single)
+								.OnGenerateRow(this, &SLiveLinkClientPanel::MakeSourceListViewWidget)
+								.OnContextMenuOpening(this, &SLiveLinkClientPanel::OnSourceConstructContextMenu)
+								.OnSelectionChanged(this, &SLiveLinkClientPanel::OnSourceListSelectionChanged)
+								.HeaderRow
+								(
+									SNew(SHeaderRow)
+									+SHeaderRow::Column(SourceListUI::TypeColumnName)
+									.FillWidth(25.f)
+									.DefaultLabel(LOCTEXT("TypeColumnHeaderName", "Source Type"))
+									+SHeaderRow::Column(SourceListUI::MachineColumnName)
+									.FillWidth(25.f)
+									.DefaultLabel(LOCTEXT("MachineColumnHeaderName", "Source Machine"))
+									+SHeaderRow::Column(SourceListUI::StatusColumnName)
+									.FillWidth(50.f)
+									.DefaultLabel(LOCTEXT("StatusColumnHeaderName", "Status"))
+									+SHeaderRow::Column(SourceListUI::ActionsColumnName)
+									.ManualWidth(20.f)
+									.DefaultLabel(LOCTEXT("ActionsColumnHeaderName", ""))
+								)
+							]
+						]
+						+SSplitter::Slot()
+						.Value(0.75f)
+						[
+							SNew(SBorder)
+							.BorderImage(FEditorStyle::GetBrush("ToolPanel.GroupBorder"))
+							.Padding(FMargin(4.0f, 4.0f))
+							[
+								SubjectsTreeView->AsShared()
+							]
 						]
 					]
 					+SSplitter::Slot()
-					.Value(0.75f)
+					.Value(0.5f)
 					[
-						SNew(SBorder)
-						.BorderImage(FEditorStyle::GetBrush("ToolPanel.GroupBorder"))
-						.Padding(FMargin(4.0f, 4.0f))
+						SNew(SWidgetSwitcher)
+						.WidgetIndex(this, &SLiveLinkClientPanel::GetDetailWidgetIndex)
+						+SWidgetSwitcher::Slot()
+						[
+							//[0] Detail view for Source
+							SettingsDetailsView.ToSharedRef()
+						]
+						+SWidgetSwitcher::Slot()
 						[
-							SubjectsTreeView->AsShared()
+							// [1] Detail view for Subject, Frame data & Static data
+							SAssignNew(DataDetailsView, SLiveLinkDataView, Client)
 						]
 					]
 				]
 				+SSplitter::Slot()
-				.Value(0.5f)
+				.Value(0.2f)
 				[
-					SNew(SWidgetSwitcher)
-					.WidgetIndex(this, &SLiveLinkClientPanel::GetDetailWidgetIndex)
-					+ SWidgetSwitcher::Slot()
+					SNew(SOverlay)
+					+SOverlay::Slot()
 					[
-						//[0] Detail view for Source
-						SettingsDetailsView.ToSharedRef()
+						MessageLogListingWidget
 					]
-					+ SWidgetSwitcher::Slot()
+					+SOverlay::Slot()
+					.HAlign(HAlign_Left)
+					.VAlign(VAlign_Bottom)
 					[
-						// [1] Detail view for Subject, Frame data & Static data
-						SAssignNew(DataDetailsView, SLiveLinkDataView, Client)
+						SNew(SHorizontalBox)
+						+SHorizontalBox::Slot()
+						.AutoWidth()
+						.Padding(10, 4, 4, 10)
+						[
+							SNew(STextBlock)
+							.Text(this, &SLiveLinkClientPanel::GetMessageCountText)
+						]
+						+SHorizontalBox::Slot()
+						.Padding(20, 4, 50, 10)
+						.AutoWidth()
+						[
+							SNew(STextBlock)
+							.Text(this, &SLiveLinkClientPanel::GetSelectedMessageOccurrenceText)
+						]
 					]
 				]
 			]
-			+ SSplitter::Slot()
-			.Value(0.2f)
+			+SVerticalBox::Slot()
+			.AutoHeight()
 			[
-				SNew(SOverlay)
-				+ SOverlay::Slot()
-				[
-					MessageLogListingWidget
-				]
-				+ SOverlay::Slot()
-				.HAlign(HAlign_Left)
-				.VAlign(VAlign_Bottom)
-
+				SNew(SBorder)
+				.BorderImage(FEditorStyle::GetBrush("SettingsEditor.CheckoutWarningBorder"))
+				.BorderBackgroundColor(FColor(166, 137, 0))
+				.Visibility(this, &SLiveLinkClientPanel::ShowEditorPerformanceThrottlingWarning)
 				[
 					SNew(SHorizontalBox)
-					+ SHorizontalBox::Slot()
-					.AutoWidth()
-					.Padding(10, 4, 4, 10)
+					+SHorizontalBox::Slot()
+					.FillWidth(1.0f)
+					.Padding(FMargin(WarningPadding, WarningPadding, WarningPadding, WarningPadding))
+					.VAlign(VAlign_Center)
 					[
 						SNew(STextBlock)
-						.Text(this, &SLiveLinkClientPanel::GetMessageCountText)
+						.Text(PerformanceWarningText)
+						.Font(FEditorStyle::GetFontStyle("PropertyWindow.NormalFont"))
+						.ShadowColorAndOpacity(FLinearColor::Black.CopyWithNewOpacity(0.3f))
+						.ShadowOffset(FVector2D::UnitVector)
 					]
-					+ SHorizontalBox::Slot()
-					.Padding(20, 4, 50, 10)
+					+SHorizontalBox::Slot()
+					.Padding(FMargin(0.f, 0.f, WarningPadding, 0.f))
 					.AutoWidth()
+					.VAlign(VAlign_Center)
 					[
-						SNew(STextBlock)
-						.Text(this, &SLiveLinkClientPanel::GetSelectedMessageOccurrenceText)
+						SNew(SButton)
+						.OnClicked(this, &SLiveLinkClientPanel::DisableEditorPerformanceThrottling)
+						.Text(LOCTEXT("LiveLinkPerformanceWarningDisable", "Disable"))
 					]
 				]
 			]
 		]
-		+SVerticalBox::Slot()
-		.AutoHeight()
-		[
-			SNew(SBorder)
-			.BorderImage(FEditorStyle::GetBrush("SettingsEditor.CheckoutWarningBorder"))
-			.BorderBackgroundColor(FColor(166, 137, 0))
-			.Visibility(this, &SLiveLinkClientPanel::ShowEditorPerformanceThrottlingWarning)
-			[
-				SNew(SHorizontalBox)
-				+ SHorizontalBox::Slot()
-				.FillWidth(1.0f)
-				.Padding(FMargin(WarningPadding, WarningPadding, WarningPadding, WarningPadding))
-				.VAlign(VAlign_Center)
-				[
-					SNew(STextBlock)
-					.Text(PerformanceWarningText)
-					.Font(FEditorStyle::GetFontStyle("PropertyWindow.NormalFont"))
-					.ShadowColorAndOpacity(FLinearColor::Black.CopyWithNewOpacity(0.3f))
-					.ShadowOffset(FVector2D::UnitVector)
-				]
-
-				+ SHorizontalBox::Slot()
-				.Padding(FMargin(0.f, 0.f, WarningPadding, 0.f))
-				.AutoWidth()
-				.VAlign(VAlign_Center)
-				[
-					SNew(SButton)
-					.OnClicked(this, &SLiveLinkClientPanel::DisableEditorPerformanceThrottling)
-					.Text(LOCTEXT("LiveLinkPerformanceWarningDisable", "Disable"))
-				]
-			]
-		]
 	];
 
 
