Change the name of LiveLinkMessageBusDiscoveryManager thread.

diff --git a/Engine/Plugins/Animation/LiveLink/Source/LiveLink/Private/LiveLinkMessageBusDiscoveryManager.cpp b/Engine/Plugins/Animation/LiveLink/Source/LiveLink/Private/LiveLinkMessageBusDiscoveryManager.cpp
--- a/Engine/Plugins/Animation/LiveLink/Source/LiveLink/Private/LiveLinkMessageBusDiscoveryManager.cpp
+++ b/Engine/Plugins/Animation/LiveLink/Source/LiveLink/Private/LiveLinkMessageBusDiscoveryManager.cpp
@@ -26,7 +26,7 @@ FLiveLinkMessageBusDiscoveryManager::FLiveLinkMessageBusDiscoveryManager()
 	bRunning = MessageEndpoint.IsValid();
 	if (bRunning)
 	{
-		Thread = FRunnableThread::Create(this, TEXT("MessageBusHeartbeatManager"));
+		Thread = FRunnableThread::Create(this, TEXT("LiveLinkMessageBusDiscoveryManager"));
 	}
 }
 
