- Fixed default interpolator not being applied to new LiveLink subjects

diff --git a/Engine/Config/BaseGame.ini b/Engine/Config/BaseGame.ini
--- a/Engine/Config/BaseGame.ini
+++ b/Engine/Config/BaseGame.ini
@@ -167,8 +167,8 @@ QueryCountWarningThreshold=0
 QueryCountWarningInterval=30.0
 
 [/Script/LiveLink.LiveLinkSettings]
-FrameInterpolationProcessor=Class'/Script/LiveLink.LiveLinkBasicFrameInterpolateProcessor'
-+DefaultRoleSettings=(Role=Class'/Script/LiveLink.LiveLinkAnimationRole', FrameInterpolationProcessor=Class'/Script/LiveLink.LiveLinkAnimationFrameInterpolateProcessor')
+FrameInterpolationProcessor=Class'/Script/LiveLink.LiveLinkBasicFrameInterpolationProcessor'
++DefaultRoleSettings=(Role=Class'/Script/LiveLink.LiveLinkAnimationRole', FrameInterpolationProcessor=Class'/Script/LiveLink.LiveLinkAnimationFrameInterpolationProcessor')
 
 [/Script/Engine.AssetManagerSettings]
 ; These lines should never be changed, due to the way config diffing works if the defaults are changed here games will be broken
diff --git a/Engine/Plugins/Animation/LiveLink/Source/LiveLink/Private/LiveLinkClient.cpp b/Engine/Plugins/Animation/LiveLink/Source/LiveLink/Private/LiveLinkClient.cpp
--- a/Engine/Plugins/Animation/LiveLink/Source/LiveLink/Private/LiveLinkClient.cpp
+++ b/Engine/Plugins/Animation/LiveLink/Source/LiveLink/Private/LiveLinkClient.cpp
@@ -448,7 +448,8 @@ void FLiveLinkClient::PushSubjectStaticData_Internal(FPendingSubjectStatic&& Sub
 
 	if(LiveLinkSubject == nullptr)
 	{
-		const FLiveLinkRoleProjectSetting DefaultSetting = GetDefault<ULiveLinkSettings>()->GetDefaultSettingForRole(SubjectStaticData.Role.Get());
+		const ULiveLinkSettings* LiveLinkSettings = GetDefault<ULiveLinkSettings>();
+		const FLiveLinkRoleProjectSetting DefaultSetting = LiveLinkSettings->GetDefaultSettingForRole(SubjectStaticData.Role.Get());
 
 		// Setting class should always be valid
 		ULiveLinkSubjectSettings* SubjectSettings = nullptr;
@@ -475,6 +476,19 @@ void FLiveLinkClient::PushSubjectStaticData_Internal(FPendingSubjectStatic&& Sub
 					FLiveLinkLog::Warning(TEXT("The interpolator '%s' is not valid for the Role '%s'"), *FrameInterpolationProcessorClass->GetName(), *SubjectStaticData.Role->GetName());
 				}
 			}
+			else
+			{
+				// If no settings were found for a specific role, check if the default interpolator is compatible with the role
+				UClass* FallbackInterpolationProcessorClass = LiveLinkSettings->FrameInterpolationProcessor.Get();
+				if (FallbackInterpolationProcessorClass != nullptr)
+				{
+					UClass* InterpolationRole = FallbackInterpolationProcessorClass->GetDefaultObject<ULiveLinkFrameInterpolationProcessor>()->GetRole();
+					if (SubjectStaticData.Role->IsChildOf(InterpolationRole))
+					{
+						SubjectSettings->InterpolationProcessor = NewObject<ULiveLinkFrameInterpolationProcessor>(SubjectSettings, FallbackInterpolationProcessorClass);
+					}
+				}
+			}
 
 			for (TSubclassOf<ULiveLinkFramePreProcessor> PreProcessor : DefaultSetting.FramePreProcessors)
 			{
diff --git a/Engine/Plugins/Animation/LiveLink/Source/LiveLink/Public/LiveLinkSettings.h b/Engine/Plugins/Animation/LiveLink/Source/LiveLink/Public/LiveLinkSettings.h
--- a/Engine/Plugins/Animation/LiveLink/Source/LiveLink/Public/LiveLinkSettings.h
+++ b/Engine/Plugins/Animation/LiveLink/Source/LiveLink/Public/LiveLinkSettings.h
@@ -61,11 +61,11 @@ protected:
 	UPROPERTY(config, EditAnywhere, Category="LiveLink")
 	TArray<FLiveLinkRoleProjectSetting> DefaultRoleSettings;
 
-	/** The interpolation to use for the subject. If null, no interpolation will be performed. */
+public:
+	/** The interpolation class to use for new Subjects if no specific settings we set for the Subject's role. */
 	UPROPERTY(config)
 	TSubclassOf<ULiveLinkFrameInterpolationProcessor> FrameInterpolationProcessor;
-
-public:
+	
 	/** The default location in which to save take presets */
 	UPROPERTY(config, EditAnywhere, Category="LiveLink", meta=(DisplayName="Preset Save Location"))
 	FDirectoryPath PresetSaveDir;
