LiveLink axis switcher. Add an option to swith the translation unit.

diff --git a/Engine/Plugins/Animation/LiveLink/Config/BaseLiveLink.ini b/Engine/Plugins/Animation/LiveLink/Config/BaseLiveLink.ini
--- a/Engine/Plugins/Animation/LiveLink/Config/BaseLiveLink.ini
+++ b/Engine/Plugins/Animation/LiveLink/Config/BaseLiveLink.ini
@@ -30,3 +30,11 @@
 +ClassRedirects=(OldName="LiveLinkAnimationFrameInterpolateProcessor", NewName="/Script/LiveLink.LiveLinkAnimationFrameInterpolationProcessor")
 +ClassRedirects=(OldName="LiveLinkAnimationRoleToTransform", NewName="/Script/LiveLink.LiveLinkAnimationRoleToTransform")
 
+;
+; 4.25 Redirects
+;
+
+; Axis Switch processors
++PropertyRedirects=(OldName="LiveLinkTransformAxisSwitchPreProcessor.AxisX",NewName="LiveLinkTransformAxisSwitchPreProcessor.OrientationAxisX")
++PropertyRedirects=(OldName="LiveLinkTransformAxisSwitchPreProcessor.AxisY",NewName="LiveLinkTransformAxisSwitchPreProcessor.OrientationAxisY")
++PropertyRedirects=(OldName="LiveLinkTransformAxisSwitchPreProcessor.AxisZ",NewName="LiveLinkTransformAxisSwitchPreProcessor.OrientationAxisZ")
\ No newline at end of file
diff --git a/Engine/Plugins/Animation/LiveLink/Source/LiveLink/Private/PreProcessor/LiveLinkAxisSwitchPreProcessor.cpp b/Engine/Plugins/Animation/LiveLink/Source/LiveLink/Private/PreProcessor/LiveLinkAxisSwitchPreProcessor.cpp
--- a/Engine/Plugins/Animation/LiveLink/Source/LiveLink/Private/PreProcessor/LiveLinkAxisSwitchPreProcessor.cpp
+++ b/Engine/Plugins/Animation/LiveLink/Source/LiveLink/Private/PreProcessor/LiveLinkAxisSwitchPreProcessor.cpp
@@ -16,21 +16,35 @@ namespace
 		return Results[Value];
 	}
 
+	float LiveLinkAxisToVectorMember(ELiveLinkAxis InAxis, const FVector& Origin)
+	{
+		const uint8 Value = static_cast<uint8>(InAxis) % 3;
+		return Origin[Value];
+	}
+
 	float AxisSign(ELiveLinkAxis InAxis)
 	{
 		return static_cast<uint8>(InAxis) < 3 ? 1.f : -1.f;
 	}
 
-	void SwitchTransform(FTransform& Transform, ELiveLinkAxis AxisX, ELiveLinkAxis AxisY, ELiveLinkAxis AxisZ)
+	void SwitchTransform(FTransform& Transform
+		, ELiveLinkAxis OrientationAxisX, ELiveLinkAxis OrientationAxisY, ELiveLinkAxis OrientationAxisZ
+		, ELiveLinkAxis TranslationAxisX, ELiveLinkAxis TranslationAxisY, ELiveLinkAxis TranslationAxisZ)
 	{
 		const FMatrix InMatrix = Transform.ToMatrixWithScale();
 
-		FVector DestAxisX = InMatrix.GetScaledAxis(LiveLinkAxisToMatrixAxis(AxisX)) * AxisSign(AxisX);
-		FVector DestAxisY = InMatrix.GetScaledAxis(LiveLinkAxisToMatrixAxis(AxisY)) * AxisSign(AxisY);
-		FVector DestAxisZ = InMatrix.GetScaledAxis(LiveLinkAxisToMatrixAxis(AxisZ)) * AxisSign(AxisZ);
+		FVector DestAxisX = InMatrix.GetScaledAxis(LiveLinkAxisToMatrixAxis(OrientationAxisX)) * AxisSign(OrientationAxisX);
+		FVector DestAxisY = InMatrix.GetScaledAxis(LiveLinkAxisToMatrixAxis(OrientationAxisY)) * AxisSign(OrientationAxisY);
+		FVector DestAxisZ = InMatrix.GetScaledAxis(LiveLinkAxisToMatrixAxis(OrientationAxisZ)) * AxisSign(OrientationAxisZ);
+
+		FVector Origin = InMatrix.GetOrigin();
+		FVector NewOrigin;
+		NewOrigin.X = LiveLinkAxisToVectorMember(TranslationAxisX, Origin) * AxisSign(TranslationAxisX);
+		NewOrigin.Y = LiveLinkAxisToVectorMember(TranslationAxisY, Origin) * AxisSign(TranslationAxisY);
+		NewOrigin.Z = LiveLinkAxisToVectorMember(TranslationAxisZ, Origin) * AxisSign(TranslationAxisZ);
 
 		FMatrix Result(InMatrix);
-		Result.SetAxes(&DestAxisX, &DestAxisY, &DestAxisZ);
+		Result.SetAxes(&DestAxisX, &DestAxisY, &DestAxisZ, &NewOrigin);
 
 		Transform.SetFromMatrix(Result);
 	}
@@ -47,7 +61,7 @@ TSubclassOf<ULiveLinkRole> ULiveLinkTransformAxisSwitchPreProcessor::FLiveLinkTr
 bool ULiveLinkTransformAxisSwitchPreProcessor::FLiveLinkTransformAxisSwitchPreProcessorWorker::PreProcessFrame(FLiveLinkFrameDataStruct& InOutFrame) const
 {
 	FLiveLinkTransformFrameData& TransformData = *InOutFrame.Cast<FLiveLinkTransformFrameData>();
-	SwitchTransform(TransformData.Transform, AxisX, AxisY, AxisZ);
+	SwitchTransform(TransformData.Transform, OrientationAxisX, OrientationAxisY, OrientationAxisZ, TranslationAxisX, TranslationAxisY, TranslationAxisZ);
 	return true;
 }
 
@@ -64,9 +78,12 @@ ULiveLinkFramePreProcessor::FWorkerSharedPtr ULiveLinkTransformAxisSwitchPreProc
 	if (!Instance.IsValid())
 	{
 		Instance = MakeShared<FLiveLinkTransformAxisSwitchPreProcessorWorker, ESPMode::ThreadSafe>();
-		Instance->AxisX = AxisX;
-		Instance->AxisY = AxisY;
-		Instance->AxisZ = AxisZ;
+		Instance->OrientationAxisX = OrientationAxisX;
+		Instance->OrientationAxisY = OrientationAxisY;
+		Instance->OrientationAxisZ = OrientationAxisZ;
+		Instance->TranslationAxisX = TranslationAxisX;
+		Instance->TranslationAxisY = TranslationAxisY;
+		Instance->TranslationAxisZ = TranslationAxisZ;
 	}
 
 	return Instance;
@@ -75,9 +92,12 @@ ULiveLinkFramePreProcessor::FWorkerSharedPtr ULiveLinkTransformAxisSwitchPreProc
 #if WITH_EDITOR
 void ULiveLinkTransformAxisSwitchPreProcessor::PostEditChangeChainProperty(struct FPropertyChangedChainEvent& PropertyChangedEvent)
 {
-	if (PropertyChangedEvent.Property->GetFName() == GET_MEMBER_NAME_CHECKED(ThisClass, AxisX) ||
-		PropertyChangedEvent.Property->GetFName() == GET_MEMBER_NAME_CHECKED(ThisClass, AxisY) ||
-		PropertyChangedEvent.Property->GetFName() == GET_MEMBER_NAME_CHECKED(ThisClass, AxisZ))
+	if (PropertyChangedEvent.Property->GetFName() == GET_MEMBER_NAME_CHECKED(ThisClass, OrientationAxisX) ||
+		PropertyChangedEvent.Property->GetFName() == GET_MEMBER_NAME_CHECKED(ThisClass, OrientationAxisY) ||
+		PropertyChangedEvent.Property->GetFName() == GET_MEMBER_NAME_CHECKED(ThisClass, OrientationAxisZ) ||
+		PropertyChangedEvent.Property->GetFName() == GET_MEMBER_NAME_CHECKED(ThisClass, TranslationAxisX) ||
+		PropertyChangedEvent.Property->GetFName() == GET_MEMBER_NAME_CHECKED(ThisClass, TranslationAxisY) ||
+		PropertyChangedEvent.Property->GetFName() == GET_MEMBER_NAME_CHECKED(ThisClass, TranslationAxisZ))
 	{
 		Instance.Reset();
 	}
@@ -101,7 +121,7 @@ bool ULiveLinkAnimationAxisSwitchPreProcessor::FLiveLinkAnimationAxisSwitchPrePr
 	
 	for (FTransform& Transform : AnimationData.Transforms)
 	{
-		SwitchTransform(Transform, AxisX, AxisY, AxisZ);
+		SwitchTransform(Transform, OrientationAxisX, OrientationAxisY, OrientationAxisZ, TranslationAxisX, TranslationAxisY, TranslationAxisZ);
 	}
 	
 	return true;
@@ -120,10 +140,13 @@ ULiveLinkFramePreProcessor::FWorkerSharedPtr ULiveLinkAnimationAxisSwitchPreProc
 	if (!Instance.IsValid())
 	{
 		Instance = MakeShared<FLiveLinkAnimationAxisSwitchPreProcessorWorker, ESPMode::ThreadSafe>();
-		Instance->AxisX = AxisX;
-		Instance->AxisY = AxisY;
-		Instance->AxisZ = AxisZ;
+		Instance->OrientationAxisX = OrientationAxisX;
+		Instance->OrientationAxisY = OrientationAxisY;
+		Instance->OrientationAxisZ = OrientationAxisZ;
+		Instance->TranslationAxisX = TranslationAxisX;
+		Instance->TranslationAxisY = TranslationAxisY;
+		Instance->TranslationAxisZ = TranslationAxisZ;
 	}
 
 	return Instance;
-}
\ No newline at end of file
+}
diff --git a/Engine/Plugins/Animation/LiveLink/Source/LiveLink/Public/PreProcessor/LiveLinkAxisSwitchPreProcessor.h b/Engine/Plugins/Animation/LiveLink/Source/LiveLink/Public/PreProcessor/LiveLinkAxisSwitchPreProcessor.h
--- a/Engine/Plugins/Animation/LiveLink/Source/LiveLink/Public/PreProcessor/LiveLinkAxisSwitchPreProcessor.h
+++ b/Engine/Plugins/Animation/LiveLink/Source/LiveLink/Public/PreProcessor/LiveLinkAxisSwitchPreProcessor.h
@@ -31,9 +31,12 @@ public:
 	class FLiveLinkTransformAxisSwitchPreProcessorWorker : public ILiveLinkFramePreProcessorWorker
 	{
 	public:
-		ELiveLinkAxis AxisX = ELiveLinkAxis::X;
-		ELiveLinkAxis AxisY = ELiveLinkAxis::Y;
-		ELiveLinkAxis AxisZ = ELiveLinkAxis::Z;
+		ELiveLinkAxis OrientationAxisX = ELiveLinkAxis::X;
+		ELiveLinkAxis OrientationAxisY = ELiveLinkAxis::Y;
+		ELiveLinkAxis OrientationAxisZ = ELiveLinkAxis::Z;
+		ELiveLinkAxis TranslationAxisX = ELiveLinkAxis::X;
+		ELiveLinkAxis TranslationAxisY = ELiveLinkAxis::Y;
+		ELiveLinkAxis TranslationAxisZ = ELiveLinkAxis::Z;
 
 		virtual TSubclassOf<ULiveLinkRole> GetRole() const override;
 		virtual bool PreProcessFrame(FLiveLinkFrameDataStruct& InOutFrame) const override;
@@ -41,11 +44,17 @@ public:
 
 protected:
 	UPROPERTY(EditAnywhere, Category = "LiveLink")
-	ELiveLinkAxis AxisX = ELiveLinkAxis::X;
+	ELiveLinkAxis OrientationAxisX = ELiveLinkAxis::X;
 	UPROPERTY(EditAnywhere, Category = "LiveLink")
-	ELiveLinkAxis AxisY = ELiveLinkAxis::Y;
+	ELiveLinkAxis OrientationAxisY = ELiveLinkAxis::Y;
 	UPROPERTY(EditAnywhere, Category = "LiveLink")
-	ELiveLinkAxis AxisZ = ELiveLinkAxis::Z;
+	ELiveLinkAxis OrientationAxisZ = ELiveLinkAxis::Z;
+	UPROPERTY(EditAnywhere, Category = "LiveLink")
+	ELiveLinkAxis TranslationAxisX = ELiveLinkAxis::X;
+	UPROPERTY(EditAnywhere, Category = "LiveLink")
+	ELiveLinkAxis TranslationAxisY = ELiveLinkAxis::Y;
+	UPROPERTY(EditAnywhere, Category = "LiveLink")
+	ELiveLinkAxis TranslationAxisZ = ELiveLinkAxis::Z;
 
 public:
 	virtual TSubclassOf<ULiveLinkRole> GetRole() const override;
