Fix CIS, missing include

diff --git a/Engine/Plugins/Animation/LiveLink/Source/LiveLinkComponents/Private/LiveLinkComponentController.cpp b/Engine/Plugins/Animation/LiveLink/Source/LiveLinkComponents/Private/LiveLinkComponentController.cpp
--- a/Engine/Plugins/Animation/LiveLink/Source/LiveLinkComponents/Private/LiveLinkComponentController.cpp
+++ b/Engine/Plugins/Animation/LiveLink/Source/LiveLinkComponents/Private/LiveLinkComponentController.cpp
@@ -2,6 +2,7 @@
 
 #include "LiveLinkComponentController.h"
 
+#include "Features/IModularFeatures.h"
 #include "LiveLinkComponentSettings.h"
 #include "LiveLinkComponentPrivate.h"
 #include "LiveLinkControllerBase.h"
