Fix CIS, missing include in for SlateColor.

diff --git a/Engine/Plugins/Animation/LiveLink/Source/LiveLinkEditor/Private/LiveLinkComponentDetailCustomization.h b/Engine/Plugins/Animation/LiveLink/Source/LiveLinkEditor/Private/LiveLinkComponentDetailCustomization.h
--- a/Engine/Plugins/Animation/LiveLink/Source/LiveLinkEditor/Private/LiveLinkComponentDetailCustomization.h
+++ b/Engine/Plugins/Animation/LiveLink/Source/LiveLinkEditor/Private/LiveLinkComponentDetailCustomization.h
@@ -5,6 +5,7 @@
 #include "IDetailCustomization.h"
 
 #include "CoreMinimal.h"
+#include "Styling/SlateColor.h"
 
 #include "LiveLinkComponentController.h"
 
