- Added PlatformTime snapshot in LiveLinkMessageBus PongMessage to be able to compute the offse between machines

diff --git a/Engine/Plugins/Animation/LiveLink/Source/LiveLink/Private/LiveLinkMessageBusDiscoveryManager.cpp b/Engine/Plugins/Animation/LiveLink/Source/LiveLink/Private/LiveLinkMessageBusDiscoveryManager.cpp
--- a/Engine/Plugins/Animation/LiveLink/Source/LiveLink/Private/LiveLinkMessageBusDiscoveryManager.cpp
+++ b/Engine/Plugins/Animation/LiveLink/Source/LiveLink/Private/LiveLinkMessageBusDiscoveryManager.cpp
@@ -109,6 +109,7 @@ void FLiveLinkMessageBusDiscoveryManager::HandlePongMessage(const FLiveLinkPongM
 
 	if (Message.PollRequest == LastPingRequest)
 	{
-		LastProviderPoolResults.Emplace(MakeShared<FProviderPollResult, ESPMode::ThreadSafe>(Context->GetSender(), Message.ProviderName, Message.MachineName));
+		const double MachineTimeOffset = LiveLinkMessageBusHelper::CalculateProviderMachineOffset(Message.CreationPlatformTime, Context);
+		LastProviderPoolResults.Emplace(MakeShared<FProviderPollResult, ESPMode::ThreadSafe>(Context->GetSender(), Message.ProviderName, Message.MachineName, MachineTimeOffset));
 	}
 }
\ No newline at end of file
diff --git a/Engine/Plugins/Animation/LiveLink/Source/LiveLink/Private/LiveLinkMessageBusFinder.cpp b/Engine/Plugins/Animation/LiveLink/Source/LiveLink/Private/LiveLinkMessageBusFinder.cpp
--- a/Engine/Plugins/Animation/LiveLink/Source/LiveLink/Private/LiveLinkMessageBusFinder.cpp
+++ b/Engine/Plugins/Animation/LiveLink/Source/LiveLink/Private/LiveLinkMessageBusFinder.cpp
@@ -9,6 +9,25 @@
 #include "LiveLinkMessageBusSource.h"
 #include "MessageEndpointBuilder.h"
 
+
+namespace LiveLinkMessageBusHelper
+{
+	double CalculateProviderMachineOffset(double SourceMachinePlatformSeconds, const TSharedRef<IMessageContext, ESPMode::ThreadSafe>& Context)
+	{
+		const FDateTime Now = FDateTime::UtcNow();
+		const double MyMachineSeconds = FPlatformTime::Seconds();
+		const FTimespan Latency = Now - Context->GetTimeSent();
+		double MachineTimeOffset = 0.0f;
+		if (SourceMachinePlatformSeconds >= 0.0)
+		{
+			MachineTimeOffset = MyMachineSeconds - SourceMachinePlatformSeconds - Latency.GetTotalSeconds();
+		}
+
+		return MachineTimeOffset;
+	}
+}
+
+
 ULiveLinkMessageBusFinder::ULiveLinkMessageBusFinder()
 	: MessageEndpoint(nullptr)
 {
@@ -58,7 +77,9 @@ void ULiveLinkMessageBusFinder::HandlePongMessage(const FLiveLinkPongMessage& Me
 	if (Message.PollRequest == CurrentPollRequest)
 	{
 		FScopeLock ScopedLock(&PollDataCriticalSection);
-		PollData.Add(FProviderPollResult(Context->GetSender(), Message.ProviderName, Message.MachineName));
+		
+		const double MachineTimeOffset = LiveLinkMessageBusHelper::CalculateProviderMachineOffset(Message.CreationPlatformTime, Context);
+		PollData.Add(FProviderPollResult(Context->GetSender(), Message.ProviderName, Message.MachineName, MachineTimeOffset));
 	}
 };
 
@@ -69,7 +90,7 @@ void ULiveLinkMessageBusFinder::ConnectToProvider(UPARAM(ref) FProviderPollResul
 	if (ModularFeatures.IsModularFeatureAvailable(ILiveLinkClient::ModularFeatureName))
 	{
 		ILiveLinkClient* LiveLinkClient = &ModularFeatures.GetModularFeature<ILiveLinkClient>(ILiveLinkClient::ModularFeatureName);
-		TSharedPtr<FLiveLinkMessageBusSource> NewSource = MakeShared<FLiveLinkMessageBusSource>(FText::FromString(Provider.Name), FText::FromString(Provider.MachineName), Provider.Address);
+		TSharedPtr<FLiveLinkMessageBusSource> NewSource = MakeShared<FLiveLinkMessageBusSource>(FText::FromString(Provider.Name), FText::FromString(Provider.MachineName), Provider.Address, Provider.MachineTimeOffset);
 		LiveLinkClient->AddSource(NewSource);
 		SourceHandle.SetSourcePointer(NewSource);
 	}
@@ -82,4 +103,6 @@ void ULiveLinkMessageBusFinder::ConnectToProvider(UPARAM(ref) FProviderPollResul
 ULiveLinkMessageBusFinder* ULiveLinkMessageBusFinder::ConstructMessageBusFinder()
 {
 	return NewObject<ULiveLinkMessageBusFinder>();
-}
\ No newline at end of file
+}
+
+
diff --git a/Engine/Plugins/Animation/LiveLink/Source/LiveLink/Private/LiveLinkMessageBusSource.cpp b/Engine/Plugins/Animation/LiveLink/Source/LiveLink/Private/LiveLinkMessageBusSource.cpp
--- a/Engine/Plugins/Animation/LiveLink/Source/LiveLink/Private/LiveLinkMessageBusSource.cpp
+++ b/Engine/Plugins/Animation/LiveLink/Source/LiveLink/Private/LiveLinkMessageBusSource.cpp
@@ -61,6 +61,7 @@ void FLiveLinkMessageBusSource::Update()
 			{
 				ConnectionAddress = Result->Address;
 				SourceMachineName = FText::FromString(Result->MachineName);
+				MachineTimeOffset = Result->MachineTimeOffset;
 				DiscoveryManager.RemoveDiscoveryMessageRequest();
 				SendConnectMessage();
 				UpdateConnectionLastActive();
@@ -187,7 +188,9 @@ void FLiveLinkMessageBusSource::InternalHandleMessage(const TSharedRef<IMessageC
 		check(MessageTypeInfo->IsChildOf(FLiveLinkBaseFrameData::StaticStruct()));
 
 		FLiveLinkFrameDataStruct DataStruct(MessageTypeInfo);
-		DataStruct.InitializeWith(MessageTypeInfo, reinterpret_cast<const FLiveLinkBaseFrameData*>(Context->GetMessage()));
+		const FLiveLinkBaseFrameData* Message = reinterpret_cast<const FLiveLinkBaseFrameData*>(Context->GetMessage());
+		DataStruct.InitializeWith(MessageTypeInfo, Message);
+		DataStruct.GetBaseData()->WorldTime = FLiveLinkWorldTime(Message->WorldTime.GetOffsettedTime(), MachineTimeOffset);
 		Client->PushSubjectFrameData_AnyThread(SubjectKey, MoveTemp(DataStruct));
 	}
 }
diff --git a/Engine/Plugins/Animation/LiveLink/Source/LiveLink/Private/LiveLinkMessageBusSource.h b/Engine/Plugins/Animation/LiveLink/Source/LiveLink/Private/LiveLinkMessageBusSource.h
--- a/Engine/Plugins/Animation/LiveLink/Source/LiveLink/Private/LiveLinkMessageBusSource.h
+++ b/Engine/Plugins/Animation/LiveLink/Source/LiveLink/Private/LiveLinkMessageBusSource.h
@@ -20,12 +20,13 @@ class LIVELINK_API FLiveLinkMessageBusSource : public ILiveLinkSource
 {
 public:
 
-	FLiveLinkMessageBusSource(const FText& InSourceType, const FText& InSourceMachineName, const FMessageAddress& InConnectionAddress)
+	FLiveLinkMessageBusSource(const FText& InSourceType, const FText& InSourceMachineName, const FMessageAddress& InConnectionAddress, double InMachineTimeOffset)
 		: ConnectionAddress(InConnectionAddress)
 		, SourceType(InSourceType)
 		, SourceMachineName(InSourceMachineName)
 		, ConnectionLastActive(0.0)
 		, bIsValid(false)
+		, MachineTimeOffset(InMachineTimeOffset)
 	{}
 
 	virtual void ReceiveClient(ILiveLinkClient* InClient, FGuid InSourceGuid) override;
@@ -79,4 +80,7 @@ private:
 
 	// Critical section to allow for threadsafe updating of the connection time
 	FCriticalSection ConnectionLastActiveSection;
+
+	// Offset between sender's machine engine time and receiver's machine engine time
+	double MachineTimeOffset;
 };
\ No newline at end of file
diff --git a/Engine/Plugins/Animation/LiveLink/Source/LiveLink/Private/LiveLinkMessageBusSourceFactory.cpp b/Engine/Plugins/Animation/LiveLink/Source/LiveLink/Private/LiveLinkMessageBusSourceFactory.cpp
--- a/Engine/Plugins/Animation/LiveLink/Source/LiveLink/Private/LiveLinkMessageBusSourceFactory.cpp
+++ b/Engine/Plugins/Animation/LiveLink/Source/LiveLink/Private/LiveLinkMessageBusSourceFactory.cpp
@@ -34,7 +34,8 @@ TSharedPtr<ILiveLinkSource> ULiveLinkMessageBusSourceFactory::CreateSource(const
 		return TSharedPtr<ILiveLinkSource>();
 	}
 
-	return MakeShared<FLiveLinkMessageBusSource>(FText::FromString(Name), FText::GetEmpty(), FMessageAddress());
+	const double TimeOffset = 0.0;
+	return MakeShared<FLiveLinkMessageBusSource>(FText::FromString(Name), FText::GetEmpty(), FMessageAddress(), TimeOffset);
 }
 
 void ULiveLinkMessageBusSourceFactory::OnSourceSelected(FProviderPollResultPtr SelectedSource, FOnLiveLinkSourceCreated InOnLiveLinkSourceCreated) const
@@ -65,7 +66,7 @@ void ULiveLinkMessageBusSourceFactory::OnSourceSelected(FProviderPollResultPtr S
 		}
 #endif
 
-		TSharedPtr<FLiveLinkMessageBusSource> SharedPtr = MakeShared<FLiveLinkMessageBusSource>(FText::FromString(SelectedSource->Name), FText::FromString(SelectedSource->MachineName), SelectedSource->Address);
+		TSharedPtr<FLiveLinkMessageBusSource> SharedPtr = MakeShared<FLiveLinkMessageBusSource>(FText::FromString(SelectedSource->Name), FText::FromString(SelectedSource->MachineName), SelectedSource->Address, SelectedSource->MachineTimeOffset);
 		FString ConnectionString = FString::Printf(TEXT("Name=\"%s\""), *SelectedSource->Name);
 		InOnLiveLinkSourceCreated.ExecuteIfBound(SharedPtr, MoveTemp(ConnectionString));
 	}
diff --git a/Engine/Plugins/Animation/LiveLink/Source/LiveLink/Public/LiveLinkMessageBusFinder.h b/Engine/Plugins/Animation/LiveLink/Source/LiveLink/Public/LiveLinkMessageBusFinder.h
--- a/Engine/Plugins/Animation/LiveLink/Source/LiveLink/Public/LiveLinkMessageBusFinder.h
+++ b/Engine/Plugins/Animation/LiveLink/Source/LiveLink/Public/LiveLinkMessageBusFinder.h
@@ -13,6 +13,12 @@
 
 struct FLiveLinkPongMessage;
 
+
+namespace LiveLinkMessageBusHelper
+{
+	double CalculateProviderMachineOffset(double SourceMachinePlatformSeconds, const TSharedRef<IMessageContext, ESPMode::ThreadSafe>& Context);
+}
+
 USTRUCT(BlueprintType)
 struct FProviderPollResult
 {
@@ -21,10 +27,11 @@ public:
 	
 	FProviderPollResult() = default;
 	
-	FProviderPollResult(const FMessageAddress& InAddress, const FString& InName, const FString& InMachineName)
+	FProviderPollResult(const FMessageAddress& InAddress, const FString& InName, const FString& InMachineName, double InMachineOffset)
 		: Address(InAddress)
 		, Name(InName)
 		, MachineName(InMachineName)
+		, MachineTimeOffset(InMachineOffset)
 	{}
 
 	FMessageAddress Address;
@@ -36,6 +43,10 @@ public:
 	// The name of the machine the provider is running on
 	UPROPERTY(BlueprintReadOnly, EditAnywhere, Category="LiveLink")
 	FString			MachineName;
+
+	// Offset between sender's engine time and receiver's engine time
+	UPROPERTY()
+	double			MachineTimeOffset;
 };
 
 typedef TSharedPtr<FProviderPollResult, ESPMode::ThreadSafe> FProviderPollResultPtr;
diff --git a/Engine/Source/Runtime/LiveLinkMessageBusFramework/Public/LiveLinkMessages.h b/Engine/Source/Runtime/LiveLinkMessageBusFramework/Public/LiveLinkMessages.h
--- a/Engine/Source/Runtime/LiveLinkMessageBusFramework/Public/LiveLinkMessages.h
+++ b/Engine/Source/Runtime/LiveLinkMessageBusFramework/Public/LiveLinkMessages.h
@@ -46,10 +46,13 @@ struct FLiveLinkPongMessage
 	UPROPERTY()
 	int32 LiveLinkVersion = 1;
 
+	UPROPERTY()
+	double CreationPlatformTime = -1.0;
+
 	// default constructor for the receiver
 	FLiveLinkPongMessage() = default;
 
-	FLiveLinkPongMessage(const FString& InProviderName, const FString& InMachineName, const FGuid& InPollRequest) : ProviderName(InProviderName), MachineName(InMachineName), PollRequest(InPollRequest) {}
+	FLiveLinkPongMessage(const FString& InProviderName, const FString& InMachineName, const FGuid& InPollRequest) : ProviderName(InProviderName), MachineName(InMachineName), PollRequest(InPollRequest), CreationPlatformTime(FPlatformTime::Seconds()) {}
 };
 
 USTRUCT()
