- Fixing LiveLink EvaluateFrame node to have cohesive behavior as other nodes (abstract class compile time error)

diff --git a/Engine/Plugins/Animation/LiveLink/Source/LiveLinkGraphNode/Private/K2Node_EvaluateLiveLinkFrame.cpp b/Engine/Plugins/Animation/LiveLink/Source/LiveLinkGraphNode/Private/K2Node_EvaluateLiveLinkFrame.cpp
--- a/Engine/Plugins/Animation/LiveLink/Source/LiveLinkGraphNode/Private/K2Node_EvaluateLiveLinkFrame.cpp
+++ b/Engine/Plugins/Animation/LiveLink/Source/LiveLinkGraphNode/Private/K2Node_EvaluateLiveLinkFrame.cpp
@@ -114,7 +114,7 @@ UScriptStruct* UK2Node_EvaluateLiveLinkFrame::GetLiveLinkRoleOutputStructType()
 	UScriptStruct* DataStructType = nullptr;
 
 	TSubclassOf<ULiveLinkRole> Role = GetDefaultRolePinValue();
-	if (Role != nullptr)
+	if (IsRoleValidForEvaluation(Role))
 	{
 		DataStructType = Role.GetDefaultObject()->GetBlueprintDataStruct();
 	}
@@ -192,7 +192,7 @@ void UK2Node_EvaluateLiveLinkFrame::PinDefaultValueChanged(UEdGraphPin* ChangedP
 		if (LiveLinkRolePin->DefaultObject != nullptr && LiveLinkRolePin->LinkedTo.Num() == 0)
 		{
 			UClass* ClassValue = Cast<UClass>(LiveLinkRolePin->DefaultObject);
-			bool bIsValid = ClassValue && ClassValue->IsChildOf(ULiveLinkRole::StaticClass()) && !ClassValue->HasAnyClassFlags(CLASS_Abstract | CLASS_Deprecated | CLASS_HideDropDown);
+			bool bIsValid = ClassValue && ClassValue->IsChildOf(ULiveLinkRole::StaticClass());
 			if (!bIsValid)
 			{
 				LiveLinkRolePin->DefaultObject = nullptr;
@@ -331,12 +331,18 @@ void UK2Node_EvaluateLiveLinkFrame::EarlyValidation(class FCompilerResultsLog& M
 	}
 
 	TSubclassOf<ULiveLinkRole> Role = GetDefaultRolePinValue();
-	if (Role == nullptr || !Role->IsChildOf(ULiveLinkRole::StaticClass()))
+	if (Role.Get() == nullptr)
 	{
 		MessageLog.Error(*LOCTEXT("NoLiveLinkRole", "No LiveLinkRole in @@").ToString(), this);
 		return;
 	}
 
+	if (!IsRoleValidForEvaluation(Role))
+	{
+		MessageLog.Error(*FText::Format(LOCTEXT("InvalidRoleClass", "Cannot EvaluateFrame for Role of type '{0}' in @@"), FText::FromString(Role->GetFName().ToString())).ToString(), this);
+		return;
+	}
+
 	UScriptStruct* DesiredOutputStruct = GetLiveLinkRoleOutputStructType();
 	UScriptStruct* ActualOutputStruct = GetReturnTypeForOutputDataStruct();
 	if (DesiredOutputStruct != ActualOutputStruct)
@@ -398,7 +404,7 @@ TSubclassOf<ULiveLinkRole> UK2Node_EvaluateLiveLinkFrame::GetDefaultRolePinValue
 	if (LiveLinkRolePin && LiveLinkRolePin->DefaultObject != nullptr && LiveLinkRolePin->LinkedTo.Num() == 0)
 	{
 		UClass* ClassValue = Cast<UClass>(LiveLinkRolePin->DefaultObject);
-		if (ClassValue && ClassValue->IsChildOf(ULiveLinkRole::StaticClass()) && !ClassValue->HasAnyClassFlags(CLASS_Abstract|CLASS_Deprecated|CLASS_HideDropDown))
+		if (ClassValue && ClassValue->IsChildOf(ULiveLinkRole::StaticClass()))
 		{
 			return ClassValue;
 		}
@@ -407,4 +413,9 @@ TSubclassOf<ULiveLinkRole> UK2Node_EvaluateLiveLinkFrame::GetDefaultRolePinValue
 }
 
 
+bool UK2Node_EvaluateLiveLinkFrame::IsRoleValidForEvaluation(TSubclassOf<ULiveLinkRole> InRoleClass) const
+{
+	return InRoleClass.Get() && !InRoleClass->HasAnyClassFlags(CLASS_Abstract | CLASS_Deprecated | CLASS_HideDropDown);
+}
+
 #undef LOCTEXT_NAMESPACE
diff --git a/Engine/Plugins/Animation/LiveLink/Source/LiveLinkGraphNode/Public/K2Node_EvaluateLiveLinkFrame.h b/Engine/Plugins/Animation/LiveLink/Source/LiveLinkGraphNode/Public/K2Node_EvaluateLiveLinkFrame.h
--- a/Engine/Plugins/Animation/LiveLink/Source/LiveLinkGraphNode/Public/K2Node_EvaluateLiveLinkFrame.h
+++ b/Engine/Plugins/Animation/LiveLink/Source/LiveLinkGraphNode/Public/K2Node_EvaluateLiveLinkFrame.h
@@ -90,4 +90,7 @@ protected:
 	void RefreshDataOutputPinType();
 
 	TSubclassOf<ULiveLinkRole> GetDefaultRolePinValue() const;
+
+	/** Verify if selected role class is valid for evaluation */
+	bool IsRoleValidForEvaluation(TSubclassOf<ULiveLinkRole> InRoleClass) const;
 };
diff --git a/Engine/Source/Runtime/LiveLinkInterface/Public/Roles/LiveLinkLightRole.h b/Engine/Source/Runtime/LiveLinkInterface/Public/Roles/LiveLinkLightRole.h
--- a/Engine/Source/Runtime/LiveLinkInterface/Public/Roles/LiveLinkLightRole.h
+++ b/Engine/Source/Runtime/LiveLinkInterface/Public/Roles/LiveLinkLightRole.h
@@ -8,7 +8,7 @@
 /**
  * Role associated for Light data.
  */
-UCLASS(BlueprintType)
+UCLASS(BlueprintType, meta = (DisplayName = "Light Role"))
 class LIVELINKINTERFACE_API ULiveLinkLightRole : public ULiveLinkTransformRole
 {
 	GENERATED_BODY()
