[LiveLink]

- Added support for multiple virtual subject sources
- Added support to create virtual subject sources from UI (not shown for now)
- Added better UI to show missing linked subjects from Animation VirtualSubject
- Added option in Animation virtual subject to append subject name to facilitate usage

diff --git a/Engine/Plugins/Animation/LiveLink/Source/LiveLink/Private/LiveLinkAnimationVirtualSubject.cpp b/Engine/Plugins/Animation/LiveLink/Source/LiveLink/Private/LiveLinkAnimationVirtualSubject.cpp
--- a/Engine/Plugins/Animation/LiveLink/Source/LiveLink/Private/LiveLinkAnimationVirtualSubject.cpp
+++ b/Engine/Plugins/Animation/LiveLink/Source/LiveLink/Private/LiveLinkAnimationVirtualSubject.cpp
@@ -11,7 +11,11 @@ namespace LiveLinkAnimationVirtualSubjectUtils
 {
 	void AddToBoneNames(TArray<FName>& BoneNames, const TArray<FName>& NewBoneNames, const FName Prefix)
 	{
-		FString NameFormat = Prefix.ToString() + TEXT("_");
+		FString NameFormat;
+		if (Prefix != NAME_None)
+		{
+			NameFormat = Prefix.ToString() + TEXT("_");
+		}
 
 		BoneNames.Reserve(BoneNames.Num() + NewBoneNames.Num());
 
@@ -133,7 +137,8 @@ void ULiveLinkAnimationVirtualSubject::BuildSkeleton(const TArray<FLiveLinkSubje
 			check(SubjectSnapShotData.StaticData.IsValid());
 			const FLiveLinkSkeletonStaticData* SubjectSkeletonData = SubjectSnapShotData.StaticData.Cast<FLiveLinkSkeletonStaticData>();
 
-			LiveLinkAnimationVirtualSubjectUtils::AddToBoneNames(BoneNames, SubjectSkeletonData->GetBoneNames(), Subjects[i]);
+			const FName BonePrefix = bAppendSubjectNameToBones ? Subjects[i] : NAME_None;
+			LiveLinkAnimationVirtualSubjectUtils::AddToBoneNames(BoneNames, SubjectSkeletonData->GetBoneNames(), BonePrefix);
 			LiveLinkAnimationVirtualSubjectUtils::AddToBoneParents(BoneParents, SubjectSkeletonData->GetBoneParents());
 			SkeletonData->PropertyNames.Append(SubjectSkeletonData->PropertyNames);
 		}
diff --git a/Engine/Plugins/Animation/LiveLink/Source/LiveLink/Private/LiveLinkClient.cpp b/Engine/Plugins/Animation/LiveLink/Source/LiveLink/Private/LiveLinkClient.cpp
--- a/Engine/Plugins/Animation/LiveLink/Source/LiveLink/Private/LiveLinkClient.cpp
+++ b/Engine/Plugins/Animation/LiveLink/Source/LiveLink/Private/LiveLinkClient.cpp
@@ -16,6 +16,7 @@
 #include "LiveLinkSourceFactory.h"
 #include "LiveLinkSourceSettings.h"
 #include "LiveLinkSubject.h"
+#include "LiveLinkVirtualSource.h"
 #include "IMediaModule.h"
 #include "Misc/App.h"
 #include "Misc/ConfigCacheIni.h"
@@ -86,6 +87,8 @@ void FLiveLinkClient::Tick()
 	CacheValues();
 	UpdateSources();
 	BuildThisTicksSubjectSnapshot();
+
+	OnLiveLinkTickedDelegate.Broadcast();
 }
 
 void FLiveLinkClient::DoPendingWork()
@@ -230,25 +233,51 @@ FGuid FLiveLinkClient::AddSource(TSharedPtr<ILiveLinkSource> InSource)
 	return Guid;
 }
 
-bool FLiveLinkClient::CreateSource(const FLiveLinkSourcePreset& InSourcePreset)
+FGuid FLiveLinkClient::AddVirtualSubjectSource(FName SourceName)
 {
 	check(Collection);
 
-	if (InSourcePreset.Settings == nullptr)
+	FGuid Guid;
+
+	if (Collection->FindVirtualSource(SourceName) == nullptr)
 	{
-		FLiveLinkLog::Warning(TEXT("Create Source Failure: The settings are not defined."));
-		return false;
+		TSharedPtr<FLiveLinkVirtualSubjectSource> Source = MakeShared<FLiveLinkVirtualSubjectSource>();
+		Guid = FGuid::NewGuid();
+
+		FLiveLinkCollectionSourceItem Data;
+		Data.Guid = Guid;
+		Data.Source = Source;
+
+		ULiveLinkVirtualSubjectSourceSettings* NewSettings = NewObject<ULiveLinkVirtualSubjectSourceSettings>(GetTransientPackage(), ULiveLinkVirtualSubjectSourceSettings::StaticClass());
+		NewSettings->SourceName = SourceName;
+		Data.Setting = NewSettings;
+		Data.bIsVirtualSource = true;
+		Collection->AddSource(Data);
+
+		Source->ReceiveClient(this, Data.Guid);
+		Source->InitializeSettings(Data.Setting);
+	}
+	else
+	{
+		FLiveLinkLog::Warning(TEXT("The virtual subject Source '%s' could not be created. It already exists."), *SourceName.ToString());
 	}
 
-	if (InSourcePreset.Settings->Factory.Get() == nullptr || InSourcePreset.Settings->Factory.Get() == ULiveLinkSourceFactory::StaticClass())
+	return Guid;
+}
+
+bool FLiveLinkClient::CreateSource(const FLiveLinkSourcePreset& InSourcePreset)
+{
+	check(Collection);
+
+	if (InSourcePreset.Settings == nullptr)
 	{
-		FLiveLinkLog::Warning(TEXT("Create Source Failure: The factory is not defined."));
+		FLiveLinkLog::Warning(TEXT("Create Source Failure: The settings are not defined."));
 		return false;
 	}
 
-	if (InSourcePreset.Guid == FLiveLinkSourceCollection::VirtualSubjectGuid)
+	if (InSourcePreset.Guid == FLiveLinkSourceCollection::DefaultVirtualSubjectGuid)
 	{
-		FLiveLinkLog::Warning(TEXT("Create Source Failure: Can't create a virtual source."));
+		FLiveLinkLog::Warning(TEXT("Create Source Failure: Can't create default virtual subject source. It will be created automatically."));
 		return false;
 	}
 
@@ -266,13 +295,29 @@ bool FLiveLinkClient::CreateSource(const FLiveLinkSourcePreset& InSourcePreset)
 
 	FLiveLinkCollectionSourceItem Data;
 	Data.Guid = InSourcePreset.Guid;
-	Data.Source = InSourcePreset.Settings->Factory.Get()->GetDefaultObject<ULiveLinkSourceFactory>()->CreateSource(InSourcePreset.Settings->ConnectionString);
-	if (!Data.Source.IsValid())
+
+	// Virtual subject source have a special settings class. We can differentiate them using this
+	if (InSourcePreset.Settings->GetClass() == ULiveLinkVirtualSubjectSourceSettings::StaticClass())
 	{
-		FLiveLinkLog::Warning(TEXT("Create Source Failure: The source couldn't be created by the factory."));
-		return false;
+		Data.Source = MakeShared<FLiveLinkVirtualSubjectSource>();
+		Data.bIsVirtualSource = true;
 	}
+	else
+	{
+		if (InSourcePreset.Settings->Factory.Get() == nullptr || InSourcePreset.Settings->Factory.Get() == ULiveLinkSourceFactory::StaticClass())
+		{
+			FLiveLinkLog::Warning(TEXT("Create Source Failure: The factory is not defined."));
+			return false;
+		}
 
+		Data.Source = InSourcePreset.Settings->Factory.Get()->GetDefaultObject<ULiveLinkSourceFactory>()->CreateSource(InSourcePreset.Settings->ConnectionString);
+		if (!Data.Source.IsValid())
+		{
+			FLiveLinkLog::Warning(TEXT("Create Source Failure: The source couldn't be created by the factory."));
+			return false;
+		}
+	}
+	
 	Data.Setting = DuplicateObject<ULiveLinkSourceSettings>(InSourcePreset.Settings, GetTransientPackage());
 	
 	Collection->AddSource(Data);
@@ -335,6 +380,21 @@ TArray<FGuid> FLiveLinkClient::GetSources() const
 	return Result;
 }
 
+TArray<FGuid> FLiveLinkClient::GetVirtualSources() const
+{
+	check(Collection);
+
+	TArray<FGuid> Result;
+	for (const FLiveLinkCollectionSourceItem& SourceItem : Collection->GetSources())
+	{
+		if (!SourceItem.bPendingKill && SourceItem.IsVirtualSource())
+		{
+			Result.Add(SourceItem.Guid);
+		}
+	}
+	return Result;
+}
+
 FLiveLinkSourcePreset FLiveLinkClient::GetSourcePreset(FGuid InSourceGuid, UObject* InDuplicatedObjectOuter) const
 {
 	check(Collection);
@@ -344,7 +404,7 @@ FLiveLinkSourcePreset FLiveLinkClient::GetSourcePreset(FGuid InSourceGuid, UObje
 	FLiveLinkSourcePreset SourcePreset;
 	if (FLiveLinkCollectionSourceItem* SourceItem = Collection->FindSource(InSourceGuid))
 	{
-		if (!SourceItem->IsVirtualSource() && SourceItem->Setting)
+		if (SourceItem->Guid != FLiveLinkSourceCollection::DefaultVirtualSubjectGuid && SourceItem->Setting)
 		{
 			SourcePreset.Guid = SourceItem->Guid;
 			SourcePreset.Settings = DuplicateObject<ULiveLinkSourceSettings>(SourceItem->Setting, DuplicatedObjectOuter);
@@ -625,15 +685,9 @@ bool FLiveLinkClient::CreateSubject(const FLiveLinkSubjectPreset& InSubjectPrese
 		return false;
 	}
 
-	if (InSubjectPreset.Key.Source == FLiveLinkSourceCollection::VirtualSubjectGuid && InSubjectPreset.VirtualSubject == nullptr)
-	{
-		FLiveLinkLog::Warning(TEXT("Create Source Failure: Can't create an empty virtual source."));
-		return false;
-	}
-
-	if (InSubjectPreset.Key.Source != FLiveLinkSourceCollection::VirtualSubjectGuid && InSubjectPreset.VirtualSubject != nullptr)
+	if (InSubjectPreset.Key.Source == FLiveLinkSourceCollection::DefaultVirtualSubjectGuid && InSubjectPreset.VirtualSubject == nullptr)
 	{
-		FLiveLinkLog::Warning(TEXT("Create Source Failure: Can't create a virtual source in another source."));
+		FLiveLinkLog::Warning(TEXT("Create Source Failure: Can't create an empty virtual subject."));
 		return false;
 	}
 
@@ -669,7 +723,7 @@ bool FLiveLinkClient::CreateSubject(const FLiveLinkSubjectPreset& InSubjectPrese
 	{
 		bool bEnabled = false;
 		ULiveLinkVirtualSubject* VSubject = DuplicateObject<ULiveLinkVirtualSubject>(InSubjectPreset.VirtualSubject, GetTransientPackage());
-		FLiveLinkCollectionSubjectItem VSubjectData(InSubjectPreset.Key.SubjectName, VSubject, bEnabled);
+		FLiveLinkCollectionSubjectItem VSubjectData(InSubjectPreset.Key, VSubject, bEnabled);
 		VSubject->Initialize(VSubjectData.Key, VSubject->GetRole(), this);
 
 		FScopeLock Lock(&CollectionAccessCriticalSection);
@@ -712,32 +766,54 @@ void FLiveLinkClient::RemoveSubject_AnyThread(const FLiveLinkSubjectKey& InSubje
 	}
 }
 
-void FLiveLinkClient::AddVirtualSubject(FLiveLinkSubjectName InNewVirtualSubjectName, TSubclassOf<ULiveLinkVirtualSubject> InVirtualSubjectClass)
+bool FLiveLinkClient::AddVirtualSubject(FLiveLinkSubjectKey InVirtualSubjectKey, TSubclassOf<ULiveLinkVirtualSubject> InVirtualSubjectClass)
 {
-	if (Collection && !InNewVirtualSubjectName.IsNone() && InVirtualSubjectClass != nullptr)
-	{
-		bool bFoundVirtualSubjectWithSameName = nullptr != Collection->GetSubjects().FindByPredicate(
-			[InNewVirtualSubjectName](const FLiveLinkCollectionSubjectItem& Other)
-		{
-			return Other.Key.SubjectName == InNewVirtualSubjectName && Other.GetVirtualSubject() != nullptr;
-		});
+	bool bResult = false;
 
-		if (!bFoundVirtualSubjectWithSameName)
+	if (Collection && !InVirtualSubjectKey.SubjectName.IsNone() && InVirtualSubjectClass != nullptr)
+	{
+		FLiveLinkCollectionSourceItem* SourceItem = Collection->FindSource(InVirtualSubjectKey.Source);
+		if (SourceItem == nullptr || SourceItem->bPendingKill)
 		{
-			ULiveLinkVirtualSubject* VSubject = NewObject<ULiveLinkVirtualSubject>(GetTransientPackage(), InVirtualSubjectClass.Get());
-			bool bEnabled = Collection->FindEnabledSubject(InNewVirtualSubjectName) == nullptr;
-			FLiveLinkCollectionSubjectItem VSubjectData(InNewVirtualSubjectName, VSubject, bEnabled);
-
-			VSubject->Initialize(VSubjectData.Key, VSubject->GetRole(), this);
-
-			FScopeLock Lock(&CollectionAccessCriticalSection);
-			Collection->AddSubject(MoveTemp(VSubjectData));
+			FLiveLinkLog::Warning(TEXT("Create Virtual Subject Failure: The source doesn't exist."));
 		}
 		else
 		{
-			FLiveLinkLog::Warning(TEXT("The virtual subject '%s' could not be created."), *InNewVirtualSubjectName.Name.ToString());
+			FScopeLock Lock(&CollectionAccessCriticalSection);
+			const bool bFoundVirtualSubject = nullptr != Collection->GetSubjects().FindByPredicate(
+				[InVirtualSubjectKey](const FLiveLinkCollectionSubjectItem& Other)
+			{
+				return Other.Key == InVirtualSubjectKey && Other.GetVirtualSubject() != nullptr;
+			});
+
+			if (!bFoundVirtualSubject)
+			{
+				ULiveLinkVirtualSubject* VSubject = NewObject<ULiveLinkVirtualSubject>(GetTransientPackage(), InVirtualSubjectClass.Get());
+				const bool bDoEnableSubject = Collection->FindEnabledSubject(InVirtualSubjectKey.SubjectName) == nullptr;
+				FLiveLinkCollectionSubjectItem VSubjectData(InVirtualSubjectKey, VSubject, bDoEnableSubject);
+
+				VSubject->Initialize(VSubjectData.Key, VSubject->GetRole(), this);
+				Collection->AddSubject(MoveTemp(VSubjectData));
+
+				bResult = true;
+			}
+			else
+			{
+				FLiveLinkLog::Warning(TEXT("The virtual subject '%s' could not be created."), *InVirtualSubjectKey.SubjectName.Name.ToString());
+			}
 		}
 	}
+
+	return bResult;
+}
+
+void FLiveLinkClient::RemoveVirtualSubject(FLiveLinkSubjectKey InVirtualSubjectKey)
+{
+	FScopeLock Lock(&CollectionAccessCriticalSection);
+	if (Collection)
+	{
+		Collection->RemoveSubject(InVirtualSubjectKey);
+	}
 }
 
 void FLiveLinkClient::ClearSubjectsFrames_AnyThread(FLiveLinkSubjectName InSubjectName)
@@ -1090,6 +1166,11 @@ bool FLiveLinkClient::EvaluateFrameAtSceneTime_AnyThread(FLiveLinkSubjectName In
 	return bResult;
 }
 
+FSimpleMulticastDelegate& FLiveLinkClient::OnLiveLinkTicked()
+{
+	return OnLiveLinkTickedDelegate;
+}
+
 TArray<FGuid> FLiveLinkClient::GetDisplayableSources() const
 {
 	TArray<FGuid> Results;
@@ -1332,6 +1413,12 @@ void FLiveLinkClient::AddVirtualSubject(FName InNewVirtualSubjectName)
 	AddVirtualSubject(InNewVirtualSubjectName, ULiveLinkAnimationVirtualSubject::StaticClass());
 }
 
+void FLiveLinkClient::AddVirtualSubject(FLiveLinkSubjectName VirtualSubjectName, TSubclassOf<ULiveLinkVirtualSubject> VirtualSubjectClass)
+{
+	const FLiveLinkSubjectKey DefaultSubjectKey(FLiveLinkSourceCollection::DefaultVirtualSubjectGuid, VirtualSubjectName);
+	AddVirtualSubject(DefaultSubjectKey, ULiveLinkAnimationVirtualSubject::StaticClass());
+}
+
 /**
  * Deprecated function
  */
diff --git a/Engine/Plugins/Animation/LiveLink/Source/LiveLink/Private/LiveLinkMessageBusSourceFactory.h b/Engine/Plugins/Animation/LiveLink/Source/LiveLink/Private/LiveLinkMessageBusSourceFactory.h
--- a/Engine/Plugins/Animation/LiveLink/Source/LiveLink/Private/LiveLinkMessageBusSourceFactory.h
+++ b/Engine/Plugins/Animation/LiveLink/Source/LiveLink/Private/LiveLinkMessageBusSourceFactory.h
@@ -7,8 +7,6 @@
 #include "LiveLinkMessageBusSourceFactory.generated.h"
 
 
-class SLiveLinkMessageBusSourceEditor;
-
 UCLASS()
 class ULiveLinkMessageBusSourceFactory : public ULiveLinkSourceFactory
 {
diff --git a/Engine/Plugins/Animation/LiveLink/Source/LiveLink/Private/LiveLinkPreset.cpp b/Engine/Plugins/Animation/LiveLink/Source/LiveLink/Private/LiveLinkPreset.cpp
--- a/Engine/Plugins/Animation/LiveLink/Source/LiveLink/Private/LiveLinkPreset.cpp
+++ b/Engine/Plugins/Animation/LiveLink/Source/LiveLink/Private/LiveLinkPreset.cpp
@@ -43,6 +43,15 @@ void ULiveLinkPreset::BuildFromClient()
 			Sources.Add(LiveLinkClient.GetSourcePreset(SourceGuid, this));
 		}
 
+		for (FGuid SourceGuid : LiveLinkClient.GetVirtualSources())
+		{
+			FLiveLinkSourcePreset NewPreset = LiveLinkClient.GetSourcePreset(SourceGuid, this);
+			if (NewPreset.Guid.IsValid())
+			{
+				Sources.Add(MoveTemp(NewPreset));
+			}
+		}
+
 		for (FLiveLinkSubjectKey SubjectKey : LiveLinkClient.GetSubjects(true, true))
 		{
 			Subjects.Add(LiveLinkClient.GetSubjectPreset(SubjectKey, this));
diff --git a/Engine/Plugins/Animation/LiveLink/Source/LiveLink/Private/LiveLinkSourceCollection.cpp b/Engine/Plugins/Animation/LiveLink/Source/LiveLink/Private/LiveLinkSourceCollection.cpp
--- a/Engine/Plugins/Animation/LiveLink/Source/LiveLink/Private/LiveLinkSourceCollection.cpp
+++ b/Engine/Plugins/Animation/LiveLink/Source/LiveLink/Private/LiveLinkSourceCollection.cpp
@@ -5,24 +5,23 @@
 #include "ILiveLinkSource.h"
 #include "LiveLinkSourceSettings.h"
 #include "LiveLinkSubjectSettings.h"
+#include "LiveLinkVirtualSource.h"
 
-
-// Completely empty "source" that virtual subjects can hang off
-struct FLiveLinkVirtualSubjectSource : public ILiveLinkSource
+/**
+ * Default VirtualSubject Source.
+ */ 
+struct FLiveLinkDefaultVirtualSubjectSource : public FLiveLinkVirtualSubjectSource
 {
-	virtual bool CanBeDisplayedInUI() const override { return false; }
-	virtual void ReceiveClient(ILiveLinkClient* InClient, FGuid InSourceGuid) override {}
-	virtual bool IsSourceStillValid() const override { return true; }
-	virtual bool RequestSourceShutdown() override { return true; }
+	FLiveLinkDefaultVirtualSubjectSource() = default;
+	virtual ~FLiveLinkDefaultVirtualSubjectSource() = default;
 
-	virtual FText GetSourceType() const override { return NSLOCTEXT("TempLocTextLiveLink", "LiveLinkVirtualSubjectName", "Virtual Subjects"); }
-	virtual FText GetSourceMachineName() const override { return FText(); }
-	virtual FText GetSourceStatus() const override { return FText(); }
+	virtual bool CanBeDisplayedInUI() const override { return false; }
 };
 
+
 bool FLiveLinkCollectionSourceItem::IsVirtualSource() const
 {
-	return Guid == FLiveLinkSourceCollection::VirtualSubjectGuid;
+	return bIsVirtualSource;
 }
 
 FLiveLinkCollectionSubjectItem::FLiveLinkCollectionSubjectItem(FLiveLinkSubjectKey InKey, TUniquePtr<FLiveLinkSubject> InLiveSubject, ULiveLinkSubjectSettings* InSettings, bool bInEnabled)
@@ -35,26 +34,29 @@ FLiveLinkCollectionSubjectItem::FLiveLinkCollectionSubjectItem(FLiveLinkSubjectK
 {
 }
 
-FLiveLinkCollectionSubjectItem::FLiveLinkCollectionSubjectItem(FLiveLinkSubjectName InSubjectName, ULiveLinkVirtualSubject* InVirtualSubject, bool bInEnabled)
-	: bEnabled(bInEnabled)
+FLiveLinkCollectionSubjectItem::FLiveLinkCollectionSubjectItem(FLiveLinkSubjectKey InKey, ULiveLinkVirtualSubject* InVirtualSubject, bool bInEnabled)
+	: Key(InKey)
+	, bEnabled(bInEnabled)
 	, bPendingKill(false)
 	, Setting(nullptr)
 	, VirtualSubject(InVirtualSubject)
 {
-	Key.Source = FLiveLinkSourceCollection::VirtualSubjectGuid;
-	Key.SubjectName = InSubjectName;
 }
 
 
-const FGuid FLiveLinkSourceCollection::VirtualSubjectGuid{ 0x4ed2dc4e, 0xcc5911ce, 0x4af0635d, 0xa8b24a5a };
+const FGuid FLiveLinkSourceCollection::DefaultVirtualSubjectGuid{ 0x4ed2dc4e, 0xcc5911ce, 0x4af0635d, 0xa8b24a5a };
 
 
 FLiveLinkSourceCollection::FLiveLinkSourceCollection()
 {
 	FLiveLinkCollectionSourceItem Data;
-	Data.Source = MakeShared<FLiveLinkVirtualSubjectSource>();
-	Data.Guid = VirtualSubjectGuid;
-	Data.Setting = nullptr;
+	Data.Source = MakeShared<FLiveLinkDefaultVirtualSubjectSource>();
+	Data.Guid = DefaultVirtualSubjectGuid;
+	ULiveLinkVirtualSubjectSourceSettings* NewSettings = NewObject<ULiveLinkVirtualSubjectSourceSettings>(GetTransientPackage(), ULiveLinkVirtualSubjectSourceSettings::StaticClass());
+	NewSettings->SourceName = TEXT("DefaultVirtualSource");
+	Data.Setting = NewSettings;
+	Data.bIsVirtualSource = true;
+	Data.Source->InitializeSettings(NewSettings);
 
 	Sources.Add(Data);
 }
@@ -83,7 +85,7 @@ void FLiveLinkSourceCollection::AddSource(FLiveLinkCollectionSourceItem InSource
 
 void FLiveLinkSourceCollection::RemoveSource(FGuid InSourceGuid)
 {
-	if (InSourceGuid != FLiveLinkSourceCollection::VirtualSubjectGuid)
+	if (InSourceGuid != FLiveLinkSourceCollection::DefaultVirtualSubjectGuid)
 	{
 		int32 SourceIndex = Sources.IndexOfByPredicate([InSourceGuid](const FLiveLinkCollectionSourceItem& Other) { return Other.Guid == InSourceGuid; });
 		if (SourceIndex != INDEX_NONE)
@@ -131,7 +133,7 @@ void FLiveLinkSourceCollection::RemoveAllSources()
 	bool bHasRemovedSource = false;
 	for (int32 Index = Sources.Num() - 1; Index >= 0; --Index)
 	{
-		if (Sources[Index].Guid != FLiveLinkSourceCollection::VirtualSubjectGuid)
+		if (Sources[Index].Guid != FLiveLinkSourceCollection::DefaultVirtualSubjectGuid)
 		{
 			bHasRemovedSource = true;
 			FGuid Key = Sources[Index].Guid;
@@ -170,6 +172,26 @@ const FLiveLinkCollectionSourceItem* FLiveLinkSourceCollection::FindSource(FGuid
 }
 
 
+FLiveLinkCollectionSourceItem* FLiveLinkSourceCollection::FindVirtualSource(FName VirtualSourceName)
+{
+	return Sources.FindByPredicate([VirtualSourceName](const FLiveLinkCollectionSourceItem& Other)
+	{
+		if (Other.IsVirtualSource())
+		{
+			if (ULiveLinkVirtualSubjectSourceSettings* VirtualSubjectSettings = Cast<ULiveLinkVirtualSubjectSourceSettings>(Other.Setting))
+			{
+				return VirtualSubjectSettings->SourceName == VirtualSourceName;
+			}
+		}
+		return false; 
+	});
+}
+
+const FLiveLinkCollectionSourceItem* FLiveLinkSourceCollection::FindVirtualSource(FName VirtualSourceName) const
+{
+	return const_cast<FLiveLinkSourceCollection*>(this)->FindVirtualSource(VirtualSourceName);
+}
+
 void FLiveLinkSourceCollection::AddSubject(FLiveLinkCollectionSubjectItem InSubject)
 {
 	Subjects.Add(MoveTemp(InSubject));
@@ -256,12 +278,12 @@ void FLiveLinkSourceCollection::RemovePendingKill()
 		FLiveLinkCollectionSourceItem& SourceItem = Sources[SourceIndex];
 		if (SourceItem.bPendingKill)
 		{
-			if (SourceItem.IsVirtualSource())
+			if (SourceItem.Guid == FLiveLinkSourceCollection::DefaultVirtualSubjectGuid)
 			{
-				// Keep the source but mark the subject as pending kill
+				// Keep the default virtual subject source but mark the subject as pending kill
 				for (FLiveLinkCollectionSubjectItem& SubjectItem : Subjects)
 				{
-					if (SubjectItem.Key.Source == FLiveLinkSourceCollection::VirtualSubjectGuid)
+					if (SubjectItem.Key.Source == FLiveLinkSourceCollection::DefaultVirtualSubjectGuid)
 					{
 						SubjectItem.bPendingKill = true;
 					}
diff --git a/Engine/Plugins/Animation/LiveLink/Source/LiveLink/Private/LiveLinkSourceCollection.h b/Engine/Plugins/Animation/LiveLink/Source/LiveLink/Private/LiveLinkSourceCollection.h
--- a/Engine/Plugins/Animation/LiveLink/Source/LiveLink/Private/LiveLinkSourceCollection.h
+++ b/Engine/Plugins/Animation/LiveLink/Source/LiveLink/Private/LiveLinkSourceCollection.h
@@ -27,6 +27,7 @@ struct FLiveLinkCollectionSourceItem
 	ULiveLinkSourceSettings* Setting; // GC by FLiveLinkSourceCollection::AddReferencedObjects
 	TSharedPtr<ILiveLinkSource> Source;
 	bool bPendingKill = false;
+	bool bIsVirtualSource = false;
 
 public:
 	bool IsVirtualSource() const;
@@ -36,7 +37,7 @@ public:
 struct FLiveLinkCollectionSubjectItem
 {
 	FLiveLinkCollectionSubjectItem(FLiveLinkSubjectKey InKey, TUniquePtr<FLiveLinkSubject> InLiveSubject, ULiveLinkSubjectSettings* InSettings, bool bInEnabled);
-	FLiveLinkCollectionSubjectItem(FLiveLinkSubjectName InSubjectName, ULiveLinkVirtualSubject* InVirtualSubject, bool bInEnabled);
+	FLiveLinkCollectionSubjectItem(FLiveLinkSubjectKey InKey, ULiveLinkVirtualSubject* InVirtualSubject, bool bInEnabled);
 
 public:
 	FLiveLinkSubjectKey Key;
@@ -71,7 +72,7 @@ class FLiveLinkSourceCollection : public FGCObject
 {
 public:
 	// "source guid" for virtual subjects
-	static const FGuid VirtualSubjectGuid;
+	static const FGuid DefaultVirtualSubjectGuid;
 	FLiveLinkSourceCollection();
 
 public:
@@ -91,6 +92,8 @@ public:
 	const FLiveLinkCollectionSourceItem* FindSource(TSharedPtr<ILiveLinkSource> Source) const;
 	FLiveLinkCollectionSourceItem* FindSource(FGuid SourceGuid);
 	const FLiveLinkCollectionSourceItem* FindSource(FGuid SourceGuid) const;
+	FLiveLinkCollectionSourceItem* FindVirtualSource(FName VirtualSourceName);
+	const FLiveLinkCollectionSourceItem* FindVirtualSource(FName VirtualSourceName) const;
 
 	void AddSubject(FLiveLinkCollectionSubjectItem Subject);
 	void RemoveSubject(FLiveLinkSubjectKey SubjectKey);
diff --git a/Engine/Plugins/Animation/LiveLink/Source/LiveLink/Private/LiveLinkVirtualSource.cpp b/Engine/Plugins/Animation/LiveLink/Source/LiveLink/Private/LiveLinkVirtualSource.cpp
new file mode 100644
--- /dev/null
+++ b/Engine/Plugins/Animation/LiveLink/Source/LiveLink/Private/LiveLinkVirtualSource.cpp
@@ -0,0 +1,54 @@
+// Copyright 1998-2019 Epic Games, Inc. All Rights Reserved.
+
+#include "LiveLinkVirtualSource.h"
+
+
+#define LOCTEXT_NAMESPACE "LiveLinkVirtualSubjectSource"
+
+
+
+bool FLiveLinkVirtualSubjectSource::CanBeDisplayedInUI() const 
+{
+	return true; 
+}
+
+void FLiveLinkVirtualSubjectSource::ReceiveClient(ILiveLinkClient* InClient, FGuid InSourceGuid) 
+{
+
+}
+
+void FLiveLinkVirtualSubjectSource::InitializeSettings(ULiveLinkSourceSettings* Settings)
+{
+	ULiveLinkVirtualSubjectSourceSettings* VirtualSettings = Cast<ULiveLinkVirtualSubjectSourceSettings>(Settings);
+	check(VirtualSettings);
+	
+	SourceName = VirtualSettings->SourceName;
+}
+
+bool FLiveLinkVirtualSubjectSource::IsSourceStillValid() const
+{
+	return true; 
+}
+
+bool FLiveLinkVirtualSubjectSource::RequestSourceShutdown()
+{
+	return true; 
+}
+
+FText FLiveLinkVirtualSubjectSource::GetSourceType() const
+{
+	return FText::FromName(SourceName);
+}
+
+FText FLiveLinkVirtualSubjectSource::GetSourceMachineName() const
+{
+	return FText();
+}
+
+FText FLiveLinkVirtualSubjectSource::GetSourceStatus() const
+{
+	return LOCTEXT("VirtualSubjectSourceStatus", "Active");
+}
+
+#undef LOCTEXT_NAMESPACE
+
diff --git a/Engine/Plugins/Animation/LiveLink/Source/LiveLink/Private/LiveLinkVirtualSource.h b/Engine/Plugins/Animation/LiveLink/Source/LiveLink/Private/LiveLinkVirtualSource.h
new file mode 100644
--- /dev/null
+++ b/Engine/Plugins/Animation/LiveLink/Source/LiveLink/Private/LiveLinkVirtualSource.h
@@ -0,0 +1,49 @@
+// Copyright 1998-2019 Epic Games, Inc. All Rights Reserved.
+
+#pragma once
+
+#include "ILiveLinkSource.h"
+
+#include "CoreMinimal.h"
+
+#include "LiveLinkVirtualSource.generated.h"
+
+class ILiveLinkClient;
+
+
+/**
+ * Completely empty "source" that virtual subjects can hang off
+ */
+class FLiveLinkVirtualSubjectSource : public ILiveLinkSource
+{
+public:
+	FLiveLinkVirtualSubjectSource() = default;
+	virtual ~FLiveLinkVirtualSubjectSource() = default;
+
+	//~ Begin ILiveLinkSource interface
+	virtual bool CanBeDisplayedInUI() const override;
+	virtual void ReceiveClient(ILiveLinkClient* InClient, FGuid InSourceGuid) override;
+	virtual void InitializeSettings(ULiveLinkSourceSettings* Settings) override;
+	virtual bool IsSourceStillValid() const override;
+	virtual bool RequestSourceShutdown() override;
+
+	virtual FText GetSourceType() const override;
+	virtual FText GetSourceMachineName() const override;
+	virtual FText GetSourceStatus() const override;
+	//~ End ILiveLinkSource interface
+
+protected:
+	FName SourceName;
+};
+
+/** VirtualSubjectSource Settings to be able to differentiate from live sources and keep a name associated to the source */
+UCLASS()
+class ULiveLinkVirtualSubjectSourceSettings : public ULiveLinkSourceSettings
+{
+	GENERATED_BODY()
+
+public:
+
+	UPROPERTY()
+	FName SourceName;
+};
diff --git a/Engine/Plugins/Animation/LiveLink/Source/LiveLink/Public/LiveLinkAnimationVirtualSubject.h b/Engine/Plugins/Animation/LiveLink/Source/LiveLink/Public/LiveLinkAnimationVirtualSubject.h
--- a/Engine/Plugins/Animation/LiveLink/Source/LiveLink/Public/LiveLinkAnimationVirtualSubject.h
+++ b/Engine/Plugins/Animation/LiveLink/Source/LiveLink/Public/LiveLinkAnimationVirtualSubject.h
@@ -43,4 +43,10 @@ protected:
 	// Tests to see if current ref skeleton is up to data
 	bool DoesSkeletonNeedRebuilding() const;
 	bool bInvalidate;
+
+public:
+
+	//Whether to append SubjectName to each bones part of the virtual hierarchy
+	UPROPERTY(EditAnywhere, Category = "Settings")
+	bool bAppendSubjectNameToBones;
 };
diff --git a/Engine/Plugins/Animation/LiveLink/Source/LiveLink/Public/LiveLinkClient.h b/Engine/Plugins/Animation/LiveLink/Source/LiveLink/Public/LiveLinkClient.h
--- a/Engine/Plugins/Animation/LiveLink/Source/LiveLink/Public/LiveLinkClient.h
+++ b/Engine/Plugins/Animation/LiveLink/Source/LiveLink/Public/LiveLinkClient.h
@@ -73,11 +73,13 @@ public:
 
 	//~ Begin ILiveLinkClient implementation
 	virtual FGuid AddSource(TSharedPtr<ILiveLinkSource> Source) override;
+	virtual FGuid AddVirtualSubjectSource(FName SourceName) override;
 	virtual bool CreateSource(const FLiveLinkSourcePreset& SourcePreset) override;
 	virtual void RemoveSource(TSharedPtr<ILiveLinkSource> Source) override;
 	virtual void RemoveSource(FGuid InEntryGuid) override;
 	virtual bool HasSourceBeenAdded(TSharedPtr<ILiveLinkSource> Source) const override;
 	virtual TArray<FGuid> GetSources() const override;
+	virtual TArray<FGuid> GetVirtualSources() const override;
 	virtual FLiveLinkSourcePreset GetSourcePreset(FGuid SourceGuid, UObject* DuplicatedObjectOuter) const override;
 	virtual FText GetSourceType(FGuid EntryGuid) const override;
 
@@ -85,6 +87,8 @@ public:
 	virtual void PushSubjectFrameData_AnyThread(const FLiveLinkSubjectKey& SubjectKey, FLiveLinkFrameDataStruct&& FrameData) override;
 
 	virtual bool CreateSubject(const FLiveLinkSubjectPreset& SubjectPreset) override;
+	virtual bool AddVirtualSubject(FLiveLinkSubjectKey VirtualSubjectKey, TSubclassOf<ULiveLinkVirtualSubject> VirtualSubjectClass) override;
+	virtual void RemoveVirtualSubject(FLiveLinkSubjectKey VirtualSubjectKey) override;
 	virtual void RemoveSubject_AnyThread(const FLiveLinkSubjectKey& SubjectKey) override;
 	virtual void ClearSubjectsFrames_AnyThread(FLiveLinkSubjectName SubjectName) override;
 	virtual void ClearSubjectsFrames_AnyThread(const FLiveLinkSubjectKey& InSubjectKey) override;
@@ -107,12 +111,14 @@ public:
 	virtual bool DoesSubjectSupportsRole(const FLiveLinkSubjectKey& SubjectKey, TSubclassOf<ULiveLinkRole> SupportedRole) const override;
 	virtual TArray<FLiveLinkTime> GetSubjectFrameTimes(const FLiveLinkSubjectKey& SubjectKey) const override;
 	virtual TArray<FLiveLinkTime> GetSubjectFrameTimes(FLiveLinkSubjectName SubjectName) const override;
+	virtual UObject* GetSubjectSettings(const FLiveLinkSubjectKey& SubjectKey) const override;
 
 	virtual bool EvaluateFrameFromSource_AnyThread(const FLiveLinkSubjectKey& SubjectKey, TSubclassOf<ULiveLinkRole> Role, FLiveLinkSubjectFrameData& OutFrame) override;
 	virtual bool EvaluateFrame_AnyThread(FLiveLinkSubjectName SubjectName, TSubclassOf<ULiveLinkRole> Role, FLiveLinkSubjectFrameData& OutFrame) override;
 	virtual bool EvaluateFrameAtWorldTime_AnyThread(FLiveLinkSubjectName SubjectName, double WorldTime, TSubclassOf<ULiveLinkRole> DesiredRole, FLiveLinkSubjectFrameData& OutFrame) override;
 	virtual bool EvaluateFrameAtSceneTime_AnyThread(FLiveLinkSubjectName SubjectName, const FTimecode& SceneTime, TSubclassOf<ULiveLinkRole> DesiredRole, FLiveLinkSubjectFrameData& OutFrame) override;
 
+	virtual FSimpleMulticastDelegate& OnLiveLinkTicked() override;
 	virtual FSimpleMulticastDelegate& OnLiveLinkSourcesChanged() override;
 	virtual FSimpleMulticastDelegate& OnLiveLinkSubjectsChanged() override;
 	virtual FOnLiveLinkSourceChangedDelegate& OnLiveLinkSourceAdded() override;
@@ -133,9 +139,6 @@ public:
 	/** Remove all sources from the live link client */
 	void RemoveAllSources();
 
-	/** Add a new virtual subject to the client */
-	void AddVirtualSubject(FLiveLinkSubjectName VirtualSubjectName, TSubclassOf<ULiveLinkVirtualSubject> VirtualSubjectClass);
-
 	/** Is the supplied subject virtual */
 	bool IsVirtualSubject(const FLiveLinkSubjectKey& Subject) const;
 
@@ -144,8 +147,6 @@ public:
 
 	TArray<FGuid> GetDisplayableSources() const;
 
-	UObject* GetSubjectSettings(const FLiveLinkSubjectKey& SubjectKey) const;
-
 	ULiveLinkSourceSettings* GetSourceSettings(FGuid InEntryGuid) const;
 
 	FLiveLinkSubjectTimeSyncData GetTimeSyncData(FLiveLinkSubjectName SubjectName);
@@ -163,6 +164,8 @@ public:
 	const TArray<FGuid>& GetSourceEntries() const;
 	UE_DEPRECATED(4.23, "FLiveLinkClient::AddVirtualSubject(FName) is deprecated. Please use AddVirtualSubject(FName, TSubClass<ULiveLinkVirtualSubject>) instead!")
 	void AddVirtualSubject(FName NewVirtualSubjectName);
+	UE_DEPRECATED(4.25, "FLiveLinkClient::AddVirtualSubject(FLiveLinkSubjectName, TSubclassOf<ULiveLinkVirtualSubject>) is deprecated. Please use bool AddVirtualSubject(FLiveLinkSubjectKey, TSubClass<ULiveLinkVirtualSubject>) instead!")
+	void AddVirtualSubject(FLiveLinkSubjectName VirtualSubjectName, TSubclassOf<ULiveLinkVirtualSubject> VirtualSubjectClass);
 	UE_DEPRECATED(4.23, "FLiveLinkClient::OnLiveLinkSourcesChanged is deprecated. Please use OnLiveLinkSourceAdded instead!")
 	FDelegateHandle RegisterSourcesChangedHandle(const FSimpleMulticastDelegate::FDelegate& SourcesChanged);
 	UE_DEPRECATED(4.23, "FLiveLinkClient::OnLiveLinkSourcesChanged is deprecated. Please use OnLiveLinkSourceAdded instead!")
@@ -231,7 +234,7 @@ private:
 	/** Key funcs for looking up a set of cached keys by its layout element */
 	TMap<FLiveLinkSubjectName, FLiveLinkSubjectKey> EnabledSubjects;
 
-	/** Lock to stop multiple threads accessing the CurrentPreset at the same time */
+	/** Lock to stop multiple threads accessing the Subjects from the collection at the same time */
 	mutable FCriticalSection CollectionAccessCriticalSection;
 
 	struct FSubjectFramesReceivedHandles
@@ -243,6 +246,9 @@ private:
 	/** Map of delegates to notify interested parties when the client receives a static or data frame for each subject */
 	TMap<FLiveLinkSubjectName, FSubjectFramesReceivedHandles> SubjectFrameReceivedHandles;
 
+	/** Delegate when LiveLinkClient has ticked. */
+	FSimpleMulticastDelegate OnLiveLinkTickedDelegate;
+
 #if WITH_EDITOR
 	/** Delegate when a subject is evaluated. */
 	FOnLiveLinkSubjectEvaluated OnLiveLinkSubjectEvaluatedDelegate;
diff --git a/Engine/Plugins/Animation/LiveLink/Source/LiveLinkEditor/Private/LiveLinkClientPanelToolbar.cpp b/Engine/Plugins/Animation/LiveLink/Source/LiveLinkEditor/Private/LiveLinkClientPanelToolbar.cpp
--- a/Engine/Plugins/Animation/LiveLink/Source/LiveLinkEditor/Private/LiveLinkClientPanelToolbar.cpp
+++ b/Engine/Plugins/Animation/LiveLink/Source/LiveLinkEditor/Private/LiveLinkClientPanelToolbar.cpp
@@ -56,6 +56,9 @@ class SVirtualSubjectCreateDialog : public SCompoundWidget
 {
 public:
 	SLATE_BEGIN_ARGS(SVirtualSubjectCreateDialog) {}
+	
+		/** Pointer to the LiveLinkClient instance. */
+		SLATE_ARGUMENT(FLiveLinkClient*, LiveLinkClient)
 
 	SLATE_END_ARGS()
 
@@ -66,6 +69,14 @@ public:
 		bOkClicked = false;
 		VirtualSubjectClass = nullptr;
 		VirtualSubjectName = DefaultVirtualSubjectName;
+		LiveLinkClient = InArgs._LiveLinkClient;
+
+		check(LiveLinkClient);
+
+		//Default VirtualSubject Source should always exist
+		TArray<FGuid> Sources = LiveLinkClient->GetVirtualSources();
+		check(Sources.Num() > 0);
+		VirtualSourceGuid = Sources[0];
 
 		TSharedPtr<STextEntryPopup> TextEntry;
 		SAssignNew(TextEntry, STextEntryPopup)
@@ -83,6 +94,41 @@ public:
 				SNew(SBox)
 				[
 					SNew(SVerticalBox)
+					
+					//For now, it's not possible to create VirtualSubject Sources from the UI. Also, VirtualSubjects created from the UI will be parented under the default VirtualSubject Source
+					//Code is present in case it's required in the future.
+					/*+ SVerticalBox::Slot()
+					.HAlign(HAlign_Fill)
+					.AutoHeight()
+					[
+						SNew(SHorizontalBox)
+						+ SHorizontalBox::Slot()
+						.AutoWidth()
+						.HAlign(HAlign_Left)
+						.Padding(FMargin(0.0f, 0.0f, 5.0f, 0.0f))
+						[
+							SNew(STextBlock)
+							.Text(LOCTEXT("SourceNameLabel", "Parent Source:"))
+						]
+						+ SHorizontalBox::Slot()
+						.AutoWidth()
+						.HAlign(HAlign_Left)
+						.Padding(FMargin(5.0f, 0.0f, 5.0f, 0.0f))
+						[
+							SNew(STextBlock)
+							.Text(this, &SVirtualSubjectCreateDialog::GetParentSourceText)
+						]
+						+ SHorizontalBox::Slot()
+						.AutoWidth()
+						.HAlign(HAlign_Left)
+						.Padding(FMargin(5.0f, 0.0f, 0.0f, 0.0f))
+						[
+							SNew(SComboButton)
+							.OnGetMenuContent(this, &SVirtualSubjectCreateDialog::HandleSourceSelectionComboButton)
+							.ContentPadding(FMargin(4.0, 2.0))
+						]
+					]*/
+
 					+ SVerticalBox::Slot()
 					.HAlign(HAlign_Fill)
 					.AutoHeight()
@@ -142,10 +188,8 @@ public:
 		return VirtualSubjectClass != nullptr;
 	}
 
-	bool ConfigureVirtualSubject(FLiveLinkClient* InLiveLinkClient)
+	bool ConfigureVirtualSubject()
 	{
-		LiveLinkClient = InLiveLinkClient;
-
 		TSharedRef<SWindow> Window = SNew(SWindow)
 			.Title(LOCTEXT("CreateVirtualSubjectCreation", "Create Virtual Subject"))
 			.ClientSize(FVector2D(400, 300))
@@ -194,6 +238,63 @@ private:
 		}
 	};
 
+	TSharedRef<SWidget> HandleSourceSelectionComboButton()
+	{
+		// Generate menu
+		FMenuBuilder MenuBuilder(true, nullptr);
+		MenuBuilder.BeginSection("AvailableVirtualSources", LOCTEXT("AvailableSources", "VirtualSubject Sources"));
+		{
+			if (LiveLinkClient)
+			{
+				for (const FGuid& SourceGuid : LiveLinkClient->GetVirtualSources())
+				{
+					//Always add a None entry
+					MenuBuilder.AddMenuEntry(
+						LiveLinkClient->GetSourceType(SourceGuid),
+						FText::FromName(NAME_None),
+						FSlateIcon(),
+						FUIAction(
+							FExecuteAction::CreateSP(this, &SVirtualSubjectCreateDialog::HandleVirtualSourceSelection, SourceGuid),
+							FCanExecuteAction(),
+							FIsActionChecked::CreateSP(this, &SVirtualSubjectCreateDialog::IsVirtualSourceSelected, SourceGuid)
+						),
+						NAME_None,
+						EUserInterfaceActionType::RadioButton
+					);
+				}
+			}
+			else
+			{
+				MenuBuilder.AddWidget(SNullWidget::NullWidget, LOCTEXT("InvalidLiveLink", "Invalid LiveLink Client"), false, false);
+			}
+			MenuBuilder.EndSection();
+
+			return MenuBuilder.MakeWidget();
+		}
+	}
+	
+	void HandleVirtualSourceSelection(FGuid InSourceGuid)
+	{
+		VirtualSourceGuid = InSourceGuid;
+	}
+
+	bool IsVirtualSourceSelected(FGuid InSourceGuid) const
+	{
+		return VirtualSourceGuid == InSourceGuid;
+	}
+
+	FText GetParentSourceText() const
+	{
+		if (LiveLinkClient != nullptr && VirtualSourceGuid.IsValid())
+		{
+			return LiveLinkClient->GetSourceType(VirtualSourceGuid);
+		}
+		else
+		{
+			return LOCTEXT("InvalidParentSource", "Invalid Source");
+		}
+	}
+
 	/** Creates the combo menu for the role class */
 	void MakeRoleClassPicker()
 	{
@@ -232,7 +333,8 @@ private:
 	{
 		if (LiveLinkClient)
 		{
-			LiveLinkClient->AddVirtualSubject(VirtualSubjectName, VirtualSubjectClass);
+			const FLiveLinkSubjectKey NewVirtualSubjectKey(VirtualSourceGuid, VirtualSubjectName);
+			LiveLinkClient->AddVirtualSubject(NewVirtualSubjectKey, VirtualSubjectClass);
 		}
 
 		CloseDialog(true);
@@ -273,18 +375,19 @@ private:
 		{
 			TArray<FLiveLinkSubjectKey> SubjectKey = LiveLinkClient->GetSubjects(true, true);
 			FName SubjectName = *NewSubjectName.ToString();
-			VirtualSubjectName = SubjectName;
+			const FLiveLinkSubjectKey ThisSubjectKey(VirtualSourceGuid, SubjectName);
 
 			if (SubjectName.IsNone())
 			{
 				VirtualSubjectTextWidgetPin->SetError(LOCTEXT("VirtualInvalidName", "Invalid Virtual Subject"));
 			}
-			else if (SubjectKey.FindByPredicate([SubjectName](const FLiveLinkSubjectKey& Key) { return Key.SubjectName == SubjectName; }))
+			else if (SubjectKey.FindByPredicate([ThisSubjectKey](const FLiveLinkSubjectKey& Key) { return Key == ThisSubjectKey; }))
 			{
 				VirtualSubjectTextWidgetPin->SetError(LOCTEXT("VirtualExistingName", "Subject already exist"));
 			}
 			else
 			{
+				VirtualSubjectName = SubjectName;
 				VirtualSubjectTextWidgetPin->SetError(FText::GetEmpty());
 			}
 		}
@@ -304,6 +407,9 @@ private:
 	/** The virtual subject's class */
 	TSubclassOf<ULiveLinkVirtualSubject> VirtualSubjectClass;
 
+	/** Selected source guid */
+	FGuid VirtualSourceGuid;
+
 	/** The virtual subject's name */
 	FName VirtualSubjectName;
 
@@ -342,7 +448,7 @@ void SLiveLinkClientPanelToolbar::Construct(const FArguments& Args, FLiveLinkCli
 			[
 				// The green button containing the "+ Add Source" items
 				SAssignNew(AddSourceButton, SComboButton)
-				.ToolTipText(LOCTEXT("AddSource_ToolTip", "Add a new live link source"))
+				.ToolTipText(LOCTEXT("AddSource_ToolTip", "Add a new LiveLink source"))
 				.ButtonStyle(FEditorStyle::Get(), "FlatButton.Success")
 				.ForegroundColor(FLinearColor::White)
 				.ContentPadding(FMargin(4, 0))
@@ -478,7 +584,7 @@ TSharedRef<SWidget> SLiveLinkClientPanelToolbar::OnGenerateSourceMenu()
 	const bool CloseAfterSelection = true;
 	FMenuBuilder MenuBuilder(CloseAfterSelection, NULL);
 
-	MenuBuilder.BeginSection("SourceSection", LOCTEXT("Sources", "Live Link Sources"));
+	MenuBuilder.BeginSection("SourceSection", LOCTEXT("Sources", "LiveLink Sources"));
 
 	for (int32 FactoryIndex = 0; FactoryIndex < Factories.Num(); ++FactoryIndex)
 	{
@@ -525,11 +631,19 @@ TSharedRef<SWidget> SLiveLinkClientPanelToolbar::OnGenerateSourceMenu()
 
 	MenuBuilder.EndSection();
 
-	MenuBuilder.BeginSection("VirtualSourceSection", LOCTEXT("VirtualSources", "Live Link Virtual Sources"));
+	MenuBuilder.BeginSection("VirtualSourceSection", LOCTEXT("VirtualSources", "LiveLink VirtualSubject Sources"));
+
+	//For now, it's not possible to create VirtualSubject Sources from the UI.
+	//Code is present in case it's required in the future.
+	//MenuBuilder.AddSubMenu(
+	//	LOCTEXT("AddVirtualSubjectSource", "Add VirtualSubject Source"),
+	//	LOCTEXT("AddVirtualSubjectSourceSubMenu_Tooltip", "Adds a new VirtualSubject Source to LiveLink. New Virtual Subjects can be created under a VirtualSubject Source."),
+	//	FNewMenuDelegate::CreateRaw(this, &SLiveLinkClientPanelToolbar::PopulateVirtualSubjectSourceCreationMenu)
+	//);
 
 	MenuBuilder.AddMenuEntry(
 		LOCTEXT("AddVirtualSubject", "Add Virtual Subject"),
-		LOCTEXT("AddVirtualSubject_Tooltip", "Adds a new virtual subject to live link. Instead of coming from a source a virtual subject is a combination of 2 or more real subjects"),
+		LOCTEXT("AddVirtualSubject_Tooltip", "Adds a new virtual subject to LiveLink. Instead of coming from a source a virtual subject is a combination of 2 or more real subjects"),
 		FSlateIcon(),
 		FUIAction(
 			FExecuteAction::CreateSP(this, &SLiveLinkClientPanelToolbar::AddVirtualSubject)
@@ -598,10 +712,69 @@ void SLiveLinkClientPanelToolbar::OnSourceCreated(TSharedPtr<ILiveLinkSource> Ne
 	FSlateApplication::Get().DismissAllMenus();
 }
 
+void SLiveLinkClientPanelToolbar::PopulateVirtualSubjectSourceCreationMenu(FMenuBuilder& InMenuBuilder)
+{
+	TSharedRef<SWidget> CreateSourceWidget =
+		SNew(SBox)
+		.WidthOverride(250)
+		[
+			SNew(SVerticalBox)
+			+ SVerticalBox::Slot()
+			.AutoHeight()
+			[
+				SNew(SHorizontalBox)
+				+ SHorizontalBox::Slot()
+				.HAlign(HAlign_Left)
+				.FillWidth(0.5f)
+				[
+					SNew(STextBlock)
+					.Text(LOCTEXT("VirtualSourceName", "Source Name"))
+				]
+				+ SHorizontalBox::Slot()
+				.HAlign(HAlign_Fill)
+				.FillWidth(0.5f)
+				[
+					SAssignNew(VirtualSubjectSourceName, SEditableTextBox)
+					.Text(LOCTEXT("VirtualSourceNameEditable", "VirtualSubjectSource"))
+				]
+			]
+			+ SVerticalBox::Slot()
+			.HAlign(HAlign_Right)
+			.AutoHeight()
+			[
+				SNew(SButton)
+				.OnClicked(this, &SLiveLinkClientPanelToolbar::OnAddVirtualSubjectSource)
+				[
+					SNew(STextBlock)
+					.Text(LOCTEXT("Ok", "Ok"))
+				]
+			]
+		];
+
+	InMenuBuilder.AddWidget(CreateSourceWidget, FText());
+}
+
+FReply SLiveLinkClientPanelToolbar::OnAddVirtualSubjectSource()
+{
+	if (TSharedPtr<SEditableTextBox> VirtualSourceNamePtr = VirtualSubjectSourceName.Pin())
+	{
+		if (Client)
+		{
+			Client->AddVirtualSubjectSource(*VirtualSourceNamePtr->GetText().ToString());
+		}
+	}
+	FSlateApplication::Get().DismissAllMenus();
+
+	return FReply::Handled();
+}
+
 void SLiveLinkClientPanelToolbar::AddVirtualSubject()
 {
-	TSharedRef<SVirtualSubjectCreateDialog> Dialog = SNew(SVirtualSubjectCreateDialog);
-	Dialog->ConfigureVirtualSubject(Client);
+	TSharedRef<SVirtualSubjectCreateDialog> Dialog = 
+		SNew(SVirtualSubjectCreateDialog)
+		.LiveLinkClient(Client);
+
+	Dialog->ConfigureVirtualSubject();
 }
 
 TSharedRef<SWidget> SLiveLinkClientPanelToolbar::OnPresetGeneratePresetsMenu()
diff --git a/Engine/Plugins/Animation/LiveLink/Source/LiveLinkEditor/Private/LiveLinkClientPanelToolbar.h b/Engine/Plugins/Animation/LiveLink/Source/LiveLinkEditor/Private/LiveLinkClientPanelToolbar.h
--- a/Engine/Plugins/Animation/LiveLink/Source/LiveLinkEditor/Private/LiveLinkClientPanelToolbar.h
+++ b/Engine/Plugins/Animation/LiveLink/Source/LiveLinkEditor/Private/LiveLinkClientPanelToolbar.h
@@ -40,19 +40,26 @@ private:
 	void OnSourceCreated(TSharedPtr<ILiveLinkSource> NewSource, FString ConnectionString, TSubclassOf<ULiveLinkSourceFactory> Factory);
 
 	TSharedRef<SWidget> OnPresetGeneratePresetsMenu();
+	TSharedRef<SWidget> HandleSourceSelectionComboButton();
+	void HandleVirtualSourceSelection(const FGuid& InSourceGuid);
+	bool IsVirtualSourceSelected(const FGuid& InSourceGuid);
 	void OnSaveAsPreset();
 	void OnImportPreset(const FAssetData& InPreset);
 	FReply OnRevertChanges();
 	bool HasLoadedLiveLinkPreset() const;
 
+	void PopulateVirtualSubjectSourceCreationMenu(FMenuBuilder& InMenuBuilder);
+	FReply OnAddVirtualSubjectSource();
 	void AddVirtualSubject();
 
+
 private:
 	FLiveLinkClient* Client;
 
 	TWeakPtr<IMenu> AddSubjectMenu;
 	TWeakPtr<IMenu> VirtualSubjectMenu;
 	TWeakPtr<STextEntryPopup> VirtualSubjectPopup;
+	TWeakPtr<SEditableTextBox> VirtualSubjectSourceName;
 	TSharedPtr<SComboButton> AddSourceButton;
 
 	TWeakObjectPtr<ULiveLinkPreset> LiveLinkPreset;
diff --git a/Engine/Plugins/Animation/LiveLink/Source/LiveLinkEditor/Private/LiveLinkVirtualSubjectDetailCustomization.cpp b/Engine/Plugins/Animation/LiveLink/Source/LiveLinkEditor/Private/LiveLinkVirtualSubjectDetailCustomization.cpp
--- a/Engine/Plugins/Animation/LiveLink/Source/LiveLinkEditor/Private/LiveLinkVirtualSubjectDetailCustomization.cpp
+++ b/Engine/Plugins/Animation/LiveLink/Source/LiveLinkEditor/Private/LiveLinkVirtualSubjectDetailCustomization.cpp
@@ -55,6 +55,15 @@ void FLiveLinkVirtualSubjectDetailCustomization::CustomizeDetails(IDetailLayoutB
 		SubjectsListItems.AddUnique(MakeShared<FName>(SubjectKey.SubjectName.Name));
 	}
 
+	// In case one of the associated subject linked to this virtual one doesn't exist anymore, add it to the list to display it red
+	for (const FLiveLinkSubjectName& SelectedSubject : Subject->GetSubjects())
+	{
+		if (false == SubjectsListItems.ContainsByPredicate([SelectedSubject](const FSubjectEntryPtr& Other) { return *Other == SelectedSubject.Name; }))
+		{
+			SubjectsListItems.AddUnique(MakeShared<FName>(SelectedSubject.Name));
+		}
+	}
+
 	IDetailCategoryBuilder& SubjectPropertyGroup = DetailBuilder.EditCategory(*SubjectsPropertyHandle->GetMetaData(TEXT("Category")));
 	SubjectPropertyGroup.AddCustomRow(LOCTEXT("SubjectsTitleLabel", "Subjects"))
 	.NameContent()
@@ -141,7 +150,44 @@ TSharedRef<ITableRow> FLiveLinkVirtualSubjectDetailCustomization::OnGenerateWidg
 			[
 				SNew(STextBlock)
 				.Text(FText::FromName(*InItem))
+				.ColorAndOpacity(this, &FLiveLinkVirtualSubjectDetailCustomization::HandleSubjectItemColor, InItem)
+				.ToolTipText(this, &FLiveLinkVirtualSubjectDetailCustomization::HandleSubjectItemToolTip, InItem)
 			]
 		];
 }
+
+FSlateColor FLiveLinkVirtualSubjectDetailCustomization::HandleSubjectItemColor(FSubjectEntryPtr InItem) const
+{
+	FSlateColor Result = FLinearColor::White;
+
+	if (ULiveLinkVirtualSubject* Subject = SubjectPtr.Get())
+	{
+		const FName ThisItem = *InItem.Get();
+		TArray<FLiveLinkSubjectKey> SubjectKeys = Client->GetSubjectsSupportingRole(Subject->GetRole(), false, false);
+		if (false == SubjectKeys.ContainsByPredicate([ThisItem](const FLiveLinkSubjectKey& Other) { return Other.SubjectName.Name == ThisItem; }))
+		{
+			Result = FLinearColor::Red;
+		}
+	}
+
+	return Result;
+}
+
+FText FLiveLinkVirtualSubjectDetailCustomization::HandleSubjectItemToolTip(FSubjectEntryPtr InItem) const
+{
+	FText Result;
+
+	if (ULiveLinkVirtualSubject* Subject = SubjectPtr.Get())
+	{
+		const FName ThisItem = *InItem.Get();
+		TArray<FLiveLinkSubjectKey> SubjectKeys = Client->GetSubjectsSupportingRole(Subject->GetRole(), false, false);
+		if (false == SubjectKeys.ContainsByPredicate([ThisItem](const FLiveLinkSubjectKey& Other) { return Other.SubjectName.Name == ThisItem; }))
+		{
+			Result = LOCTEXT("LinkedSubjectToolTip", "This subject was not found in the list of available LiveLink subjects. VirtualSubject might not work properly.");
+		}
+	}
+
+	return Result;
+}
+
 #undef LOCTEXT_NAMESPACE
diff --git a/Engine/Plugins/Animation/LiveLink/Source/LiveLinkEditor/Private/LiveLinkVirtualSubjectDetailCustomization.h b/Engine/Plugins/Animation/LiveLink/Source/LiveLinkEditor/Private/LiveLinkVirtualSubjectDetailCustomization.h
--- a/Engine/Plugins/Animation/LiveLink/Source/LiveLinkEditor/Private/LiveLinkVirtualSubjectDetailCustomization.h
+++ b/Engine/Plugins/Animation/LiveLink/Source/LiveLinkEditor/Private/LiveLinkVirtualSubjectDetailCustomization.h
@@ -8,6 +8,7 @@
 
 #include "LiveLinkTypes.h"
 #include "LiveLinkVirtualSubject.h"
+#include "Styling/SlateColor.h"
 #include "Types/SlateEnums.h"
 #include "Widgets/Views/SListView.h"
 
@@ -42,6 +43,11 @@ private:
 	// Creates subject tree entry widget
 	TSharedRef<ITableRow> OnGenerateWidgetForSubjectItem(FSubjectEntryPtr InItem, const TSharedRef<STableViewBase>& OwnerTable);
 
+	// If Item doesn't exist in Client subject's list, mark it as red
+	FSlateColor HandleSubjectItemColor(FSubjectEntryPtr InItem) const;
+
+	FText HandleSubjectItemToolTip(FSubjectEntryPtr InItem) const;
+
 	// The tile set being edited
 	TWeakObjectPtr<ULiveLinkVirtualSubject> SubjectPtr;
 
diff --git a/Engine/Source/Runtime/LiveLinkInterface/Private/LiveLinkVirtualSubject.cpp b/Engine/Source/Runtime/LiveLinkInterface/Private/LiveLinkVirtualSubject.cpp
--- a/Engine/Source/Runtime/LiveLinkInterface/Private/LiveLinkVirtualSubject.cpp
+++ b/Engine/Source/Runtime/LiveLinkInterface/Private/LiveLinkVirtualSubject.cpp
@@ -22,19 +22,7 @@ void ULiveLinkVirtualSubject::Update()
 	FrameSnapshot.StaticData.Reset();
 	FrameSnapshot.FrameData.Reset();
 
-	// Create the new translator for this frame
-	CurrentFrameTranslators.Reset();
-	for (ULiveLinkFrameTranslator* Translator : FrameTranslators)
-	{
-		if (Translator)
-		{
-			ULiveLinkFrameTranslator::FWorkerSharedPtr NewTranslator = Translator->FetchWorker();
-			if (NewTranslator.IsValid())
-			{
-				CurrentFrameTranslators.Add(NewTranslator);
-			}
-		}
-	}
+	UpdateTranslatorsForThisFrame();
 }
 
 
@@ -65,3 +53,20 @@ bool ULiveLinkVirtualSubject::DependsOnSubject(FName SubjectName) const
 {
 	return Subjects.Contains(SubjectName);
 }
+
+void ULiveLinkVirtualSubject::UpdateTranslatorsForThisFrame()
+{
+	// Create the new translator for this frame
+	CurrentFrameTranslators.Reset();
+	for (ULiveLinkFrameTranslator* Translator : FrameTranslators)
+	{
+		if (Translator)
+		{
+			ULiveLinkFrameTranslator::FWorkerSharedPtr NewTranslator = Translator->FetchWorker();
+			if (NewTranslator.IsValid())
+			{
+				CurrentFrameTranslators.Add(NewTranslator);
+			}
+		}
+	}
+}
diff --git a/Engine/Source/Runtime/LiveLinkInterface/Public/ILiveLinkClient.h b/Engine/Source/Runtime/LiveLinkInterface/Public/ILiveLinkClient.h
--- a/Engine/Source/Runtime/LiveLinkInterface/Public/ILiveLinkClient.h
+++ b/Engine/Source/Runtime/LiveLinkInterface/Public/ILiveLinkClient.h
@@ -116,6 +116,9 @@ public:
 
 	/** Add a new live link source to the client */
 	virtual FGuid AddSource(TSharedPtr<ILiveLinkSource> Source) = 0;
+	
+	/** Add a new live link VirtualSubject source to the client */
+	virtual FGuid AddVirtualSubjectSource(FName SourceName) = 0;
 
 	/** Create from the factory a new live link source and add it to the client. The settings will be duplicated. */
 	virtual bool CreateSource(const FLiveLinkSourcePreset& SourcePreset) = 0;
@@ -132,6 +135,9 @@ public:
 	/** Get a list of all the sources */
 	virtual TArray<FGuid> GetSources() const = 0;
 
+	/** Get a list of all the VirtualSubjects sources */
+	virtual TArray<FGuid> GetVirtualSources() const = 0;
+
 	/** Get the source preset from the live link client. The settings will be duplicated into DuplicatedObjectOuter. */
 	virtual FLiveLinkSourcePreset GetSourcePreset(FGuid SourceGuid, UObject* DuplicatedObjectOuter) const = 0;
 
@@ -147,6 +153,12 @@ public:
 	/** Create and add a new live link subject to the client */
 	virtual bool CreateSubject(const FLiveLinkSubjectPreset& SubjectPreset) = 0;
 
+	/** Add a new virtual subject to the client */
+	virtual bool AddVirtualSubject(FLiveLinkSubjectKey VirtualSubjectKey, TSubclassOf<ULiveLinkVirtualSubject> VirtualSubjectClass) = 0;
+
+	/** Removes a virtual subject from the client */
+	virtual void RemoveVirtualSubject(FLiveLinkSubjectKey VirtualSubjectKey) = 0;
+
 	/** Clear the subject from the specific source */
 	virtual void RemoveSubject_AnyThread(const FLiveLinkSubjectKey& SubjectName) = 0;
 
@@ -226,6 +238,12 @@ public:
 	 */
 	virtual TArray<FLiveLinkTime> GetSubjectFrameTimes(FLiveLinkSubjectName SubjectName) const = 0;
 
+	/**
+	 * Get the Settings of this subject.
+	 * @note If subject is a VirtualSubject, the VirtualSubject itself is returned.
+	 */
+	virtual UObject* GetSubjectSettings(const FLiveLinkSubjectKey& SubjectKey) const = 0;
+
 	/**
 	 * Return the evaluated subject from a specific source snapshot for a specific role.
 	 * A subject could have to go through a translator to output in the desired role.
@@ -263,6 +281,9 @@ public:
 	 */
 	virtual bool EvaluateFrameAtSceneTime_AnyThread(FLiveLinkSubjectName SubjectName, const FTimecode& SceneTime, TSubclassOf<ULiveLinkRole> DesiredRole, FLiveLinkSubjectFrameData& OutFrame) = 0;
 
+	/** Notify when the LiveLinkClient has ticked. */
+	virtual FSimpleMulticastDelegate& OnLiveLinkTicked() = 0;
+
 	/** Notify when the list of sources has changed. */
 	virtual FSimpleMulticastDelegate& OnLiveLinkSourcesChanged() = 0;
 
diff --git a/Engine/Source/Runtime/LiveLinkInterface/Public/LiveLinkTypes.h b/Engine/Source/Runtime/LiveLinkInterface/Public/LiveLinkTypes.h
--- a/Engine/Source/Runtime/LiveLinkInterface/Public/LiveLinkTypes.h
+++ b/Engine/Source/Runtime/LiveLinkInterface/Public/LiveLinkTypes.h
@@ -356,6 +356,11 @@ public:
 		return false;
 	}
 
+	bool operator!=(const FLiveLinkBaseDataStruct& Other) const
+	{
+		return !(*this == Other);
+	}
+
 	friend FArchive& operator<<(FArchive& Ar, FLiveLinkBaseDataStruct& InStruct)
 	{
 		if (Ar.IsLoading())
diff --git a/Engine/Source/Runtime/LiveLinkInterface/Public/LiveLinkVirtualSubject.h b/Engine/Source/Runtime/LiveLinkInterface/Public/LiveLinkVirtualSubject.h
--- a/Engine/Source/Runtime/LiveLinkInterface/Public/LiveLinkVirtualSubject.h
+++ b/Engine/Source/Runtime/LiveLinkInterface/Public/LiveLinkVirtualSubject.h
@@ -40,8 +40,12 @@ protected:
 public:
 	ILiveLinkClient* GetClient() const { return LiveLinkClient; }
 	const TArray<FLiveLinkSubjectName>& GetSubjects() const { return Subjects; }
+	const TArray<ULiveLinkFrameTranslator*>& GetTranslators() const { return FrameTranslators; }
 	virtual bool DependsOnSubject(FName SubjectName) const;
 
+protected:
+	void UpdateTranslatorsForThisFrame();
+
 protected:
 	/** The role the subject was build with. */
 	UPROPERTY()
