- Fixing LiveLink CIS

diff --git a/Engine/Plugins/Animation/LiveLink/Source/LiveLinkEditor/Private/LiveLinkClientPanelToolbar.h b/Engine/Plugins/Animation/LiveLink/Source/LiveLinkEditor/Private/LiveLinkClientPanelToolbar.h
--- a/Engine/Plugins/Animation/LiveLink/Source/LiveLinkEditor/Private/LiveLinkClientPanelToolbar.h
+++ b/Engine/Plugins/Animation/LiveLink/Source/LiveLinkEditor/Private/LiveLinkClientPanelToolbar.h
@@ -17,6 +17,7 @@ class FLiveLinkClient;
 class ILiveLinkSource;
 class IMenu;
 class SComboButton;
+class SEditableTextBox;
 class STextEntryPopup;
 class ULiveLinkSourceFactory;
 struct FAssetData;
