Fixing build break introduced with copy up

diff --git a/Engine/Plugins/Animation/LiveLink/Source/LiveLinkComponents/Private/Controllers/LiveLinkCameraController.cpp b/Engine/Plugins/Animation/LiveLink/Source/LiveLinkComponents/Private/Controllers/LiveLinkCameraController.cpp
--- a/Engine/Plugins/Animation/LiveLink/Source/LiveLinkComponents/Private/Controllers/LiveLinkCameraController.cpp
+++ b/Engine/Plugins/Animation/LiveLink/Source/LiveLinkComponents/Private/Controllers/LiveLinkCameraController.cpp
@@ -2,76 +2,96 @@
 
 #include "Controllers/LiveLinkCameraController.h"
 
-#include "ILiveLinkClient.h"
-#include "Roles/LiveLinkCameraRole.h"
-#include "Roles/LiveLinkCameraTypes.h"
 
 #include "Camera/CameraComponent.h"
 #include "CineCameraComponent.h"
+#include "Controllers/LiveLinkTransformController.h"
 #include "Features/IModularFeatures.h"
 #include "GameFramework/Actor.h"
+#include "ILiveLinkClient.h"
+#include "LiveLinkComponentController.h"
+#include "Roles/LiveLinkCameraRole.h"
+#include "Roles/LiveLinkCameraTypes.h"
+#include "UObject/EnterpriseObjectVersion.h"
 
 #if WITH_EDITOR
 #include "Kismet2/ComponentEditorUtils.h"
 #endif
 
 
-void ULiveLinkCameraController::OnEvaluateRegistered()
+void ULiveLinkCameraController::Tick(float DeltaTime, const FLiveLinkSubjectFrameData& SubjectData)
 {
-	AActor* OuterActor = GetOuterActor();
-	TransformData.CheckForError(OuterActor ? OuterActor->GetFName() : NAME_None, Cast<USceneComponent>(ComponentToControl.GetComponent(OuterActor)));
-}
-
+	const FLiveLinkCameraStaticData* StaticData = SubjectData.StaticData.Cast<FLiveLinkCameraStaticData>();
+	const FLiveLinkCameraFrameData* FrameData = SubjectData.FrameData.Cast<FLiveLinkCameraFrameData>();
 
-void ULiveLinkCameraController::Tick(float DeltaTime, const FLiveLinkSubjectRepresentation& SubjectRepresentation)
-{
-	if (UCameraComponent* CameraComponent = Cast<UCameraComponent>(ComponentToControl.GetComponent(GetOuterActor())))
+	if (StaticData && FrameData)
 	{
-		ILiveLinkClient& LiveLinkClient = IModularFeatures::Get().GetModularFeature<ILiveLinkClient>(ILiveLinkClient::ModularFeatureName);
-
-		FLiveLinkSubjectFrameData SubjectData;
-		if (LiveLinkClient.EvaluateFrame_AnyThread(SubjectRepresentation.Subject, SubjectRepresentation.Role, SubjectData))
+		if (UCameraComponent* CameraComponent = Cast<UCameraComponent>(AttachedComponent))
 		{
-			FLiveLinkCameraStaticData* StaticData = SubjectData.StaticData.Cast<FLiveLinkCameraStaticData>();
-			FLiveLinkCameraFrameData* FrameData = SubjectData.FrameData.Cast<FLiveLinkCameraFrameData>();
+			if (StaticData->bIsFieldOfViewSupported) { CameraComponent->SetFieldOfView(FrameData->FieldOfView); }
+			if (StaticData->bIsAspectRatioSupported) { CameraComponent->SetAspectRatio(FrameData->AspectRatio); }
+			if (StaticData->bIsProjectionModeSupported) { CameraComponent->SetProjectionMode(FrameData->ProjectionMode == ELiveLinkCameraProjectionMode::Perspective ? ECameraProjectionMode::Perspective : ECameraProjectionMode::Orthographic); }
 
-			if (StaticData && FrameData)
+			if (UCineCameraComponent* CineCameraComponent = Cast<UCineCameraComponent>(CameraComponent))
 			{
-				TransformData.ApplyTransform(CameraComponent, FrameData->Transform);
-
-				if (StaticData->bIsFieldOfViewSupported) { CameraComponent->SetFieldOfView(FrameData->FieldOfView); } 
-				if (StaticData->bIsAspectRatioSupported) { CameraComponent->SetAspectRatio(FrameData->AspectRatio); }
-				if (StaticData->bIsProjectionModeSupported) { CameraComponent->SetProjectionMode(FrameData->ProjectionMode == ELiveLinkCameraProjectionMode::Perspective ? ECameraProjectionMode::Perspective : ECameraProjectionMode::Orthographic); }
-
-				if (UCineCameraComponent* CineCameraComponent = Cast<UCineCameraComponent>(CameraComponent))
-				{
-					if (StaticData->bIsFocalLengthSupported) { CineCameraComponent->CurrentFocalLength = FrameData->FocalLength; }
-					if (StaticData->bIsApertureSupported) { CineCameraComponent->CurrentAperture = FrameData->Aperture; }
-					if (StaticData->FilmBackWidth > 0.0f) { CineCameraComponent->Filmback.SensorWidth = StaticData->FilmBackWidth; }
-					if (StaticData->FilmBackHeight > 0.0f) { CineCameraComponent->Filmback.SensorHeight = StaticData->FilmBackHeight; }
-					if (StaticData->bIsFocusDistanceSupported) { CineCameraComponent->FocusSettings.ManualFocusDistance = FrameData->FocusDistance; }
-				}
+				if (StaticData->bIsFocalLengthSupported) { CineCameraComponent->CurrentFocalLength = FrameData->FocalLength; }
+				if (StaticData->bIsApertureSupported) { CineCameraComponent->CurrentAperture = FrameData->Aperture; }
+				if (StaticData->FilmBackWidth > 0.0f) { CineCameraComponent->Filmback.SensorWidth = StaticData->FilmBackWidth; }
+				if (StaticData->FilmBackHeight > 0.0f) { CineCameraComponent->Filmback.SensorHeight = StaticData->FilmBackHeight; }
+				if (StaticData->bIsFocusDistanceSupported) { CineCameraComponent->FocusSettings.ManualFocusDistance = FrameData->FocusDistance; }
 			}
 		}
 	}
 }
 
-
 bool ULiveLinkCameraController::IsRoleSupported(const TSubclassOf<ULiveLinkRole>& RoleToSupport)
 {
-	return RoleToSupport->IsChildOf(ULiveLinkCameraRole::StaticClass());
+	return RoleToSupport == ULiveLinkCameraRole::StaticClass();
 }
 
+TSubclassOf<UActorComponent> ULiveLinkCameraController::GetDesiredComponentClass() const
+{
+	return UCameraComponent::StaticClass();
+}
 
-#if WITH_EDITOR
-void ULiveLinkCameraController::InitializeInEditor()
+void ULiveLinkCameraController::PostLoad()
 {
-	if (AActor* Actor = GetOuterActor())
+	Super::PostLoad();
+
+#if WITH_EDITOR
+	const int32 Version = GetLinkerCustomVersion(FEnterpriseObjectVersion::GUID);
+	if (Version < FEnterpriseObjectVersion::LiveLinkControllerSplitPerRole)
 	{
-		if (UCameraComponent* CameraComponent = Actor->FindComponentByClass<UCameraComponent>())
+		AActor* MyActor = GetOuterActor();
+		if (MyActor)
 		{
-			ComponentToControl = FComponentEditorUtils::MakeComponentReference(Actor, CameraComponent);
+			//Make sure all UObjects we use in our post load have been postloaded
+			MyActor->ConditionalPostLoad();
+
+			ULiveLinkComponentController* LiveLinkComponent = Cast<ULiveLinkComponentController>(MyActor->GetComponentByClass(ULiveLinkComponentController::StaticClass()));
+			if (LiveLinkComponent)
+			{
+				LiveLinkComponent->ConditionalPostLoad();
+
+				//If the transform controller that was created to drive the TransformRole is the built in one, set its data structure with the one that we had internally
+				if (LiveLinkComponent->ControllerMap.Contains(ULiveLinkTransformRole::StaticClass()))
+				{
+					ULiveLinkTransformController* TransformController = Cast<ULiveLinkTransformController>(LiveLinkComponent->ControllerMap[ULiveLinkTransformRole::StaticClass()]);
+					if (TransformController)
+					{
+						TransformController->ConditionalPostLoad();
+						TransformController->TransformData = TransformData_DEPRECATED;
+					}
+				}
+
+				//if Subjects role direct controller is us, set the component to control to what we had
+				if (LiveLinkComponent->SubjectRepresentation.Role == ULiveLinkCameraRole::StaticClass())
+				{
+					LiveLinkComponent->ComponentToControl = ComponentToControl_DEPRECATED;
+				}
+			}
 		}
 	}
+#endif //WITH_EDITOR
 }
-#endif
+
