Load a LiveLink Preset from a command.

diff --git a/Engine/Plugins/Animation/LiveLink/Source/LiveLink/Private/LiveLinkModule.cpp b/Engine/Plugins/Animation/LiveLink/Source/LiveLink/Private/LiveLinkModule.cpp
--- a/Engine/Plugins/Animation/LiveLink/Source/LiveLink/Private/LiveLinkModule.cpp
+++ b/Engine/Plugins/Animation/LiveLink/Source/LiveLink/Private/LiveLinkModule.cpp
@@ -11,6 +11,8 @@
 #include "LiveLinkLogInstance.h"
 #include "LiveLinkDebugCommand.h"
 #include "LiveLinkHeartbeatEmitter.h"
+#include "LiveLinkPreset.h"
+#include "LiveLinkSettings.h"
 
 /**
  * Implements the Messaging module.
@@ -41,6 +43,11 @@ public:
 		FLiveLinkLogInstance::CreateInstance();
 		IModularFeatures::Get().RegisterModularFeature(FLiveLinkClient::ModularFeatureName, &LiveLinkClient);
 		LiveLinkMotionController.RegisterController();
+
+		if (ULiveLinkPreset* Preset = GetDefault<ULiveLinkSettings>()->DefaultLiveLinkPreset.LoadSynchronous())
+		{
+			Preset->ApplyToClient();
+		}
 	}
 
 	virtual void ShutdownModule() override
diff --git a/Engine/Plugins/Animation/LiveLink/Source/LiveLink/Private/LiveLinkPreset.cpp b/Engine/Plugins/Animation/LiveLink/Source/LiveLink/Private/LiveLinkPreset.cpp
--- a/Engine/Plugins/Animation/LiveLink/Source/LiveLink/Private/LiveLinkPreset.cpp
+++ b/Engine/Plugins/Animation/LiveLink/Source/LiveLink/Private/LiveLinkPreset.cpp
@@ -2,7 +2,34 @@
 
 #include "LiveLinkPreset.h"
 #include "Features/IModularFeatures.h"
+#include "HAL/IConsoleManager.h"
 #include "LiveLinkClient.h"
+#include "LiveLinkLog.h"
+#include "UObject/ConstructorHelpers.h"
+
+namespace
+{
+	static FAutoConsoleCommand LiveLinkPresetApplyCmd(
+		TEXT("LiveLink.Preset.Apply"),
+		TEXT("Apply a LiveLinkPreset. Use: LiveLink.Preset.Apply Preset=/Game/Folder/MyLiveLinkPreset.MyLiveLinkPreset"),
+		FConsoleCommandWithArgsDelegate::CreateLambda([](const TArray<FString>& Args)
+			{
+				for (const FString& Element : Args)
+				{
+					const TCHAR* PresetStr = TEXT("Preset=");
+					int32 FoundElement = Element.Find(PresetStr);
+					if (FoundElement != INDEX_NONE)
+					{
+						UObject* Object = StaticLoadObject(ULiveLinkPreset::StaticClass(), nullptr, *Element + FoundElement + FCString::Strlen(PresetStr));
+						if (Object)
+						{
+							CastChecked<ULiveLinkPreset>(Object)->ApplyToClient();
+						}
+					}
+				}
+			})
+		);
+}
 
 bool ULiveLinkPreset::ApplyToClient() const
 {
@@ -26,6 +53,15 @@ bool ULiveLinkPreset::ApplyToClient() const
 		}
 	}
 
+	if (bResult)
+	{
+		FLiveLinkLog::Info(TEXT("Applied '%s'"), *GetFullName());
+	}
+	else
+	{
+		FLiveLinkLog::Error(TEXT("Could not apply '%s'"), *GetFullName());
+	}
+
 	return bResult;
 }
 
diff --git a/Engine/Plugins/Animation/LiveLink/Source/LiveLink/Public/LiveLinkSettings.h b/Engine/Plugins/Animation/LiveLink/Source/LiveLink/Public/LiveLinkSettings.h
--- a/Engine/Plugins/Animation/LiveLink/Source/LiveLink/Public/LiveLinkSettings.h
+++ b/Engine/Plugins/Animation/LiveLink/Source/LiveLink/Public/LiveLinkSettings.h
@@ -16,6 +16,8 @@
 #include "LiveLinkSettings.generated.h"
 
 
+class ULiveLinkPreset;
+
 /**
  * Settings for LiveLinkRole.
  */
@@ -65,7 +67,11 @@ public:
 	/** The interpolation class to use for new Subjects if no specific settings we set for the Subject's role. */
 	UPROPERTY(config)
 	TSubclassOf<ULiveLinkFrameInterpolationProcessor> FrameInterpolationProcessor;
-	
+
+	/** The default preset that should be applied */
+	UPROPERTY(config, EditAnywhere, Category = "LiveLink")
+	TSoftObjectPtr<ULiveLinkPreset> DefaultLiveLinkPreset;
+
 	/** The default location in which to save take presets */
 	UPROPERTY(config, EditAnywhere, Category="LiveLink", meta=(DisplayName="Preset Save Location"))
 	FDirectoryPath PresetSaveDir;
