Fix to set auto tangents on live link curves.

diff --git a/Engine/Plugins/Animation/LiveLink/Source/LiveLinkMovieScene/Private/MovieScene/MovieSceneLiveLinkPropertyHandler.cpp b/Engine/Plugins/Animation/LiveLink/Source/LiveLinkMovieScene/Private/MovieScene/MovieSceneLiveLinkPropertyHandler.cpp
--- a/Engine/Plugins/Animation/LiveLink/Source/LiveLinkMovieScene/Private/MovieScene/MovieSceneLiveLinkPropertyHandler.cpp
+++ b/Engine/Plugins/Animation/LiveLink/Source/LiveLinkMovieScene/Private/MovieScene/MovieSceneLiveLinkPropertyHandler.cpp
@@ -122,7 +122,7 @@ void FMovieSceneLiveLinkPropertyHandler<float>::Finalize(bool bInReduceKeys, con
 		const TArray<FLiveLinkPropertyKey<float>>& ElementKeys = Keys[i];
 		for (const FLiveLinkPropertyKey<float>& Key : ElementKeys)
 		{
-			PropertyStorage->FloatChannel[i].AddCubicKey(Key.Time, Key.Value, RCTM_Break);
+			PropertyStorage->FloatChannel[i].AddCubicKey(Key.Time, Key.Value, RCTM_Auto);
 		}
 	}
 
@@ -133,6 +133,13 @@ void FMovieSceneLiveLinkPropertyHandler<float>::Finalize(bool bInReduceKeys, con
 			MovieScene::Optimize(&Channel, InOptimizationParams);
 		}
 	}
+	else
+	{
+		for (FMovieSceneFloatChannel& Channel : PropertyStorage->FloatChannel)
+		{
+			Channel.AutoSetTangents();
+		}
+	}
 }
 
 template<>
