Use FQualifiedFrameTime instead of FTimecode in FApp.
Use a TOptional in FApp, if there is not TCProvider, the frame time will be invalidated. This will change soon in another CL where we will have a "default" TC if no provider is defined but the frame time can be invalid if the TC provider is not "synchronized".
Remove FTimecode for maths in LiveLink and nDisplay, only use timecode for display.
Use FTimecode::IsDropFormatTimecodeSupported for now, to convert from FrameTime to Timecode. That will be changed in another CL when we support NDF and DF properly.
Deprecated AudioCaptureTimecodeProvider.
Generate a default TC value when no TC provided is set. By default it's enabled and the framerate is 24 fps.
Add a CVar that set the engine in DropFrame or NonDropFrame when the TC is 29.97 or 59.94.
Remove the option in sequencer to select between the 2. It's a CVar and not a project settings because other softwares that link with Core need that information.

diff --git a/Engine/Plugins/Animation/LiveLink/Source/LiveLink/Private/LiveLinkBlueprintLibrary.cpp b/Engine/Plugins/Animation/LiveLink/Source/LiveLink/Private/LiveLinkBlueprintLibrary.cpp
--- a/Engine/Plugins/Animation/LiveLink/Source/LiveLink/Private/LiveLinkBlueprintLibrary.cpp
+++ b/Engine/Plugins/Animation/LiveLink/Source/LiveLink/Private/LiveLinkBlueprintLibrary.cpp
@@ -271,7 +271,7 @@ namespace EvaluateLiveLinkFrame
 	};
 
 	bool Generic_Evaluate(const ULiveLinkBlueprintLibrary* InSelf, const struct FFrame& InStack, UStructProperty* InBlueprintDataStructProp
-		, FLiveLinkSubjectRepresentation InSubjectRepresentation, EEvaluateType InEvaluationType, double InWorldTime, FTimecode& InSceneTime, void* OutBlueprintDataPtr)
+		, FLiveLinkSubjectRepresentation InSubjectRepresentation, EEvaluateType InEvaluationType, double InWorldTime, const FQualifiedFrameTime& InQualifiedFrameTime, void* OutBlueprintDataPtr)
 	{
 		bool bSuccess = false;
 
@@ -308,7 +308,7 @@ namespace EvaluateLiveLinkFrame
 							bSuccess = LiveLinkClient.EvaluateFrameAtWorldTime_AnyThread(InSubjectRepresentation.Subject, InWorldTime, InSubjectRepresentation.Role, FrameData);
 							break;
 						case EEvaluateType::SceneTime:
-							bSuccess = LiveLinkClient.EvaluateFrameAtSceneTime_AnyThread(InSubjectRepresentation.Subject, InSceneTime, InSubjectRepresentation.Role, FrameData);
+							bSuccess = LiveLinkClient.EvaluateFrameAtSceneTime_AnyThread(InSubjectRepresentation.Subject, InQualifiedFrameTime, InSubjectRepresentation.Role, FrameData);
 							break;
 						case EEvaluateType::Snapshot:
 						default:
@@ -358,7 +358,7 @@ DEFINE_FUNCTION(ULiveLinkBlueprintLibrary::execEvaluateLiveLinkFrame)
 	P_FINISH;
 
 	double WorldTime = 0.0;
-	FTimecode SceneTime;
+	FQualifiedFrameTime SceneTime;
 	bool bSuccess = EvaluateLiveLinkFrame::Generic_Evaluate(P_THIS, Stack, BlueprintDataStructProp, SubjectRepresentation, EvaluateLiveLinkFrame::EEvaluateType::Snapshot, WorldTime, SceneTime, OutBlueprintDataPtr);
 	*(bool*)RESULT_PARAM = bSuccess;
 }
@@ -377,7 +377,7 @@ DEFINE_FUNCTION(ULiveLinkBlueprintLibrary::execEvaluateLiveLinkFrameWithSpecific
 
 
 	double WorldTime = 0.0;
-	FTimecode SceneTime;
+	FQualifiedFrameTime SceneTime;
 	bool bSuccess = EvaluateLiveLinkFrame::Generic_Evaluate(P_THIS, Stack, BlueprintDataStructProp, {SubjectName, RoleClass}, EvaluateLiveLinkFrame::EEvaluateType::Snapshot, WorldTime, SceneTime, OutBlueprintDataPtr);
 	*(bool*)RESULT_PARAM = bSuccess;
 }
@@ -398,7 +398,7 @@ DEFINE_FUNCTION(ULiveLinkBlueprintLibrary::execEvaluateLiveLinkFrameAtWorldTimeO
 	P_FINISH;
 
 	double WorldTime = FApp::GetCurrentTime() + WorldTimeOffset;
-	FTimecode SceneTime;
+	FQualifiedFrameTime SceneTime;
 	bool bSuccess = EvaluateLiveLinkFrame::Generic_Evaluate(P_THIS, Stack, BlueprintDataStructProp, {SubjectName, RoleClass}, EvaluateLiveLinkFrame::EEvaluateType::WorldTime, WorldTime, SceneTime, OutBlueprintDataPtr);
 	*(bool*)RESULT_PARAM = bSuccess;
 }
@@ -417,6 +417,6 @@ DEFINE_FUNCTION(ULiveLinkBlueprintLibrary::execEvaluateLiveLinkFrameAtSceneTime)
 	P_FINISH;
 
 	double WorldTime = 0.0;
-	bool bSuccess = EvaluateLiveLinkFrame::Generic_Evaluate(P_THIS, Stack, BlueprintDataStructProp, {SubjectName, RoleClass}, EvaluateLiveLinkFrame::EEvaluateType::SceneTime, WorldTime, SceneTime, OutBlueprintDataPtr);
+	bool bSuccess = EvaluateLiveLinkFrame::Generic_Evaluate(P_THIS, Stack, BlueprintDataStructProp, {SubjectName, RoleClass}, EvaluateLiveLinkFrame::EEvaluateType::SceneTime, WorldTime, FQualifiedFrameTime(SceneTime, FApp::GetTimecodeFrameRate()), OutBlueprintDataPtr);
 	*(bool*)RESULT_PARAM = bSuccess;
 }
\ No newline at end of file
diff --git a/Engine/Plugins/Animation/LiveLink/Source/LiveLink/Private/LiveLinkClient.cpp b/Engine/Plugins/Animation/LiveLink/Source/LiveLink/Private/LiveLinkClient.cpp
--- a/Engine/Plugins/Animation/LiveLink/Source/LiveLink/Private/LiveLinkClient.cpp
+++ b/Engine/Plugins/Animation/LiveLink/Source/LiveLink/Private/LiveLinkClient.cpp
@@ -175,8 +175,7 @@ void FLiveLinkClient::CacheValues()
 {
 #if WITH_EDITOR
 	CachedEngineTime = FApp::GetCurrentTime();
-	CachedEngineTimecode = FApp::GetTimecode();
-	CachedEngineTimecodeFrameRate = FApp::GetTimecodeFrameRate();
+	CachedEngineFrameTime = FApp::GetCurrentFrameTime();
 #endif
 }
 
@@ -1060,7 +1059,7 @@ bool FLiveLinkClient::EvaluateFrame_AnyThread(FLiveLinkSubjectName InSubjectName
 #if WITH_EDITOR
 		if (OnLiveLinkSubjectEvaluated().IsBound())
 		{
-			FLiveLinkTime RequestedTime = FLiveLinkTime(CachedEngineTime, FQualifiedFrameTime(CachedEngineTimecode, CachedEngineTimecodeFrameRate));
+			FLiveLinkTime RequestedTime = FLiveLinkTime(CachedEngineTime, CachedEngineFrameTime.Get(FQualifiedFrameTime()));
 			FLiveLinkTime ResultTime;
 			if (bResult)
 			{
@@ -1122,7 +1121,7 @@ bool FLiveLinkClient::EvaluateFrameAtWorldTime_AnyThread(FLiveLinkSubjectName In
 	return bResult;
 }
 
-bool FLiveLinkClient::EvaluateFrameAtSceneTime_AnyThread(FLiveLinkSubjectName InSubjectName, const FTimecode& InSceneTime, TSubclassOf<ULiveLinkRole> InDesiredRole, FLiveLinkSubjectFrameData& OutFrame)
+bool FLiveLinkClient::EvaluateFrameAtSceneTime_AnyThread(FLiveLinkSubjectName InSubjectName, const FQualifiedFrameTime& InSceneTime, TSubclassOf<ULiveLinkRole> InDesiredRole, FLiveLinkSubjectFrameData& OutFrame)
 {
 	SCOPE_CYCLE_COUNTER(STAT_LiveLink_EvaluateFrame);
 
@@ -1147,7 +1146,7 @@ bool FLiveLinkClient::EvaluateFrameAtSceneTime_AnyThread(FLiveLinkSubjectName In
 #if WITH_EDITOR
 			if (OnLiveLinkSubjectEvaluated().IsBound())
 			{
-				FLiveLinkTime RequestedTime = FLiveLinkTime(0.0, FQualifiedFrameTime(InSceneTime, CachedEngineTimecodeFrameRate));
+				FLiveLinkTime RequestedTime = FLiveLinkTime(0.0, InSceneTime);
 				FLiveLinkTime ResultTime;
 				if (bResult)
 				{
@@ -1546,6 +1545,11 @@ const FLiveLinkSubjectFrame* FLiveLinkClient_Base_DEPRECATED::GetSubjectDataAtSc
 	return nullptr;
 }
 
+bool FLiveLinkClient_Base_DEPRECATED::EvaluateFrameAtSceneTime_AnyThread(FLiveLinkSubjectName SubjectName, const FTimecode& SceneTime, TSubclassOf<ULiveLinkRole> DesiredRole, FLiveLinkSubjectFrameData& OutFrame)
+{
+	return static_cast<ILiveLinkClient*>(this)->EvaluateFrameAtSceneTime_AnyThread(SubjectName, FQualifiedFrameTime(SceneTime, FApp::GetTimecodeFrameRate()), DesiredRole, OutFrame);
+}
+
 const TArray<FLiveLinkFrame>* FLiveLinkClient_Base_DEPRECATED::GetSubjectRawFrames(FName InSubjectName)
 {
 	static const FName NAME_UpdateYourCode = "LiveLinkClient_GetSubjectRawFrames";
diff --git a/Engine/Plugins/Animation/LiveLink/Source/LiveLink/Private/LiveLinkComponent.cpp b/Engine/Plugins/Animation/LiveLink/Source/LiveLink/Private/LiveLinkComponent.cpp
--- a/Engine/Plugins/Animation/LiveLink/Source/LiveLink/Private/LiveLinkComponent.cpp
+++ b/Engine/Plugins/Animation/LiveLink/Source/LiveLink/Private/LiveLinkComponent.cpp
@@ -4,6 +4,7 @@
 #include "Components/SkeletalMeshComponent.h"
 #include "Features/IModularFeatures.h"
 #include "GameFramework/Actor.h"
+#include "Misc/App.h"
 #include "Roles/LiveLinkAnimationRole.h"
 
 // Sets default values for this component's properties
@@ -151,7 +152,7 @@ void ULiveLinkComponent::GetSubjectDataAtSceneTime(const FName SubjectName, cons
 	if (HasLiveLinkClient())
 	{
 		FLiveLinkSubjectFrameData FrameData;
-		if (LiveLinkClient->EvaluateFrameAtSceneTime_AnyThread(SubjectName, SceneTime, ULiveLinkAnimationRole::StaticClass(), FrameData))
+		if (LiveLinkClient->EvaluateFrameAtSceneTime_AnyThread(SubjectName, FQualifiedFrameTime(SceneTime, FApp::GetTimecodeFrameRate()), ULiveLinkAnimationRole::StaticClass(), FrameData))
 		{
 			const FLiveLinkSkeletonStaticData* SkeletonData = FrameData.StaticData.Cast<FLiveLinkSkeletonStaticData>();
 			const FLiveLinkAnimationFrameData* AnimationFrameData = FrameData.FrameData.Cast<FLiveLinkAnimationFrameData>();
diff --git a/Engine/Plugins/Animation/LiveLink/Source/LiveLink/Private/LiveLinkSubject.cpp b/Engine/Plugins/Animation/LiveLink/Source/LiveLink/Private/LiveLinkSubject.cpp
--- a/Engine/Plugins/Animation/LiveLink/Source/LiveLink/Private/LiveLinkSubject.cpp
+++ b/Engine/Plugins/Animation/LiveLink/Source/LiveLink/Private/LiveLinkSubject.cpp
@@ -62,28 +62,33 @@ void FLiveLinkSubject::Update()
 	}
 	else if (GetMode() == ELiveLinkSourceMode::Timecode)
 	{
-		const FQualifiedFrameTime CurrentSyncTime = FQualifiedFrameTime(FApp::GetTimecode(), FApp::GetTimecodeFrameRate());
-		const FFrameTime CurrentFrameTimeInFrameSpace = CurrentSyncTime.ConvertTo(CachedSettings.BufferSettings.TimecodeFrameRate);
-		int32 FrameIndex = 0;
-		for (const FLiveLinkFrameDataStruct& SourceFrameData : FrameData)
+		if (FApp::GetCurrentFrameTime().IsSet())
 		{
-			FFrameTime UsedFrameTime = CurrentFrameTimeInFrameSpace - CachedSettings.BufferSettings.TimecodeFrameOffset - CachedSettings.BufferSettings.ValidTimecodeFrame;
-			FFrameTime FrameTime = SourceFrameData.GetBaseData()->MetaData.SceneTime.Time;
-			if (FrameTime > UsedFrameTime)
+			const FQualifiedFrameTime CurrentSyncTime = FApp::GetCurrentFrameTime().GetValue();
+			const FFrameTime CurrentFrameTimeInFrameSpace = CurrentSyncTime.ConvertTo(CachedSettings.BufferSettings.TimecodeFrameRate);
+			int32 FrameIndex = 0;
+			for (const FLiveLinkFrameDataStruct& SourceFrameData : FrameData)
 			{
-				break;
+				FFrameTime UsedFrameTime = CurrentFrameTimeInFrameSpace - CachedSettings.BufferSettings.TimecodeFrameOffset - CachedSettings.BufferSettings.ValidTimecodeFrame;
+				FFrameTime FrameTime = SourceFrameData.GetBaseData()->MetaData.SceneTime.Time;
+				if (FrameTime > UsedFrameTime)
+				{
+					break;
+				}
+				++FrameIndex;
 			}
-			++FrameIndex;
-		}
 
-		if (FrameIndex - 1 >= 0)
-		{
-			const int32 Count = CachedSettings.BufferSettings.bKeepAtLeastOneFrame && FrameData.Num() == FrameIndex ? FrameIndex - 1 : FrameIndex;
-			if (Count > 0)
+			if (FrameIndex - 1 >= 0)
 			{
-				FrameData.RemoveAt(0, Count, false);
+				const int32 Count = CachedSettings.BufferSettings.bKeepAtLeastOneFrame && FrameData.Num() == FrameIndex ? FrameIndex - 1 : FrameIndex;
+				if (Count > 0)
+				{
+					FrameData.RemoveAt(0, Count, false);
+				}
 			}
 		}
+
+		// no warning if GetCurrentFrameTime is not set, the warning is done bellow after GetFrameAtSceneTime
 	}
 
 	// Build a snapshot for this role
@@ -94,8 +99,15 @@ void FLiveLinkSubject::Update()
 		{
 		case ELiveLinkSourceMode::Timecode:
 		{
-			const FQualifiedFrameTime CurrentSyncTime = FQualifiedFrameTime(FApp::GetTimecode(), FApp::GetTimecodeFrameRate());
-			bSnapshotIsValid = GetFrameAtSceneTime(CurrentSyncTime, FrameSnapshot);
+			if (FApp::GetCurrentFrameTime().IsSet())
+			{
+				bSnapshotIsValid = GetFrameAtSceneTime(FApp::GetCurrentFrameTime().GetValue(), FrameSnapshot);
+			}
+			else
+			{
+				static const FName NAME_InvalidRole = "LiveLinkSubject_NoCurrentFrameTime";
+				FLiveLinkLog::WarningOnce(NAME_InvalidRole, SubjectKey, TEXT("Can't evaluate frame for subject '%s'. The engine doesn't have a timecode value set."), *SubjectKey.SubjectName.ToString());
+			}
 		}
 		break;
 
@@ -186,7 +198,7 @@ bool FLiveLinkSubject::EvaluateFrameAtWorldTime(double InWorldTime, TSubclassOf<
 	return bSuccess;
 }
 
-bool FLiveLinkSubject::EvaluateFrameAtSceneTime(const FTimecode& InSceneTime, TSubclassOf<ULiveLinkRole> InDesiredRole, FLiveLinkSubjectFrameData& OutFrame)
+bool FLiveLinkSubject::EvaluateFrameAtSceneTime(const FQualifiedFrameTime& InSceneTime, TSubclassOf<ULiveLinkRole> InDesiredRole, FLiveLinkSubjectFrameData& OutFrame)
 {
 	if (Role == nullptr)
 	{
@@ -212,16 +224,15 @@ bool FLiveLinkSubject::EvaluateFrameAtSceneTime(const FTimecode& InSceneTime, TS
 	bool bSuccess = false;
 	if (FrameData.Num() != 0)
 	{
-		const FQualifiedFrameTime UseTime = FQualifiedFrameTime(InSceneTime, FApp::GetTimecodeFrameRate());
 		if (Role == InDesiredRole || Role->IsChildOf(InDesiredRole))
 		{
-			GetFrameAtSceneTime(UseTime, OutFrame);
+			GetFrameAtSceneTime(InSceneTime, OutFrame);
 			bSuccess = true;
 		}
 		else if (SupportsRole(InDesiredRole))
 		{
 			FLiveLinkSubjectFrameData TmpFrameData;
-			GetFrameAtSceneTime(UseTime, TmpFrameData);
+			GetFrameAtSceneTime(InSceneTime, TmpFrameData);
 			bSuccess = Translate(this, InDesiredRole, TmpFrameData.StaticData, TmpFrameData.FrameData, OutFrame);
 		}
 		else
@@ -366,8 +377,10 @@ int32 FLiveLinkSubject::FindNewFrame_SceneTime(const FQualifiedFrameTime& Qualif
 		return INDEX_NONE;
 	}
 
+	// If we do not have a TC set, keep buffering, the TC may be unresponsive for a moment. We do not want to loose data.
+	if (FApp::GetCurrentFrameTime().IsSet())
 	{
-		const FQualifiedFrameTime CurrentSyncTime = FQualifiedFrameTime(FApp::GetTimecode(), FApp::GetTimecodeFrameRate());
+		const FQualifiedFrameTime CurrentSyncTime = FApp::GetCurrentFrameTime().GetValue();
 		const FFrameTime CurrentFrameTimeInFrameSpace = CurrentSyncTime.ConvertTo(CachedSettings.BufferSettings.TimecodeFrameRate);
 		const FFrameTime CurrentOffsetFrameTime = CurrentFrameTimeInFrameSpace - CachedSettings.BufferSettings.TimecodeFrameOffset - CachedSettings.BufferSettings.ValidTimecodeFrame;
 		if (QualifiedFrameTime.Time.AsDecimal() < CurrentOffsetFrameTime.AsDecimal())
@@ -856,10 +869,11 @@ bool FLiveLinkSubject::IsTimeSynchronized() const
 	if (GetMode() == ELiveLinkSourceMode::Timecode)
 	{
 		const FLiveLinkSubjectFrameData& Snapshot = GetFrameSnapshot();
-		if (Snapshot.StaticData.IsValid() && Snapshot.FrameData.IsValid() && Snapshot.FrameData.GetBaseData())
+		if (Snapshot.StaticData.IsValid() && Snapshot.FrameData.IsValid() && Snapshot.FrameData.GetBaseData() && FApp::GetCurrentFrameTime().IsSet())
 		{
-			const FFrameNumber FrameDataInEngineFrameNumber = Snapshot.FrameData.GetBaseData()->MetaData.SceneTime.ConvertTo(FApp::GetTimecodeFrameRate()).GetFrame();
-			const FFrameNumber CurrentEngineFrameNumber = FApp::GetTimecode().ToFrameNumber(FApp::GetTimecodeFrameRate());
+			const FQualifiedFrameTime CurrentQualifiedFrameTime = FApp::GetCurrentFrameTime().GetValue();
+			const FFrameNumber FrameDataInEngineFrameNumber = Snapshot.FrameData.GetBaseData()->MetaData.SceneTime.ConvertTo(CurrentQualifiedFrameTime.Rate).GetFrame();
+			const FFrameNumber CurrentEngineFrameNumber = CurrentQualifiedFrameTime.Time.GetFrame();
 			return FrameDataInEngineFrameNumber == CurrentEngineFrameNumber;
 		}
 	}
diff --git a/Engine/Plugins/Animation/LiveLink/Source/LiveLink/Private/LiveLinkSubject.h b/Engine/Plugins/Animation/LiveLink/Source/LiveLink/Private/LiveLinkSubject.h
--- a/Engine/Plugins/Animation/LiveLink/Source/LiveLink/Private/LiveLinkSubject.h
+++ b/Engine/Plugins/Animation/LiveLink/Source/LiveLink/Private/LiveLinkSubject.h
@@ -64,7 +64,7 @@ protected:
 
 public:
 	bool EvaluateFrameAtWorldTime(double InWorldTime, TSubclassOf<ULiveLinkRole> InDesiredRole, FLiveLinkSubjectFrameData& OutFrame);
-	bool EvaluateFrameAtSceneTime(const FTimecode& InSceneTime, TSubclassOf<ULiveLinkRole> InDesiredRole, FLiveLinkSubjectFrameData& OutFrame);
+	bool EvaluateFrameAtSceneTime(const FQualifiedFrameTime& InSceneTime, TSubclassOf<ULiveLinkRole> InDesiredRole, FLiveLinkSubjectFrameData& OutFrame);
 
 	bool HasStaticData() const;
 
diff --git a/Engine/Plugins/Animation/LiveLink/Source/LiveLink/Public/LiveLinkBlueprintLibrary.h b/Engine/Plugins/Animation/LiveLink/Source/LiveLink/Public/LiveLinkBlueprintLibrary.h
--- a/Engine/Plugins/Animation/LiveLink/Source/LiveLink/Public/LiveLinkBlueprintLibrary.h
+++ b/Engine/Plugins/Animation/LiveLink/Source/LiveLink/Public/LiveLinkBlueprintLibrary.h
@@ -166,7 +166,10 @@ public:
 	UFUNCTION(BlueprintCallable, CustomThunk, Category = LiveLink, meta = (CustomStructureParam = "OutBlueprintData", BlueprintInternalUseOnly = "true", AllowAbstract = "false"))
 	static bool EvaluateLiveLinkFrameAtWorldTimeOffset(FLiveLinkSubjectName SubjectName, TSubclassOf<ULiveLinkRole> Role, float WorldTimeOffset, FLiveLinkBaseBlueprintData& OutBlueprintData);
 
-	/** Fetches a frame on a subject for a specific role at a specified scene time (timecode). Output is evaluated based on the role */
+	/**
+	 * Fetches a frame on a subject for a specific role at a specified scene time (timecode).
+	 * The Timecode should be at the frame rate as the engine timecode.
+	 * Output is evaluated based on the role */
 	UFUNCTION(BlueprintCallable, CustomThunk, Category = LiveLink, meta = (CustomStructureParam = "OutBlueprintData", BlueprintInternalUseOnly = "true", AllowAbstract = "false"))
 	static bool EvaluateLiveLinkFrameAtSceneTime(FLiveLinkSubjectName SubjectName, TSubclassOf<ULiveLinkRole> Role, FTimecode SceneTime, FLiveLinkBaseBlueprintData& OutBlueprintData);
 
diff --git a/Engine/Plugins/Animation/LiveLink/Source/LiveLink/Public/LiveLinkClient.h b/Engine/Plugins/Animation/LiveLink/Source/LiveLink/Public/LiveLinkClient.h
--- a/Engine/Plugins/Animation/LiveLink/Source/LiveLink/Public/LiveLinkClient.h
+++ b/Engine/Plugins/Animation/LiveLink/Source/LiveLink/Public/LiveLinkClient.h
@@ -11,7 +11,8 @@
 #include "LiveLinkRefSkeleton.h"
 #include "LiveLinkTypes.h"
 #include "LiveLinkVirtualSubject.h"
-#include "Misc/FrameRate.h"
+#include "Misc/Optional.h"
+#include "Misc/QualifiedFrameTime.h"
 #include "Misc/Timecode.h"
 #include "Tickable.h"
 #include "UObject/GCObject.h"
@@ -49,6 +50,7 @@ public:
 	virtual const FLiveLinkSubjectFrame* GetSubjectData(FName InSubjectName) override;
 	virtual const FLiveLinkSubjectFrame* GetSubjectDataAtWorldTime(FName InSubjectName, double InWorldTime) override;
 	virtual const FLiveLinkSubjectFrame* GetSubjectDataAtSceneTime(FName InSubjectName, const FTimecode& InSceneTime) override;
+	virtual bool EvaluateFrameAtSceneTime_AnyThread(FLiveLinkSubjectName SubjectName, const FTimecode& SceneTime, TSubclassOf<ULiveLinkRole> DesiredRole, FLiveLinkSubjectFrameData& OutFrame) override;
 	virtual const TArray<FLiveLinkFrame>* GetSubjectRawFrames(FName SubjectName) override;
 	virtual void ClearSubjectsFrames(FName SubjectName) override;
 	virtual void ClearAllSubjectsFrames() override;
@@ -116,7 +118,7 @@ public:
 	virtual bool EvaluateFrameFromSource_AnyThread(const FLiveLinkSubjectKey& SubjectKey, TSubclassOf<ULiveLinkRole> Role, FLiveLinkSubjectFrameData& OutFrame) override;
 	virtual bool EvaluateFrame_AnyThread(FLiveLinkSubjectName SubjectName, TSubclassOf<ULiveLinkRole> Role, FLiveLinkSubjectFrameData& OutFrame) override;
 	virtual bool EvaluateFrameAtWorldTime_AnyThread(FLiveLinkSubjectName SubjectName, double WorldTime, TSubclassOf<ULiveLinkRole> DesiredRole, FLiveLinkSubjectFrameData& OutFrame) override;
-	virtual bool EvaluateFrameAtSceneTime_AnyThread(FLiveLinkSubjectName SubjectName, const FTimecode& SceneTime, TSubclassOf<ULiveLinkRole> DesiredRole, FLiveLinkSubjectFrameData& OutFrame) override;
+	virtual bool EvaluateFrameAtSceneTime_AnyThread(FLiveLinkSubjectName SubjectName, const FQualifiedFrameTime& FrameTime, TSubclassOf<ULiveLinkRole> DesiredRole, FLiveLinkSubjectFrameData& OutFrame) override;
 
 	virtual FSimpleMulticastDelegate& OnLiveLinkTicked() override;
 	virtual FSimpleMulticastDelegate& OnLiveLinkSourcesChanged() override;
@@ -255,7 +257,6 @@ private:
 
 	/** Cached value of the engine timecode and frame rate*/
 	double CachedEngineTime;
-	FTimecode CachedEngineTimecode;
-	FFrameRate CachedEngineTimecodeFrameRate;
+	TOptional<FQualifiedFrameTime> CachedEngineFrameTime;
 #endif
 };
diff --git a/Engine/Plugins/Animation/LiveLink/Source/LiveLinkMovieScene/Private/MovieScene/MovieSceneLiveLinkSectionTemplate.cpp b/Engine/Plugins/Animation/LiveLink/Source/LiveLinkMovieScene/Private/MovieScene/MovieSceneLiveLinkSectionTemplate.cpp
--- a/Engine/Plugins/Animation/LiveLink/Source/LiveLinkMovieScene/Private/MovieScene/MovieSceneLiveLinkSectionTemplate.cpp
+++ b/Engine/Plugins/Animation/LiveLink/Source/LiveLinkMovieScene/Private/MovieScene/MovieSceneLiveLinkSectionTemplate.cpp
@@ -122,15 +122,7 @@ bool FMovieSceneLiveLinkSectionTemplate::GetLiveLinkFrameArray(const FFrameTime&
 {
 	//See if we have a valid time code time. 
 	//If so we may can possible send raw data if not asked to only send interpolated.
-	TOptional<FQualifiedFrameTime> TimeCodeFrameTime;
-	if (GEngine && GEngine->GetTimecodeProvider() && GEngine->GetTimecodeProvider()->GetSynchronizationState() == ETimecodeProviderSynchronizationState::Synchronized)
-	{
-		const UTimecodeProvider* TimeCodeProvider = GEngine->GetTimecodeProvider();
-		FFrameRate TCFrameRate = TimeCodeProvider->GetFrameRate();
-		FTimecode TimeCode = TimeCodeProvider->GetTimecode(); //Same as FApp::GetTimecode();
-		FFrameNumber FrameNumber = TimeCode.ToFrameNumber(TCFrameRate);
-		TimeCodeFrameTime = FQualifiedFrameTime(FFrameTime(FrameNumber), TCFrameRate);
-	}
+	TOptional<FQualifiedFrameTime> TimeCodeFrameTime = FApp::GetCurrentFrameTime();
 
 	//Send interpolated if told to or no valid timecode synced.
 	const bool bAlwaysSendInterpolated = CVarSequencerAlwaysSendInterpolatedLiveLink->GetInt() == 0 ? false : true;
diff --git a/Engine/Plugins/Experimental/VirtualCamera/Source/VirtualCamera/Private/LevelSequencePlaybackController.cpp b/Engine/Plugins/Experimental/VirtualCamera/Source/VirtualCamera/Private/LevelSequencePlaybackController.cpp
--- a/Engine/Plugins/Experimental/VirtualCamera/Source/VirtualCamera/Private/LevelSequencePlaybackController.cpp
+++ b/Engine/Plugins/Experimental/VirtualCamera/Source/VirtualCamera/Private/LevelSequencePlaybackController.cpp
@@ -197,7 +197,7 @@ FTimecode ULevelSequencePlaybackController::GetCurrentSequencePlaybackTimecode()
 		if (Sequencer)
 		{
 			const FFrameTime DisplayTime = Sequencer->GetLocalTime().ConvertTo(Sequencer->GetFocusedDisplayRate());
-			return FTimecode::FromFrameNumber(DisplayTime.FrameNumber, Sequencer->GetFocusedDisplayRate(), FTimecode::IsDropFormatTimecodeSupported(Sequencer->GetFocusedDisplayRate()));
+			return FTimecode::FromFrameNumber(DisplayTime.FrameNumber, Sequencer->GetFocusedDisplayRate());
 		}
 	}
 #endif //WITH_EDITOR
diff --git a/Engine/Plugins/Lumin/MagicLeap/Source/MagicLeapHandTracking/Private/MagicLeapHandTrackingLiveLink.cpp b/Engine/Plugins/Lumin/MagicLeap/Source/MagicLeapHandTracking/Private/MagicLeapHandTrackingLiveLink.cpp
--- a/Engine/Plugins/Lumin/MagicLeap/Source/MagicLeapHandTracking/Private/MagicLeapHandTrackingLiveLink.cpp
+++ b/Engine/Plugins/Lumin/MagicLeap/Source/MagicLeapHandTracking/Private/MagicLeapHandTrackingLiveLink.cpp
@@ -191,7 +191,7 @@ void FMagicLeapHandTracking::UpdateLiveLink()
 		FLiveLinkAnimationFrameData* LiveLinkRightFramePtr = LiveLinkRightFrame.Cast<FLiveLinkAnimationFrameData>();
 
 		LiveLinkLeftFramePtr->WorldTime = LiveLinkRightFramePtr->WorldTime = FPlatformTime::Seconds();
-		LiveLinkLeftFramePtr->MetaData.SceneTime = LiveLinkRightFramePtr->MetaData.SceneTime = FQualifiedFrameTime(FApp::GetTimecode(), FApp::GetTimecodeFrameRate());
+		LiveLinkLeftFramePtr->MetaData.SceneTime = LiveLinkRightFramePtr->MetaData.SceneTime = FApp::GetCurrentFrameTime().Get(FQualifiedFrameTime());
 
 		// Update the transforms for each subject from tracking data
 		UpdateLiveLinkTransforms(LiveLinkLeftFramePtr->Transforms, LeftHand);
diff --git a/Engine/Plugins/Media/AudioCaptureTimecodeProvider/Source/AudioCaptureTimecodeProvider/Private/AudioCaptureTimecodeProvider.cpp b/Engine/Plugins/Media/AudioCaptureTimecodeProvider/Source/AudioCaptureTimecodeProvider/Private/AudioCaptureTimecodeProvider.cpp
--- a/Engine/Plugins/Media/AudioCaptureTimecodeProvider/Source/AudioCaptureTimecodeProvider/Private/AudioCaptureTimecodeProvider.cpp
+++ b/Engine/Plugins/Media/AudioCaptureTimecodeProvider/Source/AudioCaptureTimecodeProvider/Private/AudioCaptureTimecodeProvider.cpp
@@ -8,13 +8,14 @@
 #include "LinearTimecodeDecoder.h"
 #include "Stats/StatsMisc.h"
 
+PRAGMA_DISABLE_DEPRECATION_WARNINGS
 
 /* FLinearTimecodeAudioCaptureCustomTimeStepImplementation implementation
 *****************************************************************************/
-struct UAudioCaptureTimecodeProvider::FLinearTimecodeAudioCaptureCustomTimeStepImplementation
+struct UDEPRECATED_AudioCaptureTimecodeProvider::FLinearTimecodeAudioCaptureCustomTimeStepImplementation
 {
 public:
-	FLinearTimecodeAudioCaptureCustomTimeStepImplementation(UAudioCaptureTimecodeProvider* InOwner)
+	FLinearTimecodeAudioCaptureCustomTimeStepImplementation(UDEPRECATED_AudioCaptureTimecodeProvider* InOwner)
 		: bWarnedAboutTheInvalidAudioChannel(false)
 		, bFrameRateReach0Counter(0)
 		, Owner(InOwner)
@@ -139,15 +140,15 @@ public:
 	FDropTimecode Timecode;
 
 	/** Owner of the implementation */
-	UAudioCaptureTimecodeProvider* Owner;
+	UDEPRECATED_AudioCaptureTimecodeProvider* Owner;
 
 	/** If the Owner requested the implementation to stop processing */
 	volatile bool bStopRequested;
 };
 
-/* UAudioCaptureTimecodeProvider
+/* UDEPRECATED_AudioCaptureTimecodeProvider
 *****************************************************************************/
-UAudioCaptureTimecodeProvider::UAudioCaptureTimecodeProvider(const FObjectInitializer& ObjectInitializer)
+UDEPRECATED_AudioCaptureTimecodeProvider::UDEPRECATED_AudioCaptureTimecodeProvider(const FObjectInitializer& ObjectInitializer)
 	: Super(ObjectInitializer)
 	, AudioChannel(1)
 	, Implementation(nullptr)
@@ -157,7 +158,12 @@ UAudioCaptureTimecodeProvider::UAudioCaptureTimecodeProvider(const FObjectInitia
 
 /* UTimecodeProvider interface implementation
 *****************************************************************************/
-FTimecode UAudioCaptureTimecodeProvider::GetTimecode() const
+FQualifiedFrameTime UDEPRECATED_AudioCaptureTimecodeProvider::GetQualifiedFrameTime() const
+{
+	return FQualifiedFrameTime(GetTimecodeInternal(), GetFrameRateInternal());
+}
+
+FTimecode UDEPRECATED_AudioCaptureTimecodeProvider::GetTimecodeInternal() const
 {
 	FTimecode Result;
 	{
@@ -180,7 +186,7 @@ FTimecode UAudioCaptureTimecodeProvider::GetTimecode() const
 	return Result;
 }
 
-FFrameRate UAudioCaptureTimecodeProvider::GetFrameRate() const
+FFrameRate UDEPRECATED_AudioCaptureTimecodeProvider::GetFrameRateInternal() const
 {
 	FFrameRate Result = FrameRate;
 	if (bDetectFrameRate)
@@ -219,7 +225,7 @@ FFrameRate UAudioCaptureTimecodeProvider::GetFrameRate() const
 	return Result;
 }
 
-bool UAudioCaptureTimecodeProvider::Initialize(class UEngine* InEngine)
+bool UDEPRECATED_AudioCaptureTimecodeProvider::Initialize(class UEngine* InEngine)
 {
 	check(Implementation == nullptr);
 	delete Implementation; // in case
@@ -236,15 +242,17 @@ bool UAudioCaptureTimecodeProvider::Initialize(class UEngine* InEngine)
 	return bInitialized;
 }
 
-void UAudioCaptureTimecodeProvider::Shutdown(class UEngine* InEngine)
+void UDEPRECATED_AudioCaptureTimecodeProvider::Shutdown(class UEngine* InEngine)
 {
 	SynchronizationState = ETimecodeProviderSynchronizationState::Closed;
 	delete Implementation;
 	Implementation = nullptr;
 }
 
-void UAudioCaptureTimecodeProvider::BeginDestroy()
+void UDEPRECATED_AudioCaptureTimecodeProvider::BeginDestroy()
 {
 	delete Implementation;
 	Super::BeginDestroy();
 }
+
+PRAGMA_ENABLE_DEPRECATION_WARNINGS
\ No newline at end of file
diff --git a/Engine/Plugins/Media/AudioCaptureTimecodeProvider/Source/AudioCaptureTimecodeProvider/Private/AudioCaptureTimecodeProviderModule.cpp b/Engine/Plugins/Media/AudioCaptureTimecodeProvider/Source/AudioCaptureTimecodeProvider/Private/AudioCaptureTimecodeProviderModule.cpp
--- a/Engine/Plugins/Media/AudioCaptureTimecodeProvider/Source/AudioCaptureTimecodeProvider/Private/AudioCaptureTimecodeProviderModule.cpp
+++ b/Engine/Plugins/Media/AudioCaptureTimecodeProvider/Source/AudioCaptureTimecodeProvider/Private/AudioCaptureTimecodeProviderModule.cpp
@@ -4,49 +4,9 @@
 
 #include "AudioCaptureTimecodeProvider.h"
 
-#include "Engine/Engine.h"
-#include "Misc/CoreMisc.h"
 #include "Modules/ModuleInterface.h"
 #include "Modules/ModuleManager.h"
-#include "UObject/StrongObjectPtr.h"
 
 DEFINE_LOG_CATEGORY(LogAudioCaptureTimecodeProvider);
 
-class FAudioCaptureTimecodeProviderModule : public IModuleInterface, public FSelfRegisteringExec
-{
-	TStrongObjectPtr<UAudioCaptureTimecodeProvider> TimecodeProvider;
-	virtual bool Exec(UWorld* Inworld, const TCHAR* Cmd, FOutputDevice& Ar) override
-	{
-		if (FParse::Command(&Cmd, TEXT("AudioCapture")))
-		{
-			if (FParse::Command(&Cmd, TEXT("TimecodeProvider")))
-			{
-				if (FParse::Command(&Cmd, TEXT("Start")))
-				{
-					TimecodeProvider.Reset(NewObject<UAudioCaptureTimecodeProvider>());
-
-					FParse::Bool(Cmd, TEXT("DetectFrameRate="), TimecodeProvider->bDetectFrameRate);
-					FParse::Bool(Cmd, TEXT("AssumeDropFrameFormat="), TimecodeProvider->bAssumeDropFrameFormat);
-					FParse::Value(Cmd, TEXT("Numerator="), TimecodeProvider->FrameRate.Numerator);
-					FParse::Value(Cmd, TEXT("Denominator="), TimecodeProvider->FrameRate.Denominator);
-					FParse::Value(Cmd, TEXT("AudioChannel="), TimecodeProvider->AudioChannel);
-
-					GEngine->SetTimecodeProvider(TimecodeProvider.Get());
-				}
-				else if (FParse::Command(&Cmd, TEXT("Stop")))
-				{
-					if (GEngine->GetTimecodeProvider() == TimecodeProvider.Get())
-					{
-						GEngine->SetTimecodeProvider(nullptr);
-					}
-					TimecodeProvider.Reset();
-				}
-				return true;
-			}
-		}
-
-		return false;
-	}
-};
-
-IMPLEMENT_MODULE(FAudioCaptureTimecodeProviderModule, AudioCaptureTimecodeProvider)
+IMPLEMENT_MODULE(FDefaultModuleImpl, AudioCaptureTimecodeProvider)
diff --git a/Engine/Plugins/Media/AudioCaptureTimecodeProvider/Source/AudioCaptureTimecodeProvider/Public/AudioCaptureTimecodeProvider.h b/Engine/Plugins/Media/AudioCaptureTimecodeProvider/Source/AudioCaptureTimecodeProvider/Public/AudioCaptureTimecodeProvider.h
--- a/Engine/Plugins/Media/AudioCaptureTimecodeProvider/Source/AudioCaptureTimecodeProvider/Public/AudioCaptureTimecodeProvider.h
+++ b/Engine/Plugins/Media/AudioCaptureTimecodeProvider/Source/AudioCaptureTimecodeProvider/Public/AudioCaptureTimecodeProvider.h
@@ -12,8 +12,8 @@
 /**
  * Read the LTC from the audio capture device.
  */
-UCLASS(Blueprintable, editinlinenew)
-class AUDIOCAPTURETIMECODEPROVIDER_API UAudioCaptureTimecodeProvider : public UTimecodeProvider
+UCLASS(Blueprintable, editinlinenew, deprecated)
+class AUDIOCAPTURETIMECODEPROVIDER_API UDEPRECATED_AudioCaptureTimecodeProvider : public UTimecodeProvider
 {
 	GENERATED_UCLASS_BODY()
 
@@ -43,8 +43,7 @@ public:
 
 public:
 	//~ UTimecodeProvider interface
-	virtual FTimecode GetTimecode() const override;
-	virtual FFrameRate GetFrameRate() const override;
+	virtual FQualifiedFrameTime GetQualifiedFrameTime() const override;
 	virtual ETimecodeProviderSynchronizationState GetSynchronizationState() const override { return SynchronizationState; }
 	virtual bool Initialize(class UEngine* InEngine) override;
 	virtual void Shutdown(class UEngine* InEngine) override;
@@ -53,6 +52,9 @@ public:
 	virtual void BeginDestroy() override;
 	
 private:
+	FTimecode GetTimecodeInternal() const;
+	FFrameRate GetFrameRateInternal() const;
+
 	/** Audio capture object dealing with getting audio callbacks */
 	struct FLinearTimecodeAudioCaptureCustomTimeStepImplementation;
 	friend FLinearTimecodeAudioCaptureCustomTimeStepImplementation;
diff --git a/Engine/Plugins/Media/TimecodeSynchronizer/Source/TimecodeSynchronizer/Private/TimecodeSynchronizer.cpp b/Engine/Plugins/Media/TimecodeSynchronizer/Source/TimecodeSynchronizer/Private/TimecodeSynchronizer.cpp
--- a/Engine/Plugins/Media/TimecodeSynchronizer/Source/TimecodeSynchronizer/Private/TimecodeSynchronizer.cpp
+++ b/Engine/Plugins/Media/TimecodeSynchronizer/Source/TimecodeSynchronizer/Private/TimecodeSynchronizer.cpp
@@ -271,7 +271,12 @@ void UTimecodeSynchronizer::PostEditChangeChainProperty(FPropertyChangedChainEve
 }
 #endif
 
-FTimecode UTimecodeSynchronizer::GetTimecode() const
+FQualifiedFrameTime UTimecodeSynchronizer::GetQualifiedFrameTime() const
+{
+	return FQualifiedFrameTime(GetTimecodeInternal(), GetFrameRateInternal());
+}
+
+FTimecode UTimecodeSynchronizer::GetTimecodeInternal() const
 {
 	FTimecode Timecode;
 	if (IsSynchronized())
@@ -327,7 +332,7 @@ FFrameTime UTimecodeSynchronizer::GetProviderFrameTime() const
 	return ProviderFrameTime;
 }
 
-FFrameRate UTimecodeSynchronizer::GetFrameRate() const
+FFrameRate UTimecodeSynchronizer::GetFrameRateInternal() const
 {
 	return RegisteredCustomTimeStep ? RegisteredCustomTimeStep->GetFixedFrameRate() : FixedFrameRate;
 }
@@ -890,7 +895,7 @@ void UTimecodeSynchronizer::StartSources()
 	FTimeSynchronizationStartData StartData;
 	CurrentSystemFrameTime = StartData.StartFrame = CalculateSyncTime();
 
-	FApp::SetTimecodeAndFrameRate(GetTimecode(), GetFrameRate());
+	FApp::SetCurrentFrameTime(GetDelayedQualifiedFrameTime());
 
 	for (UTimeSynchronizationSource* InputSource : TimeSynchronizationInputSources)
 	{
diff --git a/Engine/Plugins/Media/TimecodeSynchronizer/Source/TimecodeSynchronizer/Public/TimecodeSynchronizer.h b/Engine/Plugins/Media/TimecodeSynchronizer/Source/TimecodeSynchronizer/Public/TimecodeSynchronizer.h
--- a/Engine/Plugins/Media/TimecodeSynchronizer/Source/TimecodeSynchronizer/Public/TimecodeSynchronizer.h
+++ b/Engine/Plugins/Media/TimecodeSynchronizer/Source/TimecodeSynchronizer/Public/TimecodeSynchronizer.h
@@ -265,8 +265,7 @@ public:
 	//~ End UObject Interface
 
 	//~ Begin TimecodeProvider Interface
-	virtual FTimecode GetTimecode() const override;
-	virtual FFrameRate GetFrameRate() const override;
+	FQualifiedFrameTime GetQualifiedFrameTime() const override;
 	virtual ETimecodeProviderSynchronizationState GetSynchronizationState() const override;
 	virtual bool Initialize(class UEngine* InEngine) override;
 	virtual void Shutdown(class UEngine* InEngine) override;
@@ -434,6 +433,9 @@ private:
 	void UpdateSourceStates();
 	FFrameTime CalculateSyncTime();
 
+	FTimecode GetTimecodeInternal() const;
+	FFrameRate GetFrameRateInternal() const;
+
 	bool IsSynchronizing() const;
 	bool IsSynchronized() const;
 	bool IsError() const;
diff --git a/Engine/Plugins/Media/TimecodeSynchronizer/Source/TimecodeSynchronizerEditor/Private/Widgets/STimecodeSynchronizerSourceViewport.cpp b/Engine/Plugins/Media/TimecodeSynchronizer/Source/TimecodeSynchronizerEditor/Private/Widgets/STimecodeSynchronizerSourceViewport.cpp
--- a/Engine/Plugins/Media/TimecodeSynchronizer/Source/TimecodeSynchronizerEditor/Private/Widgets/STimecodeSynchronizerSourceViewport.cpp
+++ b/Engine/Plugins/Media/TimecodeSynchronizer/Source/TimecodeSynchronizerEditor/Private/Widgets/STimecodeSynchronizerSourceViewport.cpp
@@ -169,8 +169,7 @@ FText STimecodeSynchronizerSourceViewport::HandleIntervalMinTimecodeText() const
 	if (AttachedSource && AttachedSource->IsReady())
 	{
 		const FFrameNumber OldestFrame = AttachedSource->GetInputSourceState().OldestAvailableSample.GetFrame();
-		const bool bIsDropFrame = FTimecode::IsDropFormatTimecodeSupported(AttachedSource->GetFrameRate());
-		Timecode = FTimecode::FromFrameNumber(OldestFrame, AttachedSource->GetFrameRate(), bIsDropFrame);
+		Timecode = FTimecode::FromFrameNumber(OldestFrame, AttachedSource->GetFrameRate());
 	}
 
 	return FText::FromString(FString("Minimum: ") + Timecode.ToString());
@@ -183,8 +182,7 @@ FText STimecodeSynchronizerSourceViewport::HandleIntervalMaxTimecodeText() const
 	if (AttachedSource && AttachedSource->IsReady())
 	{
 		const FFrameNumber NewestFrame = AttachedSource->GetInputSourceState().NewestAvailableSample.GetFrame();
-		const bool bIsDropFrame = FTimecode::IsDropFormatTimecodeSupported(AttachedSource->GetFrameRate());
-		Timecode = FTimecode::FromFrameNumber(NewestFrame, AttachedSource->GetFrameRate(), bIsDropFrame);
+		Timecode = FTimecode::FromFrameNumber(NewestFrame, AttachedSource->GetFrameRate());
 	}
 
 	return FText::FromString(FString("Maximum: ") + Timecode.ToString());
diff --git a/Engine/Plugins/Runtime/AR/Apple/AppleARKit/Source/AppleARKit/Private/AppleARKitTimecodeProvider.cpp b/Engine/Plugins/Runtime/AR/Apple/AppleARKit/Source/AppleARKit/Private/AppleARKitTimecodeProvider.cpp
--- a/Engine/Plugins/Runtime/AR/Apple/AppleARKit/Source/AppleARKit/Private/AppleARKitTimecodeProvider.cpp
+++ b/Engine/Plugins/Runtime/AR/Apple/AppleARKit/Source/AppleARKit/Private/AppleARKitTimecodeProvider.cpp
@@ -18,14 +18,11 @@ void UAppleARKitTimecodeProvider::Shutdown(class UEngine*)
 {
 }
 
-FTimecode UAppleARKitTimecodeProvider::GetTimecode() const
+FQualifiedFrameTime UAppleARKitTimecodeProvider::GetQualifiedFrameTime() const
 {
-	return FTimecode(FPlatformTime::Seconds(), FrameRate, FTimecode::IsDropFormatTimecodeSupported(FrameRate), true);
-}
-
-FFrameRate UAppleARKitTimecodeProvider::GetFrameRate() const
-{
-	return FrameRate;
+	// We construct a timecode to not have subframe.
+	FTimecode Timecode = FTimecode(FPlatformTime::Seconds(), FrameRate, true);
+	return FQualifiedFrameTime(Timecode, FrameRate);
 }
 
 ETimecodeProviderSynchronizationState UAppleARKitTimecodeProvider::GetSynchronizationState() const
diff --git a/Engine/Plugins/Runtime/AR/Apple/AppleARKit/Source/AppleARKit/Private/AppleARKitTimecodeProvider.h b/Engine/Plugins/Runtime/AR/Apple/AppleARKit/Source/AppleARKit/Private/AppleARKitTimecodeProvider.h
--- a/Engine/Plugins/Runtime/AR/Apple/AppleARKit/Source/AppleARKit/Private/AppleARKitTimecodeProvider.h
+++ b/Engine/Plugins/Runtime/AR/Apple/AppleARKit/Source/AppleARKit/Private/AppleARKitTimecodeProvider.h
@@ -19,8 +19,7 @@ public:
 	UAppleARKitTimecodeProvider();
 
 	// ~UTimecodeProvider
-	virtual FTimecode GetTimecode() const override;
-	virtual FFrameRate GetFrameRate() const override;
+	virtual FQualifiedFrameTime GetQualifiedFrameTime() const override;
 	virtual ETimecodeProviderSynchronizationState GetSynchronizationState() const override;
 	virtual bool Initialize(class UEngine* InEngine) override;
 	virtual void Shutdown(class UEngine* InEngine) override;
diff --git a/Engine/Plugins/Runtime/nDisplay/Source/DisplayCluster/Private/Cluster/Controller/DisplayClusterClusterNodeCtrlMaster.cpp b/Engine/Plugins/Runtime/nDisplay/Source/DisplayCluster/Private/Cluster/Controller/DisplayClusterClusterNodeCtrlMaster.cpp
--- a/Engine/Plugins/Runtime/nDisplay/Source/DisplayCluster/Private/Cluster/Controller/DisplayClusterClusterNodeCtrlMaster.cpp
+++ b/Engine/Plugins/Runtime/nDisplay/Source/DisplayCluster/Private/Cluster/Controller/DisplayClusterClusterNodeCtrlMaster.cpp
@@ -60,25 +60,30 @@ void FDisplayClusterClusterNodeCtrlMaster::GetDeltaTime(float& DeltaSeconds)
 	DeltaSeconds = CachedDeltaTime_GetDeltaTime;
 }
 
-void FDisplayClusterClusterNodeCtrlMaster::GetTimecode(FTimecode& Timecode, FFrameRate& FrameRate)
+void FDisplayClusterClusterNodeCtrlMaster::GetFrameTime(TOptional<FQualifiedFrameTime>& FrameTime)
 {
 	FScopeLock lock(&InternalsSyncScope);
 
 	// Cache data so it will be the same for all requests within current frame
-	if (!bIsDataCached_GetTimecode)
+	if (!bIsDataCached_GetFrameTime)
 	{
-		bIsDataCached_GetTimecode = true;
+		bIsDataCached_GetFrameTime = true;
 
 		// This values are updated in UEngine::UpdateTimeAndHandleMaxTickRate (via UpdateTimecode).
-		CachedTimecode_GetTimecode  = FApp::GetTimecode();
-		CachedFramerate_GetTimecode = FApp::GetTimecodeFrameRate();
+		CachedTimecode_GetFrameTime = FApp::GetCurrentFrameTime();
 
-		UE_LOG(LogDisplayClusterCluster, Verbose, TEXT("GetDeltaTime cached values: DeltaSeconds %f"), CachedDeltaTime_GetDeltaTime);
+		if (CachedTimecode_GetFrameTime.IsSet())
+		{
+			UE_LOG(LogDisplayClusterCluster, Verbose, TEXT("GetFrameTime cached values: Seconds %f"), CachedTimecode_GetFrameTime.GetValue().AsSeconds());
+		}
+		else
+		{
+			UE_LOG(LogDisplayClusterCluster, Verbose, TEXT("GetFrameTime cached values: [INVALID]"));
+		}
 	}
 
 	// Return cached value
-	Timecode  = CachedTimecode_GetTimecode;
-	FrameRate = CachedFramerate_GetTimecode;
+	FrameTime = CachedTimecode_GetFrameTime;
 }
 
 void FDisplayClusterClusterNodeCtrlMaster::GetSyncData(FDisplayClusterMessage::DataType& SyncData, EDisplayClusterSyncGroup SyncGroup)
@@ -183,7 +188,7 @@ void FDisplayClusterClusterNodeCtrlMaster::ClearCache()
 
 	// Reset all cache flags
 	bIsDataCached_GetDeltaTime  = false;
-	bIsDataCached_GetTimecode   = false;
+	bIsDataCached_GetFrameTime  = false;
 	bIsDataCached_GetInputData  = false;
 	bIsDataCached_GetEventsData = false;
 
diff --git a/Engine/Plugins/Runtime/nDisplay/Source/DisplayCluster/Private/Cluster/Controller/DisplayClusterClusterNodeCtrlMaster.h b/Engine/Plugins/Runtime/nDisplay/Source/DisplayCluster/Private/Cluster/Controller/DisplayClusterClusterNodeCtrlMaster.h
--- a/Engine/Plugins/Runtime/nDisplay/Source/DisplayCluster/Private/Cluster/Controller/DisplayClusterClusterNodeCtrlMaster.h
+++ b/Engine/Plugins/Runtime/nDisplay/Source/DisplayCluster/Private/Cluster/Controller/DisplayClusterClusterNodeCtrlMaster.h
@@ -27,7 +27,7 @@ public:
 	// IPDisplayClusterClusterSyncProtocol
 	//////////////////////////////////////////////////////////////////////////////////////////////
 	virtual void GetDeltaTime(float& DeltaSeconds) override;
-	virtual void GetTimecode(FTimecode& Timecode, FFrameRate& FrameRate) override;
+	virtual void GetFrameTime(TOptional<FQualifiedFrameTime>& FrameTime) override;
 	virtual void GetSyncData(FDisplayClusterMessage::DataType& SyncData, EDisplayClusterSyncGroup SyncGroup) override;
 	virtual void GetInputData(FDisplayClusterMessage::DataType& InputData) override;
 	virtual void GetEventsData(FDisplayClusterMessage::DataType& EventsData) override;
@@ -66,9 +66,8 @@ private:
 	float CachedDeltaTime_GetDeltaTime = 0.f;
 
 	// GetTimecode internals
-	bool       bIsDataCached_GetTimecode = false;
-	FTimecode  CachedTimecode_GetTimecode;
-	FFrameRate CachedFramerate_GetTimecode;
+	bool bIsDataCached_GetFrameTime = false;
+	TOptional<FQualifiedFrameTime>  CachedTimecode_GetFrameTime;
 
 	// GetSyncData internals
 	TMap<EDisplayClusterSyncGroup, bool> IsDataCached_GetSyncData;
diff --git a/Engine/Plugins/Runtime/nDisplay/Source/DisplayCluster/Private/Cluster/Controller/DisplayClusterClusterNodeCtrlSlave.cpp b/Engine/Plugins/Runtime/nDisplay/Source/DisplayCluster/Private/Cluster/Controller/DisplayClusterClusterNodeCtrlSlave.cpp
--- a/Engine/Plugins/Runtime/nDisplay/Source/DisplayCluster/Private/Cluster/Controller/DisplayClusterClusterNodeCtrlSlave.cpp
+++ b/Engine/Plugins/Runtime/nDisplay/Source/DisplayCluster/Private/Cluster/Controller/DisplayClusterClusterNodeCtrlSlave.cpp
@@ -49,9 +49,9 @@ void FDisplayClusterClusterNodeCtrlSlave::GetDeltaTime(float& DeltaSeconds)
 	ClusterSyncClient->GetDeltaTime(DeltaSeconds);
 }
 
-void FDisplayClusterClusterNodeCtrlSlave::GetTimecode(FTimecode& Timecode, FFrameRate& FrameRate)
+void FDisplayClusterClusterNodeCtrlSlave::GetFrameTime(TOptional<FQualifiedFrameTime>& FrameTime)
 {
-	ClusterSyncClient->GetTimecode(Timecode, FrameRate);
+	ClusterSyncClient->GetFrameTime(FrameTime);
 }
 
 void FDisplayClusterClusterNodeCtrlSlave::GetSyncData(FDisplayClusterMessage::DataType& SyncData, EDisplayClusterSyncGroup SyncGroup)
diff --git a/Engine/Plugins/Runtime/nDisplay/Source/DisplayCluster/Private/Cluster/Controller/DisplayClusterClusterNodeCtrlSlave.h b/Engine/Plugins/Runtime/nDisplay/Source/DisplayCluster/Private/Cluster/Controller/DisplayClusterClusterNodeCtrlSlave.h
--- a/Engine/Plugins/Runtime/nDisplay/Source/DisplayCluster/Private/Cluster/Controller/DisplayClusterClusterNodeCtrlSlave.h
+++ b/Engine/Plugins/Runtime/nDisplay/Source/DisplayCluster/Private/Cluster/Controller/DisplayClusterClusterNodeCtrlSlave.h
@@ -37,7 +37,7 @@ public:
 	virtual void WaitForFrameEnd()   override final;
 	virtual void WaitForTickEnd()    override final;
 	virtual void GetDeltaTime(float& DeltaSeconds) override;
-	virtual void GetTimecode(FTimecode& Timecode, FFrameRate& FrameRate) override;
+	virtual void GetFrameTime(TOptional<FQualifiedFrameTime>& FrameTime) override;
 	virtual void GetSyncData(FDisplayClusterMessage::DataType& SyncData, EDisplayClusterSyncGroup SyncGroup) override;
 	virtual void GetInputData(FDisplayClusterMessage::DataType& InputData) override;
 	virtual void GetEventsData(FDisplayClusterMessage::DataType& EventsData) override;
diff --git a/Engine/Plugins/Runtime/nDisplay/Source/DisplayCluster/Private/Cluster/Controller/DisplayClusterNodeCtrlBase.h b/Engine/Plugins/Runtime/nDisplay/Source/DisplayCluster/Private/Cluster/Controller/DisplayClusterNodeCtrlBase.h
--- a/Engine/Plugins/Runtime/nDisplay/Source/DisplayCluster/Private/Cluster/Controller/DisplayClusterNodeCtrlBase.h
+++ b/Engine/Plugins/Runtime/nDisplay/Source/DisplayCluster/Private/Cluster/Controller/DisplayClusterNodeCtrlBase.h
@@ -70,7 +70,7 @@ public:
 	virtual void GetDeltaTime(float& DeltaSeconds) override
 	{ }
 
-	virtual void GetTimecode(FTimecode& Timecode, FFrameRate& FrameRate) override
+	virtual void GetFrameTime(TOptional<FQualifiedFrameTime>& FrameTime) override
 	{ }
 
 	virtual void GetSyncData(FDisplayClusterMessage::DataType& SyncData, EDisplayClusterSyncGroup SyncGroup) override
diff --git a/Engine/Plugins/Runtime/nDisplay/Source/DisplayCluster/Private/Game/Classes/Basics/DisplayClusterGameEngine.cpp b/Engine/Plugins/Runtime/nDisplay/Source/DisplayCluster/Private/Game/Classes/Basics/DisplayClusterGameEngine.cpp
--- a/Engine/Plugins/Runtime/nDisplay/Source/DisplayCluster/Private/Game/Classes/Basics/DisplayClusterGameEngine.cpp
+++ b/Engine/Plugins/Runtime/nDisplay/Source/DisplayCluster/Private/Game/Classes/Basics/DisplayClusterGameEngine.cpp
@@ -11,6 +11,7 @@
 #include "Misc/CommandLine.h"
 #include "Misc/DisplayClusterAppExit.h"
 #include "Misc/Parse.h"
+#include "Misc/QualifiedFrameTime.h"
 
 #include "Misc/DisplayClusterHelpers.h"
 
@@ -195,8 +196,7 @@ void UDisplayClusterGameEngine::Tick(float DeltaSeconds, bool bIdleMode)
 	if (OperationMode == EDisplayClusterOperationMode::Cluster ||
 		OperationMode == EDisplayClusterOperationMode::Standalone)
 	{
-		FTimecode Timecode;
-		FFrameRate FrameRate;
+		TOptional<FQualifiedFrameTime> FrameTime;
 
 		//////////////////////////////////////////////////////////////////////////////////////////////
 		// Frame start barrier
@@ -212,9 +212,18 @@ void UDisplayClusterGameEngine::Tick(float DeltaSeconds, bool bIdleMode)
 		UE_LOG(LogDisplayClusterEngine, Verbose, TEXT("DisplayCluster delta seconds: %f"), DeltaSeconds);
 
 		// Sync timecode and framerate
-		NodeController->GetTimecode(Timecode, FrameRate);
-		FApp::SetTimecodeAndFrameRate(Timecode, FrameRate);
-		UE_LOG(LogDisplayClusterEngine, Verbose, TEXT("DisplayCluster timecode: %s | %s"), *Timecode.ToString(), *FrameRate.ToPrettyText().ToString());
+		NodeController->GetFrameTime(FrameTime);
+
+		if (FrameTime.IsSet())
+		{
+			FApp::SetCurrentFrameTime(FrameTime.GetValue());
+			UE_LOG(LogDisplayClusterEngine, Verbose, TEXT("DisplayCluster timecode: %s | %s"), *FTimecode::FromFrameNumber(FrameTime->Time.GetFrame(), FrameTime->Rate).ToString(), *FrameTime->Rate.ToPrettyText().ToString());
+		}
+		else
+		{
+			FApp::InvalidateCurrentFrameTime();
+			UE_LOG(LogDisplayClusterEngine, Verbose, TEXT("DisplayCluster timecode: [Invalid]"));
+		}
 
 		// Perform PreTick for DisplayCluster module
 		UE_LOG(LogDisplayClusterEngine, Verbose, TEXT("Perform PreTick()"));
diff --git a/Engine/Plugins/Runtime/nDisplay/Source/DisplayCluster/Private/Network/Protocol/IPDisplayClusterClusterSyncProtocol.h b/Engine/Plugins/Runtime/nDisplay/Source/DisplayCluster/Private/Network/Protocol/IPDisplayClusterClusterSyncProtocol.h
--- a/Engine/Plugins/Runtime/nDisplay/Source/DisplayCluster/Private/Network/Protocol/IPDisplayClusterClusterSyncProtocol.h
+++ b/Engine/Plugins/Runtime/nDisplay/Source/DisplayCluster/Private/Network/Protocol/IPDisplayClusterClusterSyncProtocol.h
@@ -4,9 +4,8 @@
 
 #include "Network/DisplayClusterMessage.h"
 #include "DisplayClusterEnums.h"
-
-struct FTimecode;
-struct FFrameRate;
+#include "Misc/Optional.h"
+#include "Misc/QualifiedFrameTime.h"
 
 
 /**
@@ -31,7 +30,7 @@ public:
 	virtual void GetDeltaTime(float& DeltaSeconds) = 0;
 
 	// Get the Timecode value for the current frame.
-	virtual void GetTimecode(FTimecode& Timecode, FFrameRate& FrameRate) = 0;
+	virtual void GetFrameTime(TOptional<FQualifiedFrameTime>& FrameTime) = 0;
 
 	// Sync objects
 	virtual void GetSyncData(FDisplayClusterMessage::DataType& SyncData, EDisplayClusterSyncGroup SyncGroup) = 0;
diff --git a/Engine/Plugins/Runtime/nDisplay/Source/DisplayCluster/Private/Network/Service/ClusterSync/DisplayClusterClusterSyncClient.cpp b/Engine/Plugins/Runtime/nDisplay/Source/DisplayCluster/Private/Network/Service/ClusterSync/DisplayClusterClusterSyncClient.cpp
--- a/Engine/Plugins/Runtime/nDisplay/Source/DisplayCluster/Private/Network/Service/ClusterSync/DisplayClusterClusterSyncClient.cpp
+++ b/Engine/Plugins/Runtime/nDisplay/Source/DisplayCluster/Private/Network/Service/ClusterSync/DisplayClusterClusterSyncClient.cpp
@@ -65,24 +65,34 @@ void FDisplayClusterClusterSyncClient::GetDeltaTime(float& DeltaSeconds)
 	}
 }
 
-void FDisplayClusterClusterSyncClient::GetTimecode(FTimecode& Timecode, FFrameRate& FrameRate)
+void FDisplayClusterClusterSyncClient::GetFrameTime(TOptional<FQualifiedFrameTime>& FrameTime)
 {
-	static const TSharedPtr<FDisplayClusterMessage> Request(new FDisplayClusterMessage(FDisplayClusterClusterSyncMsg::GetTimecode::name, FDisplayClusterClusterSyncMsg::TypeRequest, FDisplayClusterClusterSyncMsg::ProtocolName));
+	static const TSharedPtr<FDisplayClusterMessage> Request(new FDisplayClusterMessage(FDisplayClusterClusterSyncMsg::GetFrameTime::name, FDisplayClusterClusterSyncMsg::TypeRequest, FDisplayClusterClusterSyncMsg::ProtocolName));
 	TSharedPtr<FDisplayClusterMessage> Response = SendRecvMsg(Request);
 
+	FrameTime.Reset();
+
 	if (!Response.IsValid())
 	{
 		return;
 	}
 
 	// Extract sync data from response message
-	if (Response->GetArg(FDisplayClusterClusterSyncMsg::GetTimecode::argTimecode, Timecode) == false)
+	bool bIsValid = false;
+	if (Response->GetArg(FDisplayClusterClusterSyncMsg::GetFrameTime::argIsValid, bIsValid) == false)
 	{
-		UE_LOG(LogDisplayClusterNetworkMsg, Error, TEXT("Couldn't extract an argument: %s"), FDisplayClusterClusterSyncMsg::GetTimecode::argTimecode);
+		UE_LOG(LogDisplayClusterNetworkMsg, Error, TEXT("Couldn't extract an argument: %s"), FDisplayClusterClusterSyncMsg::GetFrameTime::argIsValid);
 	}
-	if (Response->GetArg(FDisplayClusterClusterSyncMsg::GetTimecode::argFrameRate, FrameRate) == false)
+
+	if (bIsValid)
 	{
-		UE_LOG(LogDisplayClusterNetworkMsg, Error, TEXT("Couldn't extract an argument: %s"), FDisplayClusterClusterSyncMsg::GetTimecode::argTimecode);
+		FQualifiedFrameTime NewFrameTime;
+		if (Response->GetArg(FDisplayClusterClusterSyncMsg::GetFrameTime::argFrameTime, NewFrameTime) == false)
+		{
+			UE_LOG(LogDisplayClusterNetworkMsg, Error, TEXT("Couldn't extract an argument: %s"), FDisplayClusterClusterSyncMsg::GetFrameTime::argFrameTime);
+		}
+
+		FrameTime = NewFrameTime;
 	}
 }
 
diff --git a/Engine/Plugins/Runtime/nDisplay/Source/DisplayCluster/Private/Network/Service/ClusterSync/DisplayClusterClusterSyncClient.h b/Engine/Plugins/Runtime/nDisplay/Source/DisplayCluster/Private/Network/Service/ClusterSync/DisplayClusterClusterSyncClient.h
--- a/Engine/Plugins/Runtime/nDisplay/Source/DisplayCluster/Private/Network/Service/ClusterSync/DisplayClusterClusterSyncClient.h
+++ b/Engine/Plugins/Runtime/nDisplay/Source/DisplayCluster/Private/Network/Service/ClusterSync/DisplayClusterClusterSyncClient.h
@@ -27,7 +27,7 @@ public:
 	virtual void WaitForFrameEnd() override;
 	virtual void WaitForTickEnd() override;
 	virtual void GetDeltaTime(float& DeltaSeconds) override;
-	virtual void GetTimecode(FTimecode& Timecode, FFrameRate& FrameRate) override;
+	virtual void GetFrameTime(TOptional<FQualifiedFrameTime>& FrameTime) override;
 	virtual void GetSyncData(FDisplayClusterMessage::DataType& SyncData, EDisplayClusterSyncGroup SyncGroup) override;
 	virtual void GetInputData(FDisplayClusterMessage::DataType& InputData) override;
 	virtual void GetEventsData(FDisplayClusterMessage::DataType& EventsData) override;
diff --git a/Engine/Plugins/Runtime/nDisplay/Source/DisplayCluster/Private/Network/Service/ClusterSync/DisplayClusterClusterSyncMsg.h b/Engine/Plugins/Runtime/nDisplay/Source/DisplayCluster/Private/Network/Service/ClusterSync/DisplayClusterClusterSyncMsg.h
--- a/Engine/Plugins/Runtime/nDisplay/Source/DisplayCluster/Private/Network/Service/ClusterSync/DisplayClusterClusterSyncMsg.h
+++ b/Engine/Plugins/Runtime/nDisplay/Source/DisplayCluster/Private/Network/Service/ClusterSync/DisplayClusterClusterSyncMsg.h
@@ -40,11 +40,11 @@ namespace FDisplayClusterClusterSyncMsg
 		constexpr static auto argDeltaSeconds = "DeltaSeconds";
 	};
 
-	namespace GetTimecode
+	namespace GetFrameTime
 	{
-		constexpr static auto name         = "GetTimecode";
-		constexpr static auto argTimecode  = "Timecode";
-		constexpr static auto argFrameRate = "FrameRate";
+		constexpr static auto name          = "GetFrameTime";
+		constexpr static auto argIsValid    = "IsValid";
+		constexpr static auto argFrameTime  = "FrameTime";
 	}
 
 	namespace GetSyncData
diff --git a/Engine/Plugins/Runtime/nDisplay/Source/DisplayCluster/Private/Network/Service/ClusterSync/DisplayClusterClusterSyncService.cpp b/Engine/Plugins/Runtime/nDisplay/Source/DisplayCluster/Private/Network/Service/ClusterSync/DisplayClusterClusterSyncService.cpp
--- a/Engine/Plugins/Runtime/nDisplay/Source/DisplayCluster/Private/Network/Service/ClusterSync/DisplayClusterClusterSyncService.cpp
+++ b/Engine/Plugins/Runtime/nDisplay/Source/DisplayCluster/Private/Network/Service/ClusterSync/DisplayClusterClusterSyncService.cpp
@@ -13,6 +13,7 @@
 #include "Network/Session/DisplayClusterSessionInternal.h"
 
 #include "Misc/DisplayClusterAppExit.h"
+#include "Misc/QualifiedFrameTime.h"
 
 #include "DisplayClusterEnums.h"
 #include "DisplayClusterGlobals.h"
@@ -131,13 +132,16 @@ TSharedPtr<FDisplayClusterMessage> FDisplayClusterClusterSyncService::ProcessMes
 		Response->SetArg(FDisplayClusterClusterSyncMsg::GetDeltaTime::argDeltaSeconds, DeltaSeconds);
 		return Response;
 	}
-	else if (ReqName == FDisplayClusterClusterSyncMsg::GetTimecode::name)
+	else if (ReqName == FDisplayClusterClusterSyncMsg::GetFrameTime::name)
 	{
-		FTimecode timecode;
-		FFrameRate frameRate;
-		GetTimecode(timecode, frameRate);
-		Response->SetArg(FDisplayClusterClusterSyncMsg::GetTimecode::argTimecode, timecode);
-		Response->SetArg(FDisplayClusterClusterSyncMsg::GetTimecode::argFrameRate, frameRate);
+		TOptional<FQualifiedFrameTime> frameTime;
+		GetFrameTime(frameTime);
+
+		Response->SetArg(FDisplayClusterClusterSyncMsg::GetFrameTime::argIsValid, frameTime.IsSet());
+		if (frameTime.IsSet())
+		{
+			Response->SetArg(FDisplayClusterClusterSyncMsg::GetFrameTime::argFrameTime, frameTime.GetValue());
+		}
 		return Response;
 	}
 	else if (ReqName == FDisplayClusterClusterSyncMsg::GetSyncData::name)
@@ -224,10 +228,10 @@ void FDisplayClusterClusterSyncService::GetDeltaTime(float& DeltaSeconds)
 	return NodeController->GetDeltaTime(DeltaSeconds);
 }
 
-void FDisplayClusterClusterSyncService::GetTimecode(FTimecode& Timecode, FFrameRate& FrameRate)
+void FDisplayClusterClusterSyncService::GetFrameTime(TOptional<FQualifiedFrameTime>& FrameTime)
 {
 	static IPDisplayClusterNodeController* const NodeController = GDisplayCluster->GetPrivateClusterMgr()->GetController();
-	return NodeController->GetTimecode(Timecode, FrameRate);
+	return NodeController->GetFrameTime(FrameTime);
 }
 
 void FDisplayClusterClusterSyncService::GetSyncData(FDisplayClusterMessage::DataType& SyncData, EDisplayClusterSyncGroup SyncGroup)
diff --git a/Engine/Plugins/Runtime/nDisplay/Source/DisplayCluster/Private/Network/Service/ClusterSync/DisplayClusterClusterSyncService.h b/Engine/Plugins/Runtime/nDisplay/Source/DisplayCluster/Private/Network/Service/ClusterSync/DisplayClusterClusterSyncService.h
--- a/Engine/Plugins/Runtime/nDisplay/Source/DisplayCluster/Private/Network/Service/ClusterSync/DisplayClusterClusterSyncService.h
+++ b/Engine/Plugins/Runtime/nDisplay/Source/DisplayCluster/Private/Network/Service/ClusterSync/DisplayClusterClusterSyncService.h
@@ -44,7 +44,7 @@ private:
 	virtual void WaitForFrameEnd() override;
 	virtual void WaitForTickEnd() override;
 	virtual void GetDeltaTime(float& DeltaSeconds) override;
-	virtual void GetTimecode(FTimecode& Timecode, FFrameRate& FrameRate) override;
+	virtual void GetFrameTime(TOptional<FQualifiedFrameTime>& FrameTime) override;
 	virtual void GetSyncData(FDisplayClusterMessage::DataType& SyncData, EDisplayClusterSyncGroup SyncGroup)  override;
 	virtual void GetInputData(FDisplayClusterMessage::DataType& InputData) override;
 	virtual void GetEventsData(FDisplayClusterMessage::DataType& EventsData) override;
diff --git a/Engine/Plugins/Runtime/nDisplay/Source/DisplayClusterUtils/DisplayClusterTypesConverter.h b/Engine/Plugins/Runtime/nDisplay/Source/DisplayClusterUtils/DisplayClusterTypesConverter.h
--- a/Engine/Plugins/Runtime/nDisplay/Source/DisplayClusterUtils/DisplayClusterTypesConverter.h
+++ b/Engine/Plugins/Runtime/nDisplay/Source/DisplayClusterUtils/DisplayClusterTypesConverter.h
@@ -5,6 +5,7 @@
 #include "CoreMinimal.h"
 #include "Misc/Timecode.h"
 #include "Misc/FrameRate.h"
+#include "Misc/QualifiedFrameTime.h"
 #include "DisplayClusterEnums.h"
 #include "DisplayClusterUtils/DisplayClusterCommonStrings.h"
 #include "DisplayClusterUtils/DisplayClusterCommonHelpers.h"
@@ -31,10 +32,10 @@ public:
 
 template <> inline FString FDisplayClusterTypesConverter::ToString<>(const FString& from)    { return from; }
 template <> inline FString FDisplayClusterTypesConverter::ToString<>(const bool& from)       { return (from ? DisplayClusterStrings::cfg::spec::ValTrue : DisplayClusterStrings::cfg::spec::ValFalse); }
-template <> inline FString FDisplayClusterTypesConverter::ToString<> (const int8& from)      { return FString::FromInt(from); }
-template <> inline FString FDisplayClusterTypesConverter::ToString<> (const uint8& from)     { return ToString(static_cast<int8>(from)); }
-template <> inline FString FDisplayClusterTypesConverter::ToString<> (const int32& from)     { return FString::FromInt(from); }
-template <> inline FString FDisplayClusterTypesConverter::ToString<> (const uint32& from)    { return ToString(static_cast<int32>(from)); }
+template <> inline FString FDisplayClusterTypesConverter::ToString<>(const int8& from)      { return FString::FromInt(from); }
+template <> inline FString FDisplayClusterTypesConverter::ToString<>(const uint8& from)     { return ToString(static_cast<int8>(from)); }
+template <> inline FString FDisplayClusterTypesConverter::ToString<>(const int32& from)     { return FString::FromInt(from); }
+template <> inline FString FDisplayClusterTypesConverter::ToString<>(const uint32& from)    { return ToString(static_cast<int32>(from)); }
 template <> inline FString FDisplayClusterTypesConverter::ToString<>(const float& from)      { return FString::SanitizeFloat(from); }
 template <> inline FString FDisplayClusterTypesConverter::ToString<>(const double& from)     { return FString::Printf(TEXT("%lf"), from); }
 template <> inline FString FDisplayClusterTypesConverter::ToString<>(const FVector& from)    { return from.ToString(); }
@@ -45,6 +46,7 @@ template <> inline FString FDisplayClusterTypesConverter::ToString<>(const FMatr
 // We can't just use FTimecode ToString as that loses information.
 template <> inline FString FDisplayClusterTypesConverter::ToString<>(const FTimecode& from)  { return FString::Printf(TEXT("%d;%d;%d;%d;%d"), from.bDropFrameFormat ? 1 : 0, from.Hours, from.Minutes, from.Seconds, from.Frames); }
 template <> inline FString FDisplayClusterTypesConverter::ToString<>(const FFrameRate& from) { return FString::Printf(TEXT("%d;%d"), from.Numerator, from.Denominator); }
+template <> inline FString FDisplayClusterTypesConverter::ToString<>(const FQualifiedFrameTime& from) { return FString::Printf(TEXT("%d;%s;%d;%d"), from.Time.GetFrame().Value, *FString::SanitizeFloat(from.Time.GetSubFrame()), from.Rate.Numerator, from.Rate.Denominator); }
 
 template <> inline FString FDisplayClusterTypesConverter::ToString<>(const EDisplayClusterOperationMode& from)
 {
@@ -178,6 +180,25 @@ template <> inline FFrameRate FDisplayClusterTypesConverter::FromString<> (const
 	return frameRate;
 }
 
+template <> inline FQualifiedFrameTime FDisplayClusterTypesConverter::FromString<>(const FString& from)
+{
+	FQualifiedFrameTime frameTime;
+
+	TArray<FString> parts;
+	parts.Reserve(4);
+	const int32 found = from.ParseIntoArray(parts, TEXT(";"));
+
+	// We are expecting 4 "parts" - Frame, SubFrame, Numerator, Denominator.
+	if (found == 4)
+	{
+		frameTime.Time = FFrameTime(FromString<int32>(parts[0]), FromString<float>(parts[1]));
+		frameTime.Rate.Numerator = FromString<int32>(parts[2]);
+		frameTime.Rate.Denominator = FromString<int32>(parts[3]);
+	}
+
+	return frameTime;
+}
+
 template <> inline EDisplayClusterSyncGroup FDisplayClusterTypesConverter::FromString<>(const FString& from)
 {
 	if (from.Equals(TEXT("PreTick")))
diff --git a/Engine/Plugins/VirtualProduction/Takes/Source/TakeRecorder/Private/Recorder/TakeRecorder.cpp b/Engine/Plugins/VirtualProduction/Takes/Source/TakeRecorder/Private/Recorder/TakeRecorder.cpp
--- a/Engine/Plugins/VirtualProduction/Takes/Source/TakeRecorder/Private/Recorder/TakeRecorder.cpp
+++ b/Engine/Plugins/VirtualProduction/Takes/Source/TakeRecorder/Private/Recorder/TakeRecorder.cpp
@@ -235,7 +235,7 @@ private:
 
 				FFrameNumber TotalFrames = FFrameNumber(static_cast<int32>(FrameRate.AsDecimal() * RecordingDuration.GetTotalSeconds()));
 
-				FTimecode Timecode = FTimecode::FromFrameNumber(TotalFrames, FrameRate, FTimecode::IsDropFormatTimecodeSupported(FrameRate));
+				FTimecode Timecode = FTimecode::FromFrameNumber(TotalFrames, FrameRate);
 
 				return FText::Format(LOCTEXT("RecordingTimecodeText", "Recording...{0}"), FText::FromString(Timecode.ToString()));
 			}
diff --git a/Engine/Source/Editor/Sequencer/Private/SSequencer.cpp b/Engine/Source/Editor/Sequencer/Private/SSequencer.cpp
--- a/Engine/Source/Editor/Sequencer/Private/SSequencer.cpp
+++ b/Engine/Source/Editor/Sequencer/Private/SSequencer.cpp
@@ -1875,7 +1875,7 @@ void SSequencer::FillPlaybackSpeedMenu(FMenuBuilder& InMenuBarBuilder)
 void SSequencer::FillTimeDisplayFormatMenu(FMenuBuilder& MenuBuilder)
 {
 	TSharedPtr<FSequencer> Sequencer = SequencerPtr.Pin();
-	bool bSupportsDropFormatDisplay = FTimecode::IsDropFormatTimecodeSupported(Sequencer->GetFocusedDisplayRate());
+	bool bShouldDisplayDropFormat = FTimecode::UseDropFormatTimecode(Sequencer->GetFocusedDisplayRate());
 
 	const UEnum* FrameNumberDisplayEnum = StaticEnum<EFrameNumberDisplayFormats>();
 	check(FrameNumberDisplayEnum);
@@ -1888,9 +1888,17 @@ void SSequencer::FillTimeDisplayFormatMenu(FMenuBuilder& MenuBuilder)
 			{
 				EFrameNumberDisplayFormats Value = (EFrameNumberDisplayFormats)FrameNumberDisplayEnum->GetValueByIndex(Index);
 
+				// Don't show None Drop Frame Timecode when the format support drop format and the engine wants to use the drop format by default.
+				if (Value == EFrameNumberDisplayFormats::NonDropFrameTimecode && bShouldDisplayDropFormat)
+				{
+					continue;
+				}
+
 				// Don't show Drop Frame Timecode when they're in a format that doesn't support it.
-				if (Value == EFrameNumberDisplayFormats::DropFrameTimecode && !bSupportsDropFormatDisplay)
+				if (Value == EFrameNumberDisplayFormats::DropFrameTimecode && !bShouldDisplayDropFormat)
+				{
 					continue;
+				}
 
 				MenuBuilder.AddMenuEntry(
 					FrameNumberDisplayEnum->GetDisplayNameTextByIndex(Index),
diff --git a/Engine/Source/Editor/Sequencer/Private/Sequencer.cpp b/Engine/Source/Editor/Sequencer/Private/Sequencer.cpp
--- a/Engine/Source/Editor/Sequencer/Private/Sequencer.cpp
+++ b/Engine/Source/Editor/Sequencer/Private/Sequencer.cpp
@@ -10333,8 +10333,12 @@ void FSequencer::BindCommands()
 
 			// If the next framerate in the list is drop format timecode and we're not in a play rate that supports drop format timecode,
 			// then we will skip over it.
-			bool bCanShowDropFrameTimecode = FTimecode::IsDropFormatTimecodeSupported(GetFocusedDisplayRate());
-			if (!bCanShowDropFrameTimecode && NextFormat == EFrameNumberDisplayFormats::DropFrameTimecode)
+			bool bCanShowDropFrameTimecode = FTimecode::UseDropFormatTimecode(GetFocusedDisplayRate());
+			if (bCanShowDropFrameTimecode && NextFormat == EFrameNumberDisplayFormats::NonDropFrameTimecode)
+			{
+				NextFormat = EFrameNumberDisplayFormats::DropFrameTimecode;
+			}
+			else if (!bCanShowDropFrameTimecode && NextFormat == EFrameNumberDisplayFormats::DropFrameTimecode)
 			{
 				NextFormat = EFrameNumberDisplayFormats::Seconds;
 			}
diff --git a/Engine/Source/Runtime/Core/Private/Misc/App.cpp b/Engine/Source/Runtime/Core/Private/Misc/App.cpp
--- a/Engine/Source/Runtime/Core/Private/Misc/App.cpp
+++ b/Engine/Source/Runtime/Core/Private/Misc/App.cpp
@@ -36,8 +36,7 @@ double FApp::LastTime = 0.0;
 double FApp::DeltaTime = 1 / 30.0;
 double FApp::IdleTime = 0.0;
 double FApp::IdleTimeOvershoot = 0.0;
-FTimecode FApp::Timecode = FTimecode();
-FFrameRate FApp::TimecodeFrameRate = FFrameRate(60,1);
+TOptional<FQualifiedFrameTime> FApp::CurrentFrameTime;
 float FApp::VolumeMultiplier = 1.0f;
 float FApp::UnfocusedVolumeMultiplier = 0.0f;
 bool FApp::bUseVRFocus = false;
@@ -262,6 +261,24 @@ bool FApp::ShouldUseThreadingForPerformance()
 
 #endif // HAVE_RUNTIME_THREADING_SWITCHES
 
+FTimecode FApp::GetTimecode()
+{
+	if (CurrentFrameTime.IsSet())
+	{
+		return FTimecode::FromFrameNumber(CurrentFrameTime->Time.GetFrame(), CurrentFrameTime->Rate);
+	}
+	return FTimecode();
+}
+
+FFrameRate FApp::GetTimecodeFrameRate()
+{
+	if (CurrentFrameTime.IsSet())
+	{
+		return CurrentFrameTime->Rate;
+	}
+	return FFrameRate();
+}
+
 static bool GUnfocusedVolumeMultiplierInitialised = false;
 float FApp::GetUnfocusedVolumeMultiplier()
 {
diff --git a/Engine/Source/Runtime/Core/Private/Misc/Timecode.cpp b/Engine/Source/Runtime/Core/Private/Misc/Timecode.cpp
new file mode 100644
--- /dev/null
+++ b/Engine/Source/Runtime/Core/Private/Misc/Timecode.cpp
@@ -0,0 +1,17 @@
+// Copyright 1998-2019 Epic Games, Inc. All Rights Reserved.
+
+#include "Misc/Timecode.h"
+
+
+static TAutoConsoleVariable<bool> CVarUseDropFormatTimecodeByDefaultWhenSupported(
+	TEXT("timecode.UseDropFormatTimecodeByDefaultWhenSupported"),
+	1.0f,
+	TEXT("By default, should we generate a timecode in drop frame format when the frame rate does support it."),
+	ECVF_Default);
+
+/* FTimecode interface
+ *****************************************************************************/
+bool FTimecode::UseDropFormatTimecodeByDefaultWhenSupported()
+{
+	return CVarUseDropFormatTimecodeByDefaultWhenSupported.GetValueOnAnyThread();
+}
diff --git a/Engine/Source/Runtime/Core/Public/Misc/App.h b/Engine/Source/Runtime/Core/Public/Misc/App.h
--- a/Engine/Source/Runtime/Core/Public/Misc/App.h
+++ b/Engine/Source/Runtime/Core/Public/Misc/App.h
@@ -13,8 +13,8 @@
 #include "Misc/Guid.h"
 #include "Misc/CoreMisc.h"
 #include "Misc/CommandLine.h"
-#include "Misc/Timecode.h"
-#include "Misc/FrameRate.h"
+#include "Misc/Optional.h"
+#include "Misc/QualifiedFrameTime.h"
 #include "HAL/PlatformProcess.h"
 
 // platforms which can have runtime threading switches
@@ -597,30 +597,31 @@ public:
 	}
 
 	/**
-	 * Gets the current timecode.
+	 * Convert the current frame time into a readable timecode.
+	 * If the current frame time is not set, the timecode value will be defaulted.
 	 *
 	 * @return the current timecode.
 	 */
-	FORCEINLINE static FTimecode GetTimecode()
-	{
-		return Timecode;
-	}
+	static FTimecode GetTimecode();
 
 	/**
-	 * Gets the current timecode frame rate.
+	 * Get the frame rate of the current frame time.
+	 * If the current frame time is not set, the frame rate value will be defaulted.
 	 *
 	 * @return the current timecode frame rate.
 	 */
-	FORCEINLINE static FFrameRate GetTimecodeFrameRate()
-	{
-		return TimecodeFrameRate;
-	}
-	
-	UE_DEPRECATED(4.21, "Please use the version of SetTimecodeAndFrameRate")
-	static void SetTimecode(FTimecode InTimecode)
+	static FFrameRate GetTimecodeFrameRate();
+
+	/**
+	 * Gets a frame number generated by the engine's timecode provider.
+	 * The current frame time is used to generate a timecode value.
+	 * The optional will be false if no timecode provider was set or is not in a synchronized state.
+	 *
+	 * @return the current frame time.
+	 */
+	FORCEINLINE static TOptional<FQualifiedFrameTime> GetCurrentFrameTime()
 	{
-		Timecode = InTimecode;
-		TimecodeFrameRate = FFrameRate(60, 1);
+		return CurrentFrameTime;
 	}
 
 	/**
@@ -629,10 +630,26 @@ public:
 	 * @param InTimecode - current timecode.
 	 * @param InTimecodeFrameRate - current timecode framerate.
 	 */
+	UE_DEPRECATED(4.25, "Please use SetQualifiedFrameTime")
 	static void SetTimecodeAndFrameRate(FTimecode InTimecode, FFrameRate InTimecodeFrameRate)
 	{
-		Timecode = InTimecode;
-		TimecodeFrameRate = InTimecodeFrameRate;
+		CurrentFrameTime = FQualifiedFrameTime(InTimecode, InTimecodeFrameRate);
+	}
+
+	/**
+	 * Sets the current frame time.
+	 *
+	 * @param InFrameTime - current frame time and frame rate.
+	 */
+	static void SetCurrentFrameTime(FQualifiedFrameTime InFrameTime)
+	{
+		CurrentFrameTime = InFrameTime;
+	}
+
+	/** Invalidate the current frame time. It will reset the TOptional. */
+	static void InvalidateCurrentFrameTime()
+	{
+		CurrentFrameTime.Reset();
 	}
 
 	/**
@@ -726,7 +743,6 @@ private:
 	/** Holds the name the graphics RHI currently in use*/
 	static FString GraphicsRHI;
 
-
 	/** List of authorized session users. */
 	static TArray<FString> SessionUsers;
 
@@ -757,11 +773,11 @@ private:
 	/** Holds the amount of IdleTime that was LONGER than we tried to sleep. The OS can't sleep the exact amount of time, so this measures that overshoot. */
 	static double IdleTimeOvershoot;
 
-	/** Holds the current timecode. */
-	static FTimecode Timecode;
+	/** Holds the current frame time and framerate. */
+	static TOptional<FQualifiedFrameTime> CurrentFrameTime;
 
-	/** Holds the current timecode frame rate. */
-	static FFrameRate TimecodeFrameRate;
+	/** Holds if we should generate a drop frame timecode when the frame rate does support it. */
+	static bool bUseDropFrameFormatWhenSupported;
 
 	/** Use to affect the app volume when it loses focus */
 	static float VolumeMultiplier;
diff --git a/Engine/Source/Runtime/Core/Public/Misc/Timecode.h b/Engine/Source/Runtime/Core/Public/Misc/Timecode.h
--- a/Engine/Source/Runtime/Core/Public/Misc/Timecode.h
+++ b/Engine/Source/Runtime/Core/Public/Misc/Timecode.h
@@ -61,6 +61,18 @@ struct FTimecode
 		*this = FromFrameNumber(FFrameNumber(NumberOfFrames), InFrameRate, InbDropFrame);
 	}
 
+	/**
+	 * User construction from a time in seconds
+	 * @param InbRollover	- If true, the hours will be the modulo of 24.
+	 * @note Be aware that the Cycles may not correspond to the System Time. See FDateTime and "leap seconds".
+	 */
+	explicit FTimecode(double InSeconds, const FFrameRate& InFrameRate, bool InbRollover)
+		: FTimecode(InSeconds, InFrameRate, UseDropFormatTimecode(InFrameRate), InbRollover)
+	{
+	}
+
+	
+
 	friend bool operator==(const FFrameRate& A, const FFrameRate& B);
 	friend bool operator!=(const FFrameRate& A, const FFrameRate& B);
 
@@ -75,6 +87,11 @@ public:
 		const int32 NumberOfFramesInMinute = NumberOfFramesInSecond * 60;
 		const int32 NumberOfFramesInHour = NumberOfFramesInMinute * 60;
 
+		if (NumberOfFramesInSecond <= 0)
+		{
+			return FFrameNumber();
+		}
+
 		// Do a quick pre-pass to take any overflow values and move them into bigger time units.
 		int32 SafeSeconds = Seconds + Frames / NumberOfFramesInSecond;
 		int32 SafeFrames = Frames % NumberOfFramesInSecond;
@@ -123,6 +140,11 @@ public:
 		const int32 NumberOfFramesInMinute = NumberOfFramesInSecond * 60;
 		const int32 NumberOfFramesInHour = NumberOfFramesInMinute * 60;
 
+		if (NumberOfFramesInSecond <= 0)
+		{
+			return FTimecode();
+		}
+
 		if (InbDropFrame)
 		{
 			// Drop Frame Timecode (DFT) was created to address the issue with playing back whole frames at fractional framerates.
@@ -190,6 +212,18 @@ public:
 		}
 	}
 
+	/**
+	 * Create a FTimecode from a specific frame number at the given frame rate.
+	 *
+	 * @param InFrameNumber - The frame number to convert into a timecode. This should already be converted to InFrameRate's resolution.
+	 * @param InFrameRate	- The framerate that this timecode is based in. This should be the playback framerate as it is used to determine
+	 *						  when the Frame value wraps over.
+	 */
+	static FTimecode FromFrameNumber(const FFrameNumber& InFrameNumber, const FFrameRate& InFrameRate)
+	{
+		return FromFrameNumber(InFrameNumber, InFrameRate, UseDropFormatTimecode(InFrameRate));
+	}
+
 	/**
 	 * Converts this Timecode back into a timespan at the given framerate, taking into account if this is a drop-frame format timecode.
 	 */
@@ -219,6 +253,20 @@ public:
 		return FTimecode(InTimespan.GetTotalSeconds(), InFrameRate, InbDropFrame, InbRollover);
 	}
 
+	/**
+	 * Create a FTimecode from a timespan at the given frame rate.
+	 *
+	 * @param InFrameNumber	- The timespan to convert into a timecode.
+	 * @param InFrameRate	- The framerate that this timecode is based in. This should be the playback framerate as it is used to determine
+	 *						  when the Frame value wraps over.
+	 * @param InbRollover	- If true, the hours will be the modulo of 24.
+	 */
+	static FTimecode FromTimespan(const FTimespan& InTimespan, const FFrameRate& InFrameRate, bool InbRollover)
+	{
+		return FTimecode(InTimespan.GetTotalSeconds(), InFrameRate, UseDropFormatTimecode(InFrameRate), InbRollover);
+	}
+
+	/** Drop frame is only support for frame rate of 29.97 or 59.94. */
 	static bool IsDropFormatTimecodeSupported(const FFrameRate& InFrameRate)
 	{
 		// Drop Format Timecode is only valid for 29.97 and 59.94.
@@ -228,6 +276,15 @@ public:
 		return InFrameRate == TwentyNineNineSeven || InFrameRate == FiftyNineNineFour;
 	}
 
+	/** If the frame rate support drop frame format and the app wish to use drop frame format by default. */
+	static bool UseDropFormatTimecode(const FFrameRate& InFrameRate)
+	{
+		return IsDropFormatTimecodeSupported(InFrameRate) && UseDropFormatTimecodeByDefaultWhenSupported();
+	}
+
+	/** By default, should we generate a timecode in drop frame format when the frame rate does support it. */
+	static CORE_API bool UseDropFormatTimecodeByDefaultWhenSupported();
+
 	/**
 	 * Get the Qualified Timecode formatted in HH:MM:SS:FF or HH;MM;SS;FF depending on if this represents drop-frame timecode or not.
 	 * @param bForceSignDisplay - Forces the timecode to be prepended with a positive or negative sign.
@@ -237,23 +294,25 @@ public:
 	{
 		bool bHasNegativeComponent = Hours < 0 || Minutes < 0 || Seconds < 0 || Frames < 0;
 
-		FString SignText;
+		const TCHAR* NegativeSign = TEXT("- ");
+		const TCHAR* PositiveSign = TEXT("+ ");
+		const TCHAR* SignText = TEXT("");
 		if (bHasNegativeComponent)
 		{
-			SignText = "- ";
+			SignText = NegativeSign;
 		}
 		else if (bForceSignDisplay)
 		{
-			SignText = "+ ";
+			SignText = PositiveSign;
 		}
 
 		if (bDropFrameFormat)
 		{
-			return FString::Printf(TEXT("%s%02d;%02d;%02d;%02d"), *SignText, FMath::Abs(Hours), FMath::Abs(Minutes), FMath::Abs(Seconds), FMath::Abs(Frames));
+			return FString::Printf(TEXT("%s%02d;%02d;%02d;%02d"), SignText, FMath::Abs(Hours), FMath::Abs(Minutes), FMath::Abs(Seconds), FMath::Abs(Frames));
 		}
 		else
 		{
-			return FString::Printf(TEXT("%s%02d:%02d:%02d:%02d"), *SignText, FMath::Abs(Hours), FMath::Abs(Minutes), FMath::Abs(Seconds), FMath::Abs(Frames));
+			return FString::Printf(TEXT("%s%02d:%02d:%02d:%02d"), SignText, FMath::Abs(Hours), FMath::Abs(Minutes), FMath::Abs(Seconds), FMath::Abs(Frames));
 		}
 	}
 
diff --git a/Engine/Source/Runtime/Engine/Classes/Engine/Engine.h b/Engine/Source/Runtime/Engine/Classes/Engine/Engine.h
--- a/Engine/Source/Runtime/Engine/Classes/Engine/Engine.h
+++ b/Engine/Source/Runtime/Engine/Classes/Engine/Engine.h
@@ -1354,15 +1354,26 @@ private:
 	bool bIsCurrentTimecodeProviderInitialized;
 
 public:
+	/** Set TimecodeProvider when the engine is started. */
+	UPROPERTY(config, EditAnywhere, Category=Timecode, meta=(MetaClass="TimecodeProvider", DisplayName="Timecode Provider", ConfigRestartRequired=true))
+	FSoftClassPath TimecodeProviderClassName;
+
 	/**
-	 * Set TimecodeProvider when the engine is started.
+	 * Generate a default timecode from the computer clock when there is no timecode provider.
+	 * On desktop, the system time will be used and will behave as if a USystemTimecodeProvider was set.
+	 * On console, the high performance clock will be used. That may introduce drift over time.
+	 * If you wish to use the system time on console, set the timecode provider to USystemeTimecodeProvider.
 	 */
-	UPROPERTY(config, EditAnywhere, Category=Timecode, meta=(MetaClass="TimecodeProvider", DisplayName="TimecodeProvider", ConfigRestartRequired=true))
-	FSoftClassPath TimecodeProviderClassName;
+	UPROPERTY(config, EditAnywhere, Category=Timecode)
+	bool bGenerateDefaultTimecode;
+
+	/** When generating a default timecode (bGenerateDefaultTimecode is true and no timecode provider is set) at which frame rate it should be generated (number of frames). */
+	UPROPERTY(config, EditAnywhere, Category=Timecode, meta=(EditCondition="bGenerateDefaultTimecode"))
+	FFrameRate GenerateDefaultTimecodeFrameRate;
 
 public:
 	/** 
-	 * Whether we should check for more than N pawns spawning in a single frame.  
+	 * Whether we should check for more than N pawns spawning in a single frame.
 	 * Basically, spawning pawns and all of their attachments can be slow.  And on consoles it
 	 * can be really slow.  If this bool is true we will display a 
 	 **/
diff --git a/Engine/Source/Runtime/Engine/Classes/Engine/SystemTimeTimecodeProvider.h b/Engine/Source/Runtime/Engine/Classes/Engine/SystemTimeTimecodeProvider.h
--- a/Engine/Source/Runtime/Engine/Classes/Engine/SystemTimeTimecodeProvider.h
+++ b/Engine/Source/Runtime/Engine/Classes/Engine/SystemTimeTimecodeProvider.h
@@ -3,8 +3,6 @@
 #pragma once
 
 #include "TimecodeProvider.h"
-#include "Misc/FrameRate.h"
-#include "UObject/UnrealType.h"
 #include "SystemTimeTimecodeProvider.generated.h"
 
 /**
@@ -30,13 +28,18 @@ public:
 	{
 	}
 
-	//~ Begin UTimecodeProvider Interface
-	virtual FTimecode GetTimecode() const override;
+	/** Generate a timecode value using the system clock. */
+	static FTimecode GenerateTimecodeFromSystemTime(FFrameRate Rate);
 
-	virtual FFrameRate GetFrameRate() const override
-	{
-		return FrameRate;
-	}
+	/**
+	 * Generate a timecode value using the high performance clock
+	 * Using the high performance clock is faster but will make the value drift over time.
+	 * This is an optimized version. Prefer GenerateTimecodeFromSystemTime, if the value need to be accurate.
+	 **/
+	static FTimecode GenerateTimecodeFromHighPerformanceClock(FFrameRate Rate);
+
+	//~ Begin UTimecodeProvider Interface
+	virtual FQualifiedFrameTime GetQualifiedFrameTime() const override;
 	
 	virtual ETimecodeProviderSynchronizationState GetSynchronizationState() const override
 	{
diff --git a/Engine/Source/Runtime/Engine/Classes/Engine/TimecodeProvider.h b/Engine/Source/Runtime/Engine/Classes/Engine/TimecodeProvider.h
--- a/Engine/Source/Runtime/Engine/Classes/Engine/TimecodeProvider.h
+++ b/Engine/Source/Runtime/Engine/Classes/Engine/TimecodeProvider.h
@@ -8,6 +8,7 @@
 
 #include "Misc/FrameRate.h"
 #include "Misc/Timecode.h"
+#include "Misc/QualifiedFrameTime.h"
 
 #include "TimecodeProvider.generated.h"
 
@@ -42,32 +43,37 @@ class ENGINE_API UTimecodeProvider : public UObject
 
 public:
 	/**
-	 * Number of frames to subtract from the current timecode when GetDelayedTimecode is called.
-	 * @see GetDelayedTimecode
+	 * Number of frames to subtract from the qualified frame time when GetDelayedQualifiedFrameTime or GetDelayedTimecode is called.
+	 * @see GetDelayedQualifiedFrameTime, GetDelayedTimecode
 	 */
 	UPROPERTY(EditAnywhere, AdvancedDisplay, Category = "Settings")
 	int32 FrameDelay = 0;
-
+	
 	/**
-	 * Return the Timecode at that moment. It may not be in sync with the current frame.
-	 * Only valid when GetSynchronizationState() is Synchronized.
-	 */
+	 * Return the frame number and the frame rate of the frame number at that moment. It may not be in sync with the current frame.
+	 * Depending on the implementation, it may or may not be valid only when GetSynchronizationState() is Synchronized.
+	*/
 	UFUNCTION(BlueprintCallable, Category = "Provider")
-	virtual FTimecode GetTimecode() const PURE_VIRTUAL(UTimecodeProvider::GetTimecode, return FTimecode(););
+	virtual FQualifiedFrameTime GetQualifiedFrameTime() const PURE_VIRTUAL(UTimecodeProvider::GetQualifiedFrameTime, return FQualifiedFrameTime(););
 
 	/**
-	 * Return the Timecode with the Frame Delay applied. It may not be in sync with the current frame.
-	 * Only valid when GetSynchronizationState() is Synchronized.
-	 */
+	 * Return the frame number and the frame rate of the frame number with the Frame Delay applied. It may not be in sync with the current frame.
+	 * Depending on the implementation, it may or may not be valid only when GetSynchronizationState() is Synchronized.
+	*/
+	UFUNCTION(BlueprintCallable, Category = "Provider")
+	FQualifiedFrameTime GetDelayedQualifiedFrameTime() const;
+
+	/** Return the frame time converted into a timecode value. */
+	UFUNCTION(BlueprintCallable, Category = "Provider")
+	FTimecode GetTimecode() const;
+
+	/** Return the delayed frame time converted into a timecode value. */
 	UFUNCTION(BlueprintCallable, Category = "Provider")
 	FTimecode GetDelayedTimecode() const;
 	
-	/**
-	 * Return the frame rate.
-	 * Depending on the implementation, it may or may not be valid only when GetSynchronizationState() is Synchronized.
-	 */
+	/** Return the frame rate of the frame time. */
 	UFUNCTION(BlueprintCallable, Category = "Provider")
-	virtual FFrameRate GetFrameRate() const PURE_VIRTUAL(UTimecodeProvider::GetFrameRate, return FFrameRate(););
+	FFrameRate GetFrameRate() const { return GetQualifiedFrameTime().Rate; }
 
 	/** The state of the TimecodeProvider and if it's currently synchronized and the Timecode and FrameRate are valid. */
 	UFUNCTION(BlueprintCallable, Category = "Provider")
diff --git a/Engine/Source/Runtime/Engine/Private/Engine/SystemTimeTimecodeProvider.cpp b/Engine/Source/Runtime/Engine/Private/Engine/SystemTimeTimecodeProvider.cpp
--- a/Engine/Source/Runtime/Engine/Private/Engine/SystemTimeTimecodeProvider.cpp
+++ b/Engine/Source/Runtime/Engine/Private/Engine/SystemTimeTimecodeProvider.cpp
@@ -2,13 +2,38 @@
 
 #include "Engine/SystemTimeTimecodeProvider.h"
 
+#include "HAL/PlatformTime.h"
 #include "Misc/CoreMisc.h"
 #include "Misc/DateTime.h"
 
-FTimecode USystemTimeTimecodeProvider::GetTimecode() const
+
+FTimecode USystemTimeTimecodeProvider::GenerateTimecodeFromSystemTime(FFrameRate FrameRate)
 {
 	const FDateTime DateTime = FDateTime::Now();
 	const FTimespan Timespan = DateTime.GetTimeOfDay();
+	return FTimecode::FromTimespan(Timespan, FrameRate, false);
+}
+
+
+static double ComputeTimeCodeOffset()
+{
+	const FDateTime DateTime = FDateTime::Now();
+	const double HighPerformanceClock = FPlatformTime::Seconds();
+	return DateTime.GetTimeOfDay().GetTotalSeconds() - HighPerformanceClock;
+}
 
-	return FTimecode::FromTimespan(Timespan, FrameRate, FTimecode::IsDropFormatTimecodeSupported(FrameRate), false);
+
+FTimecode USystemTimeTimecodeProvider::GenerateTimecodeFromHighPerformanceClock(FFrameRate FrameRate)
+{
+	constexpr double SecondsPerDay = 24.0 * 60.0 * 60.0;
+	static double HighPerformanceClockDelta = ComputeTimeCodeOffset();
+	const FTimespan Timespan = FTimespan::FromSeconds(fmod(HighPerformanceClockDelta + FPlatformTime::Seconds(), SecondsPerDay));
+	return FTimecode::FromTimespan(Timespan, FrameRate, false);
+}
+
+
+FQualifiedFrameTime USystemTimeTimecodeProvider::GetQualifiedFrameTime() const
+{
+	// Generate a TC to prevent subframe
+	return FQualifiedFrameTime(GenerateTimecodeFromSystemTime(FrameRate), FrameRate);
 }
diff --git a/Engine/Source/Runtime/Engine/Private/Engine/TimecodeProvider.cpp b/Engine/Source/Runtime/Engine/Private/Engine/TimecodeProvider.cpp
--- a/Engine/Source/Runtime/Engine/Private/Engine/TimecodeProvider.cpp
+++ b/Engine/Source/Runtime/Engine/Private/Engine/TimecodeProvider.cpp
@@ -3,15 +3,23 @@
 #include "Engine/TimecodeProvider.h"
 
 
+FQualifiedFrameTime UTimecodeProvider::GetDelayedQualifiedFrameTime() const
+{
+	FQualifiedFrameTime NewFrameTime = GetQualifiedFrameTime();
+	NewFrameTime.Time.FrameNumber -= FrameDelay;
+	return NewFrameTime;
+}
+
+
+FTimecode UTimecodeProvider::GetTimecode() const
+{
+	const FQualifiedFrameTime NewFrameTime = GetQualifiedFrameTime();
+	return FTimecode::FromFrameNumber(NewFrameTime.Time.GetFrame(), NewFrameTime.Rate);
+}
+
+
 FTimecode UTimecodeProvider::GetDelayedTimecode() const
 {
-	FTimecode NewTimecode = GetTimecode();
-	FFrameRate NewFrameRate = GetFrameRate();
-	if (FrameDelay != 0)
-	{
-		FFrameNumber TimecodeFrameNumber = NewTimecode.ToFrameNumber(NewFrameRate);
-		TimecodeFrameNumber -= FrameDelay;
-		NewTimecode = FTimecode::FromFrameNumber(TimecodeFrameNumber, NewFrameRate, NewTimecode.bDropFrameFormat);
-	}
-	return NewTimecode;
+	const FQualifiedFrameTime NewFrameTime = GetDelayedQualifiedFrameTime();
+	return FTimecode::FromFrameNumber(NewFrameTime.Time.GetFrame(), NewFrameTime.Rate);
 }
diff --git a/Engine/Source/Runtime/Engine/Private/GameEngine.cpp b/Engine/Source/Runtime/Engine/Private/GameEngine.cpp
--- a/Engine/Source/Runtime/Engine/Private/GameEngine.cpp
+++ b/Engine/Source/Runtime/Engine/Private/GameEngine.cpp
@@ -722,6 +722,9 @@ UEngine::UEngine(const FObjectInitializer& ObjectInitializer)
 
 	bIsVanillaProduct = false;
 
+	bGenerateDefaultTimecode = true;
+	GenerateDefaultTimecodeFrameRate = FFrameRate(1, 24);
+
 	GameScreenshotSaveDirectory.Path = FPaths::ScreenShotDir();
 
 	LastGCFrame = TNumericLimits<uint64>::Max();
diff --git a/Engine/Source/Runtime/Engine/Private/UnrealEngine.cpp b/Engine/Source/Runtime/Engine/Private/UnrealEngine.cpp
--- a/Engine/Source/Runtime/Engine/Private/UnrealEngine.cpp
+++ b/Engine/Source/Runtime/Engine/Private/UnrealEngine.cpp
@@ -2235,14 +2235,24 @@ bool UEngine::SetTimecodeProvider(UTimecodeProvider* InTimecodeProvider)
 
 void UEngine::UpdateTimecode()
 {
+	FApp::InvalidateCurrentFrameTime();
+
 	const UTimecodeProvider* Provider = GetTimecodeProvider();
-	if (Provider && Provider->GetSynchronizationState() == ETimecodeProviderSynchronizationState::Synchronized)
+	if (Provider)
 	{
-		FApp::SetTimecodeAndFrameRate(Provider->GetDelayedTimecode(), Provider->GetFrameRate());
+		if (Provider->GetSynchronizationState() == ETimecodeProviderSynchronizationState::Synchronized)
+		{
+			FApp::SetCurrentFrameTime(Provider->GetDelayedQualifiedFrameTime());
+		}
 	}
-	else
+	else if(bGenerateDefaultTimecode)
 	{
-		FApp::SetTimecodeAndFrameRate(FTimecode(), FFrameRate());
+#if PLATFORM_DESKTOP
+		FApp::SetCurrentFrameTime(FQualifiedFrameTime(USystemTimeTimecodeProvider::GenerateTimecodeFromSystemTime(GenerateDefaultTimecodeFrameRate), GenerateDefaultTimecodeFrameRate));
+#else //PLATFORM_DESKTOP
+		//If a user wish to have an accurate TC value on console, he should set an accurate TC provider in his project settings.
+		FApp::SetCurrentFrameTime(FQualifiedFrameTime(USystemTimeTimecodeProvider::GenerateTimecodeFromHighPerformanceClock(GenerateDefaultTimecodeFrameRate), GenerateDefaultTimecodeFrameRate));
+#endif
 	}
 }
 
diff --git a/Engine/Source/Runtime/LiveLinkInterface/Public/ILiveLinkClient.h b/Engine/Source/Runtime/LiveLinkInterface/Public/ILiveLinkClient.h
--- a/Engine/Source/Runtime/LiveLinkInterface/Public/ILiveLinkClient.h
+++ b/Engine/Source/Runtime/LiveLinkInterface/Public/ILiveLinkClient.h
@@ -271,15 +271,18 @@ public:
 	 */
 	virtual bool EvaluateFrameAtWorldTime_AnyThread(FLiveLinkSubjectName SubjectName, double WorldTime, TSubclassOf<ULiveLinkRole> DesiredRole, FLiveLinkSubjectFrameData& OutFrame) = 0;
 
+	UE_DEPRECATED(4.25, "ILiveLinkClient::EvaluateFrameAtSceneTime_AnyThread is deprecated. Please use ILiveLinkClient::EvaluateFrameAtSceneTime_AnyThread with a QualifiedFrameTime instead!")
+	virtual bool EvaluateFrameAtSceneTime_AnyThread(FLiveLinkSubjectName SubjectName, const FTimecode& SceneTime, TSubclassOf<ULiveLinkRole> DesiredRole, FLiveLinkSubjectFrameData& OutFrame) { return false; }
+
 	/**
-	 * Evaluates a subject for a specific role at a Timecode.
+	 * Evaluates a subject for a specific role at a scene time.
 	 * The subject may go through a translator to get the desired role's frame data.
 	 * If it's a virtual subject EvaluateFrame_AnyThread will be used instead.
 	 * @return True if a frame data was calculated.
 	 * @note This value is not cached.
 	 * @see ULiveLinkSourceSettings
 	 */
-	virtual bool EvaluateFrameAtSceneTime_AnyThread(FLiveLinkSubjectName SubjectName, const FTimecode& SceneTime, TSubclassOf<ULiveLinkRole> DesiredRole, FLiveLinkSubjectFrameData& OutFrame) = 0;
+	virtual bool EvaluateFrameAtSceneTime_AnyThread(FLiveLinkSubjectName SubjectName, const FQualifiedFrameTime& SceneTime, TSubclassOf<ULiveLinkRole> DesiredRole, FLiveLinkSubjectFrameData& OutFrame) = 0;
 
 	/** Notify when the LiveLinkClient has ticked. */
 	virtual FSimpleMulticastDelegate& OnLiveLinkTicked() = 0;
diff --git a/Engine/Source/Runtime/TimeManagement/Public/TimeManagementBlueprintLibrary.h b/Engine/Source/Runtime/TimeManagement/Public/TimeManagementBlueprintLibrary.h
--- a/Engine/Source/Runtime/TimeManagement/Public/TimeManagementBlueprintLibrary.h
+++ b/Engine/Source/Runtime/TimeManagement/Public/TimeManagementBlueprintLibrary.h
@@ -80,10 +80,10 @@ public:
 
 public:
 	/** Get the current timecode of the engine. */
-	UFUNCTION(BlueprintPure, Category = "Utilities|Timecode Provider")
+	UFUNCTION(BlueprintPure, Category = "Utilities|Time Management")
 	static FTimecode GetTimecode();
 
 	/** Gets the current timecode frame rate. */
-	UFUNCTION(BlueprintPure, Category = "Utilities|Timecode Provider")
+	UFUNCTION(BlueprintPure, Category = "Utilities|Time Management")
 	static FFrameRate GetTimecodeFrameRate();
 };
\ No newline at end of file
diff --git a/Engine/Source/Runtime/TimeManagement/Public/TimeSynchronizationSource.h b/Engine/Source/Runtime/TimeManagement/Public/TimeSynchronizationSource.h
--- a/Engine/Source/Runtime/TimeManagement/Public/TimeSynchronizationSource.h
+++ b/Engine/Source/Runtime/TimeManagement/Public/TimeSynchronizationSource.h
@@ -133,8 +133,7 @@ public:
 	/** Convenience method to convert a FrameTime and FrameRate to a timecode value. */
 	FORCEINLINE static FTimecode ConvertFrameTimeToTimecode(const FFrameTime& FrameTime, const FFrameRate& FrameRate)
 	{
-		const bool bIsDropFrame = FTimecode::IsDropFormatTimecodeSupported(FrameRate);
-		return FTimecode::FromFrameNumber(FrameTime.GetFrame(), FrameRate, bIsDropFrame);
+		return FTimecode::FromFrameNumber(FrameTime.GetFrame(), FrameRate);
 	}
 
 	/**
