Remove woverloaded-virtual warning for clang compiler introduced with CL10730014

diff --git a/Engine/Plugins/Animation/LiveLink/Source/LiveLink/Public/LiveLinkClient.h b/Engine/Plugins/Animation/LiveLink/Source/LiveLink/Public/LiveLinkClient.h
--- a/Engine/Plugins/Animation/LiveLink/Source/LiveLink/Public/LiveLinkClient.h
+++ b/Engine/Plugins/Animation/LiveLink/Source/LiveLink/Public/LiveLinkClient.h
@@ -175,6 +175,7 @@ public:
 
 protected:
 	//~ Begin FLiveLinkClient_Base_DEPRECATED implementation
+	using FLiveLinkClient_Base_DEPRECATED::EvaluateFrameAtSceneTime_AnyThread;
 	virtual void AquireLock_Deprecation() override;
 	virtual void ReleaseLock_Deprecation() override;
 	virtual void ClearFrames_Deprecation(const FLiveLinkSubjectKey& SubjectKey) override;
