Removed unused function from LiveLinkComponentController

diff --git a/Engine/Plugins/Animation/LiveLink/Source/LiveLinkComponents/Public/LiveLinkComponentController.h b/Engine/Plugins/Animation/LiveLink/Source/LiveLinkComponents/Public/LiveLinkComponentController.h
--- a/Engine/Plugins/Animation/LiveLink/Source/LiveLinkComponents/Public/LiveLinkComponentController.h
+++ b/Engine/Plugins/Animation/LiveLink/Source/LiveLinkComponents/Public/LiveLinkComponentController.h
@@ -50,12 +50,9 @@ public:
 	/** Creates an instance of the desired controller class for a specified Role class */
 	void SetControllerClassForRole(TSubclassOf<ULiveLinkRole> RoleClass, TSubclassOf<ULiveLinkControllerBase> DesiredControllerClass);
 
-	/** Creates an instance of the desired controller class for a specified Role class */
+	/** Return Representation of Subject that is used in the controller */
 	FLiveLinkSubjectRepresentation GetSubjectRepresentation() const { return SubjectRepresentation; }
 
-	/** Sets SubjectRepresentation and if required, will update the Controllers map associated to the Role */
-	void SetSubjectRepresentation(const FLiveLinkSubjectRepresentation& InSubjectRepresentation);
-	
 	/** Returns true if ControllerMap needs to be updated for the current Role. Useful for customization or C++ modification to the Role */
 	bool IsControllerMapOutdated() const;
 	
