Initialize the live link source setting when connecting via BP.

diff --git a/Engine/Plugins/Animation/LiveLink/Source/LiveLink/Private/LiveLinkMessageBusFinder.cpp b/Engine/Plugins/Animation/LiveLink/Source/LiveLink/Private/LiveLinkMessageBusFinder.cpp
--- a/Engine/Plugins/Animation/LiveLink/Source/LiveLink/Private/LiveLinkMessageBusFinder.cpp
+++ b/Engine/Plugins/Animation/LiveLink/Source/LiveLink/Private/LiveLinkMessageBusFinder.cpp
@@ -1,12 +1,13 @@
-﻿// Copyright 1998-2019 Epic Games, Inc. All Rights Reserved.
+﻿// Copyright 1998-2020 Epic Games, Inc. All Rights Reserved.
 
 #include "LiveLinkMessageBusFinder.h"
 
 #include "Engine/Engine.h"
 #include "Features/IModularFeatures.h"
-#include "ILiveLinkClient.h"
+#include "LiveLinkClient.h"
 #include "LiveLinkMessages.h"
 #include "LiveLinkMessageBusSource.h"
+#include "LiveLinkMessageBusSourceFactory.h"
 #include "MessageEndpointBuilder.h"
 
 
@@ -89,9 +90,17 @@ void ULiveLinkMessageBusFinder::ConnectToProvider(UPARAM(ref) FProviderPollResul
 
 	if (ModularFeatures.IsModularFeatureAvailable(ILiveLinkClient::ModularFeatureName))
 	{
-		ILiveLinkClient* LiveLinkClient = &ModularFeatures.GetModularFeature<ILiveLinkClient>(ILiveLinkClient::ModularFeatureName);
+		FLiveLinkClient* LiveLinkClient = &ModularFeatures.GetModularFeature<FLiveLinkClient>(ILiveLinkClient::ModularFeatureName);
 		TSharedPtr<FLiveLinkMessageBusSource> NewSource = MakeShared<FLiveLinkMessageBusSource>(FText::FromString(Provider.Name), FText::FromString(Provider.MachineName), Provider.Address, Provider.MachineTimeOffset);
-		LiveLinkClient->AddSource(NewSource);
+		FGuid NewSourceGuid = LiveLinkClient->AddSource(NewSource);
+		if (NewSourceGuid.IsValid())
+		{
+			if (ULiveLinkSourceSettings* Settings = LiveLinkClient->GetSourceSettings(NewSourceGuid))
+			{
+				Settings->ConnectionString = ULiveLinkMessageBusSourceFactory::CreateConnectionString(Provider);
+				Settings->Factory = ULiveLinkMessageBusSourceFactory::StaticClass();
+			}
+		}
 		SourceHandle.SetSourcePointer(NewSource);
 	}
 	else
diff --git a/Engine/Plugins/Animation/LiveLink/Source/LiveLink/Private/LiveLinkMessageBusSourceFactory.cpp b/Engine/Plugins/Animation/LiveLink/Source/LiveLink/Private/LiveLinkMessageBusSourceFactory.cpp
--- a/Engine/Plugins/Animation/LiveLink/Source/LiveLink/Private/LiveLinkMessageBusSourceFactory.cpp
+++ b/Engine/Plugins/Animation/LiveLink/Source/LiveLink/Private/LiveLinkMessageBusSourceFactory.cpp
@@ -1,4 +1,4 @@
-// Copyright 1998-2019 Epic Games, Inc. All Rights Reserved.
+// Copyright 1998-2020 Epic Games, Inc. All Rights Reserved.
 #include "LiveLinkMessageBusSourceFactory.h"
 
 #include "ILiveLinkClient.h"
@@ -38,6 +38,12 @@ TSharedPtr<ILiveLinkSource> ULiveLinkMessageBusSourceFactory::CreateSource(const
 	return MakeShared<FLiveLinkMessageBusSource>(FText::FromString(Name), FText::GetEmpty(), FMessageAddress(), TimeOffset);
 }
 
+
+FString ULiveLinkMessageBusSourceFactory::CreateConnectionString(const FProviderPollResult& Result)
+{
+	return FString::Printf(TEXT("Name=\"%s\""), *Result.Name);
+}
+
 void ULiveLinkMessageBusSourceFactory::OnSourceSelected(FProviderPollResultPtr SelectedSource, FOnLiveLinkSourceCreated InOnLiveLinkSourceCreated) const
 {
 	if (SelectedSource.IsValid())
@@ -67,7 +73,7 @@ void ULiveLinkMessageBusSourceFactory::OnSourceSelected(FProviderPollResultPtr S
 #endif
 
 		TSharedPtr<FLiveLinkMessageBusSource> SharedPtr = MakeShared<FLiveLinkMessageBusSource>(FText::FromString(SelectedSource->Name), FText::FromString(SelectedSource->MachineName), SelectedSource->Address, SelectedSource->MachineTimeOffset);
-		FString ConnectionString = FString::Printf(TEXT("Name=\"%s\""), *SelectedSource->Name);
+		FString ConnectionString = CreateConnectionString(*SelectedSource.Get());
 		InOnLiveLinkSourceCreated.ExecuteIfBound(SharedPtr, MoveTemp(ConnectionString));
 	}
 }
diff --git a/Engine/Plugins/Animation/LiveLink/Source/LiveLink/Private/LiveLinkMessageBusSourceFactory.h b/Engine/Plugins/Animation/LiveLink/Source/LiveLink/Private/LiveLinkMessageBusSourceFactory.h
--- a/Engine/Plugins/Animation/LiveLink/Source/LiveLink/Private/LiveLinkMessageBusSourceFactory.h
+++ b/Engine/Plugins/Animation/LiveLink/Source/LiveLink/Private/LiveLinkMessageBusSourceFactory.h
@@ -1,4 +1,4 @@
-// Copyright 1998-2019 Epic Games, Inc. All Rights Reserved.
+// Copyright 1998-2020 Epic Games, Inc. All Rights Reserved.
 
 #pragma once
 
@@ -20,6 +20,8 @@ public:
 	virtual TSharedPtr<SWidget> BuildCreationPanel(FOnLiveLinkSourceCreated OnLiveLinkSourceCreated) const override;
 	virtual TSharedPtr<ILiveLinkSource> CreateSource(const FString& ConnectionString) const override;
 
+	static FString CreateConnectionString(const FProviderPollResult& Result);
+
 private:
 	void OnSourceSelected(FProviderPollResultPtr SelectedSource, FOnLiveLinkSourceCreated OnLiveLinkSourceCreated) const;
 };
