SubFrame delay for TC Provider and LiveLink

diff --git a/Engine/Plugins/Animation/LiveLink/Source/LiveLink/Private/LiveLinkSubject.cpp b/Engine/Plugins/Animation/LiveLink/Source/LiveLink/Private/LiveLinkSubject.cpp
--- a/Engine/Plugins/Animation/LiveLink/Source/LiveLink/Private/LiveLinkSubject.cpp
+++ b/Engine/Plugins/Animation/LiveLink/Source/LiveLink/Private/LiveLinkSubject.cpp
@@ -1,4 +1,4 @@
-// Copyright 1998-2019 Epic Games, Inc. All Rights Reserved.
+// Copyright 1998-2020 Epic Games, Inc. All Rights Reserved.
 
 #include "LiveLinkSubject.h"
 
@@ -69,7 +69,7 @@ void FLiveLinkSubject::Update()
 			int32 FrameIndex = 0;
 			for (const FLiveLinkFrameDataStruct& SourceFrameData : FrameData)
 			{
-				FFrameTime UsedFrameTime = CurrentFrameTimeInFrameSpace - CachedSettings.BufferSettings.TimecodeFrameOffset - CachedSettings.BufferSettings.ValidTimecodeFrame;
+				FFrameTime UsedFrameTime = CurrentFrameTimeInFrameSpace - FFrameTime::FromDecimal(CachedSettings.BufferSettings.TimecodeFrameOffset) - CachedSettings.BufferSettings.ValidTimecodeFrame;
 				FFrameTime FrameTime = SourceFrameData.GetBaseData()->MetaData.SceneTime.Time;
 				if (FrameTime > UsedFrameTime)
 				{
@@ -382,7 +382,7 @@ int32 FLiveLinkSubject::FindNewFrame_SceneTime(const FQualifiedFrameTime& Qualif
 	{
 		const FQualifiedFrameTime CurrentSyncTime = FApp::GetCurrentFrameTime().GetValue();
 		const FFrameTime CurrentFrameTimeInFrameSpace = CurrentSyncTime.ConvertTo(CachedSettings.BufferSettings.TimecodeFrameRate);
-		const FFrameTime CurrentOffsetFrameTime = CurrentFrameTimeInFrameSpace - CachedSettings.BufferSettings.TimecodeFrameOffset - CachedSettings.BufferSettings.ValidTimecodeFrame;
+		const FFrameTime CurrentOffsetFrameTime = CurrentFrameTimeInFrameSpace - FFrameTime::FromDecimal(CachedSettings.BufferSettings.TimecodeFrameOffset) - CachedSettings.BufferSettings.ValidTimecodeFrame;
 		if (QualifiedFrameTime.Time.AsDecimal() < CurrentOffsetFrameTime.AsDecimal())
 		{
 			static const FName NAME_InvalidTC = "LiveLinkSubject_InvalidTC";
@@ -647,7 +647,7 @@ bool FLiveLinkSubject::GetFrameAtSceneTime_Closest(const FQualifiedFrameTime& In
 
 	bool bBuiltFrame = false;
 
-	const FFrameTime FrameOffset = FQualifiedFrameTime(CachedSettings.BufferSettings.TimecodeFrameOffset, CachedSettings.BufferSettings.TimecodeFrameRate).ConvertTo(InTimeInEngineFrameRate.Rate);
+	const FFrameTime FrameOffset = FQualifiedFrameTime(FFrameTime::FromDecimal(CachedSettings.BufferSettings.TimecodeFrameOffset), CachedSettings.BufferSettings.TimecodeFrameRate).ConvertTo(InTimeInEngineFrameRate.Rate);
 	const FFrameTime ReadTime = InTimeInEngineFrameRate.Time - FrameOffset;
 	const FQualifiedFrameTime LookupQFrameTime = FQualifiedFrameTime(ReadTime, InTimeInEngineFrameRate.Rate);
 	const double TimeInSeconds = LookupQFrameTime.AsSeconds();
@@ -701,7 +701,7 @@ bool FLiveLinkSubject::GetFrameAtSceneTime_Interpolated(const FQualifiedFrameTim
 {
 	check(FrameData.Num() != 0);
 
-	const FFrameTime FrameOffset = FQualifiedFrameTime(CachedSettings.BufferSettings.TimecodeFrameOffset, CachedSettings.BufferSettings.TimecodeFrameRate).ConvertTo(InTimeInEngineFrameRate.Rate);
+	const FFrameTime FrameOffset = FQualifiedFrameTime(FFrameTime::FromDecimal(CachedSettings.BufferSettings.TimecodeFrameOffset), CachedSettings.BufferSettings.TimecodeFrameRate).ConvertTo(InTimeInEngineFrameRate.Rate);
 	const FFrameTime ReadTime = InTimeInEngineFrameRate.Time - FrameOffset;
 	const FQualifiedFrameTime LookupQFrameTime = FQualifiedFrameTime(ReadTime, InTimeInEngineFrameRate.Rate);
 	FrameInterpolationProcessor->Interpolate(LookupQFrameTime, StaticData, FrameData, OutFrame);
diff --git a/Engine/Source/Runtime/Core/Public/Misc/FrameTime.h b/Engine/Source/Runtime/Core/Public/Misc/FrameTime.h
--- a/Engine/Source/Runtime/Core/Public/Misc/FrameTime.h
+++ b/Engine/Source/Runtime/Core/Public/Misc/FrameTime.h
@@ -1,4 +1,4 @@
-// Copyright 1998-2019 Epic Games, Inc. All Rights Reserved.
+// Copyright 1998-2020 Epic Games, Inc. All Rights Reserved.
 
 #pragma once
 
@@ -100,6 +100,7 @@ private:
 	friend bool operator<=(FFrameTime A, FFrameTime B);
 
 	friend FFrameTime& operator+=(FFrameTime& LHS, FFrameTime RHS);
+	friend FFrameTime& operator-=(FFrameTime& LHS, FFrameTime RHS);
 	friend FFrameTime  operator+(FFrameTime A, FFrameTime B);
 	friend FFrameTime  operator-(FFrameTime A, FFrameTime B);
 	friend FFrameTime  operator%(FFrameTime A, FFrameTime B);
@@ -209,6 +210,19 @@ FORCEINLINE_DEBUGGABLE FFrameTime operator+(FFrameTime A, FFrameTime B)
 }
 
 
+FORCEINLINE_DEBUGGABLE FFrameTime& operator-=(FFrameTime& LHS, FFrameTime RHS)
+{
+	// Ensure SubFrame is always between 0 and 1
+	// Note that the difference between frame -1.5 and 1.5 is 2, not 3, since sub frame positions are always positive
+	const float        NewSubFrame     = LHS.SubFrame - RHS.SubFrame;
+	const int          FlooredSubFrame = FMath::FloorToInt(NewSubFrame);
+	LHS.FrameNumber  = LHS.FrameNumber - RHS.FrameNumber + FFrameNumber(FlooredSubFrame);
+	LHS.SubFrame = NewSubFrame - FlooredSubFrame;
+
+	return LHS;
+}
+
+
 FORCEINLINE_DEBUGGABLE FFrameTime operator-(FFrameTime A, FFrameTime B)
 {
 	// Ensure SubFrame is always between 0 and 1
diff --git a/Engine/Source/Runtime/Engine/Classes/Engine/TimecodeProvider.h b/Engine/Source/Runtime/Engine/Classes/Engine/TimecodeProvider.h
--- a/Engine/Source/Runtime/Engine/Classes/Engine/TimecodeProvider.h
+++ b/Engine/Source/Runtime/Engine/Classes/Engine/TimecodeProvider.h
@@ -1,4 +1,4 @@
-// Copyright 1998-2019 Epic Games, Inc. All Rights Reserved.
+// Copyright 1998-2020 Epic Games, Inc. All Rights Reserved.
 
 #pragma once
 
@@ -47,7 +47,7 @@ public:
 	 * @see GetDelayedQualifiedFrameTime, GetDelayedTimecode
 	 */
 	UPROPERTY(EditAnywhere, AdvancedDisplay, Category = "Settings")
-	int32 FrameDelay = 0;
+	float FrameDelay = 0.f;
 	
 	/**
 	 * Return the frame number and the frame rate of the frame number at that moment. It may not be in sync with the current frame.
diff --git a/Engine/Source/Runtime/Engine/Private/Engine/SystemTimeTimecodeProvider.cpp b/Engine/Source/Runtime/Engine/Private/Engine/SystemTimeTimecodeProvider.cpp
--- a/Engine/Source/Runtime/Engine/Private/Engine/SystemTimeTimecodeProvider.cpp
+++ b/Engine/Source/Runtime/Engine/Private/Engine/SystemTimeTimecodeProvider.cpp
@@ -1,4 +1,4 @@
-// Copyright 1998-2019 Epic Games, Inc. All Rights Reserved.
+// Copyright 1998-2020 Epic Games, Inc. All Rights Reserved.
 
 #include "Engine/SystemTimeTimecodeProvider.h"
 
@@ -21,7 +21,7 @@ namespace
 
 USystemTimeTimecodeProvider::USystemTimeTimecodeProvider()
 	: FrameRate(60, 1)
-	, bGenerateFullFrame(false)
+	, bGenerateFullFrame(true)
 	, State(ETimecodeProviderSynchronizationState::Closed)
 {
 }
@@ -60,7 +60,6 @@ FTimecode USystemTimeTimecodeProvider::GenerateTimecodeFromHighPerformanceClock(
 
 FQualifiedFrameTime USystemTimeTimecodeProvider::GetQualifiedFrameTime() const
 {
-	// Generate a TC to prevent subframe
 	return bGenerateFullFrame ? 
 		FQualifiedFrameTime(GenerateTimecodeFromSystemTime(FrameRate), FrameRate)
 		: FQualifiedFrameTime(GenerateFrameTimeFromSystemTime(FrameRate), FrameRate);
diff --git a/Engine/Source/Runtime/Engine/Private/Engine/TimecodeProvider.cpp b/Engine/Source/Runtime/Engine/Private/Engine/TimecodeProvider.cpp
--- a/Engine/Source/Runtime/Engine/Private/Engine/TimecodeProvider.cpp
+++ b/Engine/Source/Runtime/Engine/Private/Engine/TimecodeProvider.cpp
@@ -1,4 +1,4 @@
-// Copyright 1998-2019 Epic Games, Inc. All Rights Reserved.
+// Copyright 1998-2020 Epic Games, Inc. All Rights Reserved.
 
 #include "Engine/TimecodeProvider.h"
 
@@ -6,7 +6,7 @@
 FQualifiedFrameTime UTimecodeProvider::GetDelayedQualifiedFrameTime() const
 {
 	FQualifiedFrameTime NewFrameTime = GetQualifiedFrameTime();
-	NewFrameTime.Time.FrameNumber -= FrameDelay;
+	NewFrameTime.Time -= FFrameTime::FromDecimal(FrameDelay);
 	return NewFrameTime;
 }
 
@@ -23,3 +23,16 @@ FTimecode UTimecodeProvider::GetDelayedTimecode() const
 	const FQualifiedFrameTime NewFrameTime = GetDelayedQualifiedFrameTime();
 	return FTimecode::FromFrameNumber(NewFrameTime.Time.GetFrame(), NewFrameTime.Rate);
 }
+
+#if WITH_EDITOR
+void UTimecodeProvider::PostEditChangeProperty(struct FPropertyChangedEvent& PropertyChangedEvent)
+{
+	if (PropertyChangedEvent.GetPropertyName() == GET_MEMBER_NAME_CHECKED(UTimecodeProvider, FrameDelay))
+	{
+		if (!SupportSubFrameDelay())
+		{
+			FrameDelay = FMath::FloorToInt(FrameDelay);
+		}
+	}
+}
+#endif
\ No newline at end of file
diff --git a/Engine/Source/Runtime/LiveLinkInterface/Public/LiveLinkSourceSettings.h b/Engine/Source/Runtime/LiveLinkInterface/Public/LiveLinkSourceSettings.h
--- a/Engine/Source/Runtime/LiveLinkInterface/Public/LiveLinkSourceSettings.h
+++ b/Engine/Source/Runtime/LiveLinkInterface/Public/LiveLinkSourceSettings.h
@@ -1,4 +1,4 @@
-// Copyright 1998-2019 Epic Games, Inc. All Rights Reserved.
+// Copyright 1998-2020 Epic Games, Inc. All Rights Reserved.
 
 #pragma once
 
@@ -69,7 +69,7 @@ struct FLiveLinkSourceBufferManagementSettings
 
 	/** When evaluating with timecode: how far back from current timecode should we read the buffer (in TimecodeFrameRate). */
 	UPROPERTY(EditAnywhere, Category = "Settings")
-	int32 TimecodeFrameOffset = 0;
+	float TimecodeFrameOffset = 0.f;
 
 	/** When evaluating with latest: how far back from latest frame should we read the buffer */
 	UPROPERTY(EditAnywhere, Category = "Settings")
