Add a picker for LiveLinkSubjectKey.

diff --git a/Engine/Plugins/Animation/LiveLink/Source/LiveLink/Private/LiveLinkClient.cpp b/Engine/Plugins/Animation/LiveLink/Source/LiveLink/Private/LiveLinkClient.cpp
--- a/Engine/Plugins/Animation/LiveLink/Source/LiveLink/Private/LiveLinkClient.cpp
+++ b/Engine/Plugins/Animation/LiveLink/Source/LiveLink/Private/LiveLinkClient.cpp
@@ -403,9 +403,10 @@ FLiveLinkSourcePreset FLiveLinkClient::GetSourcePreset(FGuid InSourceGuid, UObje
 	FLiveLinkSourcePreset SourcePreset;
 	if (FLiveLinkCollectionSourceItem* SourceItem = Collection->FindSource(InSourceGuid))
 	{
-		if (SourceItem->Guid != FLiveLinkSourceCollection::DefaultVirtualSubjectGuid && SourceItem->Setting)
+		if (SourceItem->Guid != FLiveLinkSourceCollection::DefaultVirtualSubjectGuid && SourceItem->Setting && SourceItem->Source)
 		{
 			SourcePreset.Guid = SourceItem->Guid;
+			SourcePreset.SourceType = SourceItem->Source->CanBeDisplayedInUI() ? SourceItem->Source->GetSourceType() : FText::GetEmpty();
 			SourcePreset.Settings = DuplicateObject<ULiveLinkSourceSettings>(SourceItem->Setting, DuplicatedObjectOuter);
 		}
 	}
diff --git a/Engine/Plugins/Animation/LiveLink/Source/LiveLinkEditor/Private/LiveLinkEditorModule.cpp b/Engine/Plugins/Animation/LiveLink/Source/LiveLinkEditor/Private/LiveLinkEditorModule.cpp
--- a/Engine/Plugins/Animation/LiveLink/Source/LiveLinkEditor/Private/LiveLinkEditorModule.cpp
+++ b/Engine/Plugins/Animation/LiveLink/Source/LiveLinkEditor/Private/LiveLinkEditorModule.cpp
@@ -30,6 +30,7 @@
 #include "LiveLinkGraphPanelPinFactory.h"
 #include "LiveLinkSettings.h"
 #include "LiveLinkSourceSettingsDetailCustomization.h"
+#include "LiveLinkSubjectKeyDetailCustomization.h"
 #include "LiveLinkSubjectNameDetailCustomization.h"
 #include "LiveLinkSubjectRepresentationDetailCustomization.h"
 #include "LiveLinkTypes.h"
@@ -234,6 +235,7 @@ private:
 	{
 		FPropertyEditorModule& PropertyEditorModule = FModuleManager::Get().LoadModuleChecked<FPropertyEditorModule>("PropertyEditor");
 		PropertyEditorModule.RegisterCustomClassLayout(ULiveLinkVirtualSubject::StaticClass()->GetFName(), FOnGetDetailCustomizationInstance::CreateStatic(&FLiveLinkVirtualSubjectDetailCustomization::MakeInstance));
+		PropertyEditorModule.RegisterCustomPropertyTypeLayout(FLiveLinkSubjectKey::StaticStruct()->GetFName(), FOnGetPropertyTypeCustomizationInstance::CreateStatic(&FLiveLinkSubjectKeyDetailCustomization::MakeInstance));
 		PropertyEditorModule.RegisterCustomPropertyTypeLayout(FLiveLinkSubjectRepresentation::StaticStruct()->GetFName(), FOnGetPropertyTypeCustomizationInstance::CreateStatic(&FLiveLinkSubjectRepresentationDetailCustomization::MakeInstance));
 		PropertyEditorModule.RegisterCustomPropertyTypeLayout(FLiveLinkSubjectName::StaticStruct()->GetFName(), FOnGetPropertyTypeCustomizationInstance::CreateStatic(&FLiveLinkSubjectNameDetailCustomization::MakeInstance));
 		PropertyEditorModule.RegisterCustomClassLayout(ULiveLinkSourceSettings::StaticClass()->GetFName(), FOnGetDetailCustomizationInstance::CreateStatic(&FLiveLinkSourceSettingsDetailCustomization::MakeInstance));
@@ -255,6 +257,7 @@ private:
 				PropertyEditorModule->UnregisterCustomClassLayout(ULiveLinkSourceSettings::StaticClass()->GetFName());
 				PropertyEditorModule->UnregisterCustomPropertyTypeLayout(FLiveLinkSubjectName::StaticStruct()->GetFName());
 				PropertyEditorModule->UnregisterCustomPropertyTypeLayout(FLiveLinkSubjectRepresentation::StaticStruct()->GetFName());
+				PropertyEditorModule->UnregisterCustomPropertyTypeLayout(FLiveLinkSubjectKey::StaticStruct()->GetFName());
 				PropertyEditorModule->UnregisterCustomClassLayout(ULiveLinkVirtualSubject::StaticClass()->GetFName());
 			}
 		}
diff --git a/Engine/Plugins/Animation/LiveLink/Source/LiveLinkEditor/Private/LiveLinkSubjectKeyDetailCustomization.cpp b/Engine/Plugins/Animation/LiveLink/Source/LiveLinkEditor/Private/LiveLinkSubjectKeyDetailCustomization.cpp
new file mode 100644
--- /dev/null
+++ b/Engine/Plugins/Animation/LiveLink/Source/LiveLinkEditor/Private/LiveLinkSubjectKeyDetailCustomization.cpp
@@ -0,0 +1,128 @@
+// Copyright 1998-2019 Epic Games, Inc. All Rights Reserved.
+
+#include "LiveLinkSubjectKeyDetailCustomization.h"
+
+#include "DetailLayoutBuilder.h"
+#include "DetailWidgetRow.h"
+#include "Editor.h"
+#include "IDetailChildrenBuilder.h"
+#include "IPropertyUtilities.h"
+#include "IPropertyTypeCustomization.h"
+#include "PropertyHandle.h"
+#include "SLiveLinkSubjectRepresentationPicker.h"
+
+
+#define LOCTEXT_NAMESPACE "LiveLinkSubjectKeyDetailCustomization"
+
+
+TSharedRef<IPropertyTypeCustomization> FLiveLinkSubjectKeyDetailCustomization::MakeInstance()
+{
+	return MakeShareable(new FLiveLinkSubjectKeyDetailCustomization);
+}
+
+
+void FLiveLinkSubjectKeyDetailCustomization::CustomizeHeader(TSharedRef<IPropertyHandle> InPropertyHandle, FDetailWidgetRow& HeaderRow, IPropertyTypeCustomizationUtils& CustomizationUtils)
+{
+	StructPropertyHandle = InPropertyHandle;
+	TSharedPtr<IPropertyUtilities> PropertyUtils = CustomizationUtils.GetPropertyUtilities();
+
+	check(CastFieldChecked<FStructProperty>(StructPropertyHandle->GetProperty())->Struct == FLiveLinkSubjectKey::StaticStruct());
+
+	HeaderRow.NameContent()
+	[
+		StructPropertyHandle->CreatePropertyNameWidget()
+	]
+	.ValueContent()
+	.MinDesiredWidth(125.f)
+	.MaxDesiredWidth(400.f)
+	[
+		SNew(SLiveLinkSubjectRepresentationPicker)
+		.ShowRole(false)
+		.ShowSource(true)
+		.Font(CustomizationUtils.GetRegularFont())
+		.HasMultipleValues(this, &FLiveLinkSubjectKeyDetailCustomization::HasMultipleValues)
+		.Value(this, &FLiveLinkSubjectKeyDetailCustomization::GetValue)
+		.OnValueChanged(this, &FLiveLinkSubjectKeyDetailCustomization::SetValue)
+	].IsEnabled(MakeAttributeLambda([=] { return !InPropertyHandle->IsEditConst() && PropertyUtils->IsPropertyEditingEnabled(); }));
+}
+
+void FLiveLinkSubjectKeyDetailCustomization::CustomizeChildren(TSharedRef<IPropertyHandle> PropertyHandle, IDetailChildrenBuilder& ChildBuilder, IPropertyTypeCustomizationUtils& CustomizationUtils)
+{
+	TSharedPtr<IPropertyUtilities> PropertyUtils = CustomizationUtils.GetPropertyUtilities();
+	uint32 NumberOfChild;
+	if (PropertyHandle->GetNumChildren(NumberOfChild) == FPropertyAccess::Success)
+	{
+		for (uint32 Index = 0; Index < NumberOfChild; ++Index)
+		{
+			TSharedRef<IPropertyHandle> ChildPropertyHandle = PropertyHandle->GetChildHandle(Index).ToSharedRef();
+			ChildBuilder.AddProperty(ChildPropertyHandle)
+				.ShowPropertyButtons(true)
+				.IsEnabled(MakeAttributeLambda([=] { return !PropertyHandle->IsEditConst() && PropertyUtils->IsPropertyEditingEnabled(); }));
+		}
+	}
+}
+
+SLiveLinkSubjectRepresentationPicker::FLiveLinkSourceSubjectRole FLiveLinkSubjectKeyDetailCustomization::GetValue() const
+{
+	TArray<const void*> RawData;
+	StructPropertyHandle->AccessRawData(RawData);
+
+	for (const void* RawPtr : RawData)
+	{
+		if (RawPtr)
+		{
+			return SLiveLinkSubjectRepresentationPicker::FLiveLinkSourceSubjectRole (*reinterpret_cast<const FLiveLinkSubjectKey*>(RawPtr));
+		}
+	}
+
+	return SLiveLinkSubjectRepresentationPicker::FLiveLinkSourceSubjectRole();
+}
+
+void FLiveLinkSubjectKeyDetailCustomization::SetValue(SLiveLinkSubjectRepresentationPicker::FLiveLinkSourceSubjectRole NewValue)
+{
+	FStructProperty* StructProperty = CastFieldChecked<FStructProperty>(StructPropertyHandle->GetProperty());
+
+	TArray<void*> RawData;
+	StructPropertyHandle->AccessRawData(RawData);
+	FLiveLinkSubjectKey* PreviousValue = reinterpret_cast<FLiveLinkSubjectKey*>(RawData[0]);
+	FLiveLinkSubjectKey NewSubjectKey = NewValue.ToSubjectKey();
+
+	FString TextValue;
+	StructProperty->Struct->ExportText(TextValue, &NewSubjectKey, PreviousValue, nullptr, EPropertyPortFlags::PPF_None, nullptr);
+	ensure(StructPropertyHandle->SetValueFromFormattedString(TextValue, EPropertyValueSetFlags::DefaultFlags) == FPropertyAccess::Result::Success);
+}
+
+bool FLiveLinkSubjectKeyDetailCustomization::HasMultipleValues() const
+{
+	TArray<const void*> RawData;
+	StructPropertyHandle->AccessRawData(RawData);
+
+	TOptional<FLiveLinkSubjectKey> CompareAgainst;
+	for (const void* RawPtr : RawData)
+	{
+		if (RawPtr == nullptr)
+		{
+			if (CompareAgainst.IsSet())
+			{
+				return false;
+			}
+		}
+		else
+		{
+			FLiveLinkSubjectKey ThisValue = *reinterpret_cast<const FLiveLinkSubjectKey*>(RawPtr);
+
+			if (!CompareAgainst.IsSet())
+			{
+				CompareAgainst = ThisValue;
+			}
+			else if (!(ThisValue == CompareAgainst.GetValue()))
+			{
+				return true;
+			}
+		}
+	}
+
+	return false;
+}
+
+#undef LOCTEXT_NAMESPACE
\ No newline at end of file
diff --git a/Engine/Plugins/Animation/LiveLink/Source/LiveLinkEditor/Private/LiveLinkSubjectKeyDetailCustomization.h b/Engine/Plugins/Animation/LiveLink/Source/LiveLinkEditor/Private/LiveLinkSubjectKeyDetailCustomization.h
new file mode 100644
--- /dev/null
+++ b/Engine/Plugins/Animation/LiveLink/Source/LiveLinkEditor/Private/LiveLinkSubjectKeyDetailCustomization.h
@@ -0,0 +1,28 @@
+// Copyright 1998-2019 Epic Games, Inc. All Rights Reserved.
+
+#pragma once
+
+#include "CoreMinimal.h"
+#include "IPropertyTypeCustomization.h"
+#include "LiveLinkRole.h"
+#include "LiveLinkTypes.h"
+#include "SLiveLinkSubjectRepresentationPicker.h"
+
+class FLiveLinkSubjectKeyDetailCustomization : public IPropertyTypeCustomization
+{
+public:
+	/** Makes a new instance of this detail layout class for a specific detail view requesting it. */
+	static TSharedRef<IPropertyTypeCustomization> MakeInstance();
+
+public:
+	// IDetailCustomization interface
+	virtual void CustomizeHeader(TSharedRef<IPropertyHandle> PropertyHandle, FDetailWidgetRow& HeaderRow, IPropertyTypeCustomizationUtils& CustomizationUtils) override;
+	virtual void CustomizeChildren(TSharedRef<IPropertyHandle> PropertyHandle, IDetailChildrenBuilder& ChildBuilder, IPropertyTypeCustomizationUtils& CustomizationUtils) override;
+
+private:
+	SLiveLinkSubjectRepresentationPicker::FLiveLinkSourceSubjectRole GetValue() const;
+	void SetValue(SLiveLinkSubjectRepresentationPicker::FLiveLinkSourceSubjectRole NewValue);
+	bool HasMultipleValues() const;
+
+	TSharedPtr<IPropertyHandle> StructPropertyHandle;
+};
diff --git a/Engine/Plugins/Animation/LiveLink/Source/LiveLinkEditor/Private/LiveLinkSubjectNameDetailCustomization.cpp b/Engine/Plugins/Animation/LiveLink/Source/LiveLinkEditor/Private/LiveLinkSubjectNameDetailCustomization.cpp
--- a/Engine/Plugins/Animation/LiveLink/Source/LiveLinkEditor/Private/LiveLinkSubjectNameDetailCustomization.cpp
+++ b/Engine/Plugins/Animation/LiveLink/Source/LiveLinkEditor/Private/LiveLinkSubjectNameDetailCustomization.cpp
@@ -44,7 +44,7 @@ void FLiveLinkSubjectNameDetailCustomization::CustomizeHeader(TSharedRef<IProper
 }
 
 
-FLiveLinkSubjectRepresentation FLiveLinkSubjectNameDetailCustomization::GetValue() const
+SLiveLinkSubjectRepresentationPicker::FLiveLinkSourceSubjectRole FLiveLinkSubjectNameDetailCustomization::GetValue() const
 {
 	TArray<const void*> RawData;
 	StructPropertyHandle->AccessRawData(RawData);
@@ -55,14 +55,14 @@ FLiveLinkSubjectRepresentation FLiveLinkSubjectNameDetailCustomization::GetValue
 		{
 			FLiveLinkSubjectRepresentation Representation;
 			Representation.Subject = *reinterpret_cast<const FLiveLinkSubjectName *>(RawPtr);
-			return Representation;
+			return SLiveLinkSubjectRepresentationPicker::FLiveLinkSourceSubjectRole(Representation);
 		}
 	}
 
-	return FLiveLinkSubjectRepresentation();
+	return SLiveLinkSubjectRepresentationPicker::FLiveLinkSourceSubjectRole();
 }
 
-void FLiveLinkSubjectNameDetailCustomization::SetValue(FLiveLinkSubjectRepresentation NewValue)
+void FLiveLinkSubjectNameDetailCustomization::SetValue(SLiveLinkSubjectRepresentationPicker::FLiveLinkSourceSubjectRole NewValue)
 {
 	FStructProperty* StructProperty = CastFieldChecked<FStructProperty>(StructPropertyHandle->GetProperty());
 
diff --git a/Engine/Plugins/Animation/LiveLink/Source/LiveLinkEditor/Private/LiveLinkSubjectNameDetailCustomization.h b/Engine/Plugins/Animation/LiveLink/Source/LiveLinkEditor/Private/LiveLinkSubjectNameDetailCustomization.h
--- a/Engine/Plugins/Animation/LiveLink/Source/LiveLinkEditor/Private/LiveLinkSubjectNameDetailCustomization.h
+++ b/Engine/Plugins/Animation/LiveLink/Source/LiveLinkEditor/Private/LiveLinkSubjectNameDetailCustomization.h
@@ -6,6 +6,7 @@
 #include "LiveLinkRole.h"
 #include "LiveLinkTypes.h"
 #include "IPropertyTypeCustomization.h"
+#include "SLiveLinkSubjectRepresentationPicker.h"
 
 class FLiveLinkSubjectNameDetailCustomization : public IPropertyTypeCustomization
 {
@@ -19,8 +20,8 @@ public:
 	virtual void CustomizeChildren(TSharedRef<IPropertyHandle> PropertyHandle, IDetailChildrenBuilder& ChildBuilder, IPropertyTypeCustomizationUtils& CustomizationUtils) override {}
 
 private:
-	FLiveLinkSubjectRepresentation GetValue() const;
-	void SetValue(FLiveLinkSubjectRepresentation NewValue);
+	SLiveLinkSubjectRepresentationPicker::FLiveLinkSourceSubjectRole GetValue() const;
+	void SetValue(SLiveLinkSubjectRepresentationPicker::FLiveLinkSourceSubjectRole NewValue);
 	bool HasMultipleValues() const;
 
 	TSharedPtr<IPropertyHandle> StructPropertyHandle;
diff --git a/Engine/Plugins/Animation/LiveLink/Source/LiveLinkEditor/Private/LiveLinkSubjectRepresentationDetailCustomization.cpp b/Engine/Plugins/Animation/LiveLink/Source/LiveLinkEditor/Private/LiveLinkSubjectRepresentationDetailCustomization.cpp
--- a/Engine/Plugins/Animation/LiveLink/Source/LiveLinkEditor/Private/LiveLinkSubjectRepresentationDetailCustomization.cpp
+++ b/Engine/Plugins/Animation/LiveLink/Source/LiveLinkEditor/Private/LiveLinkSubjectRepresentationDetailCustomization.cpp
@@ -61,7 +61,7 @@ void FLiveLinkSubjectRepresentationDetailCustomization::CustomizeChildren(TShare
 	}
 }
 
-FLiveLinkSubjectRepresentation FLiveLinkSubjectRepresentationDetailCustomization::GetValue() const
+SLiveLinkSubjectRepresentationPicker::FLiveLinkSourceSubjectRole FLiveLinkSubjectRepresentationDetailCustomization::GetValue() const
 {
 	TArray<const void*> RawData;
 	StructPropertyHandle->AccessRawData(RawData);
@@ -70,23 +70,24 @@ FLiveLinkSubjectRepresentation FLiveLinkSubjectRepresentationDetailCustomization
 	{
 		if (RawPtr)
 		{
-			return *reinterpret_cast<const FLiveLinkSubjectRepresentation *>(RawPtr);
+			return SLiveLinkSubjectRepresentationPicker::FLiveLinkSourceSubjectRole(*reinterpret_cast<const FLiveLinkSubjectRepresentation *>(RawPtr));
 		}
 	}
 
-	return FLiveLinkSubjectRepresentation();
+	return SLiveLinkSubjectRepresentationPicker::FLiveLinkSourceSubjectRole();
 }
 
-void FLiveLinkSubjectRepresentationDetailCustomization::SetValue(FLiveLinkSubjectRepresentation NewValue)
+void FLiveLinkSubjectRepresentationDetailCustomization::SetValue(SLiveLinkSubjectRepresentationPicker::FLiveLinkSourceSubjectRole NewValue)
 {
 	FStructProperty* StructProperty = CastFieldChecked<FStructProperty>(StructPropertyHandle->GetProperty());
 
 	TArray<void*> RawData;
 	StructPropertyHandle->AccessRawData(RawData);
 	FLiveLinkSubjectRepresentation* PreviousValue = reinterpret_cast<FLiveLinkSubjectRepresentation*>(RawData[0]);
+	FLiveLinkSubjectRepresentation NewSubRep = NewValue.ToSubjectRepresentation();
 
 	FString TextValue;
-	StructProperty->Struct->ExportText(TextValue, &NewValue, PreviousValue, nullptr, EPropertyPortFlags::PPF_None, nullptr);
+	StructProperty->Struct->ExportText(TextValue, &NewSubRep, PreviousValue, nullptr, EPropertyPortFlags::PPF_None, nullptr);
 	ensure(StructPropertyHandle->SetValueFromFormattedString(TextValue, EPropertyValueSetFlags::DefaultFlags) == FPropertyAccess::Result::Success);
 }
 
diff --git a/Engine/Plugins/Animation/LiveLink/Source/LiveLinkEditor/Private/LiveLinkSubjectRepresentationDetailCustomization.h b/Engine/Plugins/Animation/LiveLink/Source/LiveLinkEditor/Private/LiveLinkSubjectRepresentationDetailCustomization.h
--- a/Engine/Plugins/Animation/LiveLink/Source/LiveLinkEditor/Private/LiveLinkSubjectRepresentationDetailCustomization.h
+++ b/Engine/Plugins/Animation/LiveLink/Source/LiveLinkEditor/Private/LiveLinkSubjectRepresentationDetailCustomization.h
@@ -6,6 +6,7 @@
 #include "LiveLinkRole.h"
 #include "LiveLinkTypes.h"
 #include "IPropertyTypeCustomization.h"
+#include "SLiveLinkSubjectRepresentationPicker.h"
 
 class FLiveLinkSubjectRepresentationDetailCustomization : public IPropertyTypeCustomization
 {
@@ -19,8 +20,8 @@ public:
 	virtual void CustomizeChildren(TSharedRef<IPropertyHandle> PropertyHandle, IDetailChildrenBuilder& ChildBuilder, IPropertyTypeCustomizationUtils& CustomizationUtils) override;
 
 private:
-	FLiveLinkSubjectRepresentation GetValue() const;
-	void SetValue(FLiveLinkSubjectRepresentation NewValue);
+	SLiveLinkSubjectRepresentationPicker::FLiveLinkSourceSubjectRole GetValue() const;
+	void SetValue(SLiveLinkSubjectRepresentationPicker::FLiveLinkSourceSubjectRole NewValue);
 	bool HasMultipleValues() const;
 
 	TSharedPtr<IPropertyHandle> StructPropertyHandle;
diff --git a/Engine/Plugins/Animation/LiveLink/Source/LiveLinkEditor/Private/SLiveLinkSubjectNameGraphPin.cpp b/Engine/Plugins/Animation/LiveLink/Source/LiveLinkEditor/Private/SLiveLinkSubjectNameGraphPin.cpp
--- a/Engine/Plugins/Animation/LiveLink/Source/LiveLinkEditor/Private/SLiveLinkSubjectNameGraphPin.cpp
+++ b/Engine/Plugins/Animation/LiveLink/Source/LiveLinkEditor/Private/SLiveLinkSubjectNameGraphPin.cpp
@@ -23,7 +23,7 @@ TSharedRef<SWidget> SLiveLinkSubjectNameGraphPin::GetDefaultValueWidget()
 		.Visibility(this, &SGraphPin::GetDefaultValueVisibility);
 }
 
-FLiveLinkSubjectRepresentation SLiveLinkSubjectNameGraphPin::GetValue() const
+SLiveLinkSubjectRepresentationPicker::FLiveLinkSourceSubjectRole SLiveLinkSubjectNameGraphPin::GetValue() const
 {
 	FLiveLinkSubjectRepresentation SubjectRepresentation;
 
@@ -33,10 +33,10 @@ FLiveLinkSubjectRepresentation SLiveLinkSubjectNameGraphPin::GetValue() const
 		SubjectRepresentation.Role = Cast<UClass>(GraphPinObj->DefaultObject);
 	}
 
-	return SubjectRepresentation;
+	return SLiveLinkSubjectRepresentationPicker::FLiveLinkSourceSubjectRole(SubjectRepresentation);
 }
 
-void SLiveLinkSubjectNameGraphPin::SetValue(FLiveLinkSubjectRepresentation NewValue)
+void SLiveLinkSubjectNameGraphPin::SetValue(SLiveLinkSubjectRepresentationPicker::FLiveLinkSourceSubjectRole NewValue)
 {
 	FString ValueString;
 	FLiveLinkSubjectName::StaticStruct()->ExportText(ValueString, &NewValue.Subject, nullptr, nullptr, EPropertyPortFlags::PPF_None, nullptr);
diff --git a/Engine/Plugins/Animation/LiveLink/Source/LiveLinkEditor/Private/SLiveLinkSubjectNameGraphPin.h b/Engine/Plugins/Animation/LiveLink/Source/LiveLinkEditor/Private/SLiveLinkSubjectNameGraphPin.h
--- a/Engine/Plugins/Animation/LiveLink/Source/LiveLinkEditor/Private/SLiveLinkSubjectNameGraphPin.h
+++ b/Engine/Plugins/Animation/LiveLink/Source/LiveLinkEditor/Private/SLiveLinkSubjectNameGraphPin.h
@@ -5,6 +5,7 @@
 #include "CoreMinimal.h"
 #include "LiveLinkRole.h"
 #include "SGraphPin.h"
+#include "SLiveLinkSubjectRepresentationPicker.h"
 #include "Widgets/SWidget.h"
 #include "Widgets/DeclarativeSyntaxSupport.h"
 
@@ -21,6 +22,6 @@ public:
 	//~ End SGraphPin Interface
 
 private:
-	FLiveLinkSubjectRepresentation GetValue() const;
-	void SetValue(FLiveLinkSubjectRepresentation NewValue);
+	SLiveLinkSubjectRepresentationPicker::FLiveLinkSourceSubjectRole GetValue() const;
+	void SetValue(SLiveLinkSubjectRepresentationPicker::FLiveLinkSourceSubjectRole NewValue);
 };
diff --git a/Engine/Plugins/Animation/LiveLink/Source/LiveLinkEditor/Private/SLiveLinkSubjectRepresentationGraphPin.cpp b/Engine/Plugins/Animation/LiveLink/Source/LiveLinkEditor/Private/SLiveLinkSubjectRepresentationGraphPin.cpp
--- a/Engine/Plugins/Animation/LiveLink/Source/LiveLinkEditor/Private/SLiveLinkSubjectRepresentationGraphPin.cpp
+++ b/Engine/Plugins/Animation/LiveLink/Source/LiveLinkEditor/Private/SLiveLinkSubjectRepresentationGraphPin.cpp
@@ -27,14 +27,14 @@ TSharedRef<SWidget>	SLiveLinkSubjectRepresentationGraphPin::GetDefaultValueWidge
 		.OnValueChanged(this, &SLiveLinkSubjectRepresentationGraphPin::SetValue);
 }
 
-FLiveLinkSubjectRepresentation SLiveLinkSubjectRepresentationGraphPin::GetValue() const
+SLiveLinkSubjectRepresentationPicker::FLiveLinkSourceSubjectRole SLiveLinkSubjectRepresentationGraphPin::GetValue() const
 {
-	return SubjectRepresentation;
+	return SLiveLinkSubjectRepresentationPicker::FLiveLinkSourceSubjectRole(SubjectRepresentation);
 }
 
-void SLiveLinkSubjectRepresentationGraphPin::SetValue(FLiveLinkSubjectRepresentation NewValue)
+void SLiveLinkSubjectRepresentationGraphPin::SetValue(SLiveLinkSubjectRepresentationPicker::FLiveLinkSourceSubjectRole NewValue)
 {
-	SubjectRepresentation = NewValue;
+	SubjectRepresentation = NewValue.ToSubjectRepresentation();
 
 	FString ValueString;
 	FLiveLinkSubjectRepresentation::StaticStruct()->ExportText(ValueString, &SubjectRepresentation, nullptr, nullptr, EPropertyPortFlags::PPF_None, nullptr);
diff --git a/Engine/Plugins/Animation/LiveLink/Source/LiveLinkEditor/Private/SLiveLinkSubjectRepresentationGraphPin.h b/Engine/Plugins/Animation/LiveLink/Source/LiveLinkEditor/Private/SLiveLinkSubjectRepresentationGraphPin.h
--- a/Engine/Plugins/Animation/LiveLink/Source/LiveLinkEditor/Private/SLiveLinkSubjectRepresentationGraphPin.h
+++ b/Engine/Plugins/Animation/LiveLink/Source/LiveLinkEditor/Private/SLiveLinkSubjectRepresentationGraphPin.h
@@ -6,6 +6,7 @@
 #include "Widgets/SWidget.h"
 #include "Widgets/DeclarativeSyntaxSupport.h"
 #include "SGraphPin.h"
+#include "SLiveLinkSubjectRepresentationPicker.h"
 
 #include "LiveLinkRole.h"
 
@@ -26,8 +27,8 @@ public:
 	//~ End SGraphPin Interface
 
 private:
-	FLiveLinkSubjectRepresentation GetValue() const;
-	void SetValue(FLiveLinkSubjectRepresentation NewValue);
+	SLiveLinkSubjectRepresentationPicker::FLiveLinkSourceSubjectRole GetValue() const;
+	void SetValue(SLiveLinkSubjectRepresentationPicker::FLiveLinkSourceSubjectRole NewValue);
 
 	FLiveLinkSubjectRepresentation SubjectRepresentation;
 };
diff --git a/Engine/Plugins/Animation/LiveLink/Source/LiveLinkEditor/Private/SLiveLinkSubjectRepresentationPicker.cpp b/Engine/Plugins/Animation/LiveLink/Source/LiveLinkEditor/Private/SLiveLinkSubjectRepresentationPicker.cpp
--- a/Engine/Plugins/Animation/LiveLink/Source/LiveLinkEditor/Private/SLiveLinkSubjectRepresentationPicker.cpp
+++ b/Engine/Plugins/Animation/LiveLink/Source/LiveLinkEditor/Private/SLiveLinkSubjectRepresentationPicker.cpp
@@ -32,32 +32,37 @@
 namespace SubjectUI
 {
 	static const FName EnabledColumnName(TEXT("Enabled"));
+	static const FName SourceColumnName(TEXT("Source"));
 	static const FName NameColumnName(TEXT("Name"));
 	static const FName RoleColumnName(TEXT("Role"));
 };
 
-struct FLiveLinkSubjectRepEntry
+struct FLiveLinkSubjectRepresentationPickerEntry
 {
-	FLiveLinkSubjectRepEntry(const FLiveLinkSubjectRepresentation& InSubRep, bool bInEnabled)
-		: SubjectRepresentation(InSubRep)
+	FLiveLinkSubjectRepresentationPickerEntry() = default;
+	FLiveLinkSubjectRepresentationPickerEntry(const SLiveLinkSubjectRepresentationPicker::FLiveLinkSourceSubjectRole& InSourceSubjectRole, const FText & InSourceType, bool bInEnabled)
+		: SourceSubjectRole(InSourceSubjectRole)
+		, SourceType(InSourceType)
 		, bEnabled(bInEnabled)
 	{}
-	FLiveLinkSubjectRepresentation SubjectRepresentation;
+
+	SLiveLinkSubjectRepresentationPicker::FLiveLinkSourceSubjectRole SourceSubjectRole;
+	FText SourceType;
 	bool bEnabled;
 };
 
-class SLiveLinkSubjectEntryRow : public SMultiColumnTableRow<FLiveLinkSubjectRepEntryPtr>
+class SLiveLinkSubjectEntryRow : public SMultiColumnTableRow<FLiveLinkSubjectRepresentationPickerEntryPtr>
 {
 public:
 	SLATE_BEGIN_ARGS(SLiveLinkSubjectEntryRow) {}
-		SLATE_ARGUMENT(FLiveLinkSubjectRepEntryPtr, Entry)
+		SLATE_ARGUMENT(FLiveLinkSubjectRepresentationPickerEntryPtr, Entry)
 	SLATE_END_ARGS()
 
 	void Construct(const FArguments& Args, const TSharedRef<STableViewBase>& OwnerTableView)
 	{
 		EntryPtr = Args._Entry;
 
-		SMultiColumnTableRow<FLiveLinkSubjectRepEntryPtr>::Construct(
+		SMultiColumnTableRow<FLiveLinkSubjectRepresentationPickerEntryPtr>::Construct(
 			FSuperRowType::FArguments()
 			.Padding(0.f),
 			OwnerTableView
@@ -73,21 +78,30 @@ public:
 				.IsChecked(EntryPtr->bEnabled ? ECheckBoxState::Checked : ECheckBoxState::Unchecked)
 				.IsEnabled(false);
 		}
+		else if (ColumnName == SubjectUI::SourceColumnName)
+		{
+			return	SNew(STextBlock)
+				.Text(EntryPtr->SourceType);
+		}
 		else if (ColumnName == SubjectUI::NameColumnName)
 		{
 			return	SNew(STextBlock)
-				.Text(FText::FromName(EntryPtr->SubjectRepresentation.Subject));
+				.Text(FText::FromName(EntryPtr->SourceSubjectRole.Subject));
 		}
 		else if (ColumnName == SubjectUI::RoleColumnName)
 		{
-			return SNew(STextBlock)
-				.Text(EntryPtr->SubjectRepresentation.Role->GetDefaultObject<ULiveLinkRole>()->GetDisplayName());
+			if (EntryPtr->SourceSubjectRole.Role)
+			{
+				return SNew(STextBlock)
+					.Text(EntryPtr->SourceSubjectRole.Role->GetDefaultObject<ULiveLinkRole>()->GetDisplayName());
+			}
 		}
 
 		return SNullWidget::NullWidget;
 	}
 
-	FLiveLinkSubjectRepEntryPtr EntryPtr;
+private:
+	FLiveLinkSubjectRepresentationPickerEntryPtr EntryPtr;
 };
 
 void SLiveLinkSubjectRepresentationPicker::Construct(const FArguments& InArgs)
@@ -95,61 +109,81 @@ void SLiveLinkSubjectRepresentationPicker::Construct(const FArguments& InArgs)
 	ValueAttribute = InArgs._Value;
 	OnValueChangedDelegate = InArgs._OnValueChanged;
 	HasMultipleValuesAttribute = InArgs._HasMultipleValues;
+	bShowSource = InArgs._ShowSource;
 	bShowRole = InArgs._ShowRole;
 
 	SubjectRepData.Reset();
 	SelectedLiveLinkPreset.Reset();
 
+	ILiveLinkClient& LiveLinkClient = IModularFeatures::Get().GetModularFeature<ILiveLinkClient>(ILiveLinkClient::ModularFeatureName);
+	CachedSourceType = LiveLinkClient.GetSourceType(ValueAttribute.Get().Source);
+
 	TSharedPtr<SWidget> ComboButtonContent;
-	if (bShowRole)
+	if (bShowRole || bShowSource)
 	{
-		ComboButtonContent = SNew(SHorizontalBox)
-			+ SHorizontalBox::Slot()
-			.FillWidth(1)
-			[	
-				SNew(SVerticalBox)
-				+ SVerticalBox::Slot()
+		TSharedRef<SVerticalBox> ContentBox = SNew(SVerticalBox);
+		if (bShowSource)
+		{
+			ContentBox->AddSlot()
+			[
+				SNew(STextBlock)
+				.Font(InArgs._Font)
+				.Text(this, &SLiveLinkSubjectRepresentationPicker::GetSourceNameValueText)
+			];
+		}
+
+		ContentBox->AddSlot()
+		[
+			SNew(STextBlock)
+			.Font(InArgs._Font)
+			.Text(this, &SLiveLinkSubjectRepresentationPicker::GetSubjectNameValueText)
+		];
+
+		if (bShowRole)
+		{
+			ContentBox->AddSlot()
+			[
+				SNew(SHorizontalBox)
+				+ SHorizontalBox::Slot()
+				.AutoWidth()
+				.HAlign(HAlign_Left)
+				.VAlign(VAlign_Center)
 				[
-					SNew(STextBlock)
-					.Font(InArgs._Font)
-					.Text(this, &SLiveLinkSubjectRepresentationPicker::GetSubjectNameValueText)
+					SNew(SImage)
+					.Image(this, &SLiveLinkSubjectRepresentationPicker::GetRoleIcon)
 				]
-				+ SVerticalBox::Slot()
+				+ SHorizontalBox::Slot()
+				.FillWidth(1)
+				.VAlign(VAlign_Center)
 				[
-					SNew(SHorizontalBox)
-					+ SHorizontalBox::Slot()
-					.AutoWidth()
-					.HAlign(HAlign_Left)
-					.VAlign(VAlign_Center)
-					[
-						SNew(SImage)
-						.Image(this, &SLiveLinkSubjectRepresentationPicker::GetRoleIcon)
-					]
-					+ SHorizontalBox::Slot()
-					.FillWidth(1)
-					.VAlign(VAlign_Center)
-					[
-						SNew(STextBlock)
-						.Font(InArgs._Font)
-						.Text(this, &SLiveLinkSubjectRepresentationPicker::GetRoleText)
-					]
+					SNew(STextBlock)
+					.Font(InArgs._Font)
+					.Text(this, &SLiveLinkSubjectRepresentationPicker::GetRoleText)
 				]
 			];
+		}
+
+		SAssignNew(ComboButtonContent, SHorizontalBox)
+		+ SHorizontalBox::Slot()
+		.FillWidth(1)
+		[	
+			ContentBox
+		];
 	}
 	else
 	{
 		ComboButtonContent = SNew(SBorder)
-			.BorderImage(FEditorStyle::GetBrush("NoBorder"))
-			.Padding(FMargin(0, 0, 5, 0))
-			[
-				SNew(SEditableTextBox)
-				.Text(this, &SLiveLinkSubjectRepresentationPicker::GetSubjectNameValueText)
-				.OnTextCommitted(this, &SLiveLinkSubjectRepresentationPicker::OnComboTextCommitted)
-				.SelectAllTextWhenFocused(true)
-				.SelectAllTextOnCommit(true)
-				.ClearKeyboardFocusOnCommit(false)
-				.Font(InArgs._Font)
-			];
+		.BorderImage(FEditorStyle::GetBrush("NoBorder"))
+		.Padding(FMargin(0, 0, 5, 0))
+		[
+			SNew(SEditableTextBox)
+			.Text(this, &SLiveLinkSubjectRepresentationPicker::GetSubjectNameValueText)
+			.OnTextCommitted(this, &SLiveLinkSubjectRepresentationPicker::OnComboTextCommitted)
+			.SelectAllTextWhenFocused(true)
+			.SelectAllTextOnCommit(true)
+			.ClearKeyboardFocusOnCommit(false)
+			.Font(InArgs._Font)
+		];
 	}
 
 	ChildSlot
@@ -168,11 +202,22 @@ void SLiveLinkSubjectRepresentationPicker::Construct(const FArguments& InArgs)
 	];
 }
 
-FLiveLinkSubjectRepresentation SLiveLinkSubjectRepresentationPicker::GetCurrentValue() const
+SLiveLinkSubjectRepresentationPicker::FLiveLinkSourceSubjectRole SLiveLinkSubjectRepresentationPicker::GetCurrentValue() const
 {
 	return ValueAttribute.Get();
 }
 
+FText SLiveLinkSubjectRepresentationPicker::GetSourceNameValueText() const
+{
+	const bool bHasMultipleValues = HasMultipleValuesAttribute.Get();
+	if (bHasMultipleValues)
+	{
+		return LOCTEXT("MultipleValuesText", "<multiple values>");
+	}
+
+	return CachedSourceType;
+}
+
 FText SLiveLinkSubjectRepresentationPicker::GetSubjectNameValueText() const
 {
 	const bool bHasMultipleValues = HasMultipleValuesAttribute.Get();
@@ -181,17 +226,16 @@ FText SLiveLinkSubjectRepresentationPicker::GetSubjectNameValueText() const
 		return LOCTEXT("MultipleValuesText", "<multiple values>");
 	}
 
-	FLiveLinkSubjectRepresentation CurrentSubjectRepresentation = ValueAttribute.Get();
-	return FText::FromName(CurrentSubjectRepresentation.Subject);
+	return FText::FromName(ValueAttribute.Get().Subject);
 }
 
 const FSlateBrush* SLiveLinkSubjectRepresentationPicker::GetRoleIcon() const
 {
 	const bool bHasMultipleValues = HasMultipleValuesAttribute.Get();
-	FLiveLinkSubjectRepresentation CurrentSubjectRepresentation = ValueAttribute.Get();
-	if (!bHasMultipleValues && CurrentSubjectRepresentation.Role != nullptr)
+	FLiveLinkSourceSubjectRole CurrentSrcSubRole = ValueAttribute.Get();
+	if (!bHasMultipleValues && CurrentSrcSubRole.Role != nullptr)
 	{
-		return FSlateIconFinder::FindIconBrushForClass(CurrentSubjectRepresentation.Role);
+		return FSlateIconFinder::FindIconBrushForClass(CurrentSrcSubRole.Role);
 	}
 	return FSlateIconFinder::FindIconBrushForClass(ULiveLinkRole::StaticClass());
 }
@@ -204,12 +248,12 @@ FText SLiveLinkSubjectRepresentationPicker::GetRoleText() const
 		return LOCTEXT("MultipleValuesText", "<multiple values>");
 	}
 
-	FLiveLinkSubjectRepresentation CurrentSubjectRepresentation = ValueAttribute.Get();
-	if (CurrentSubjectRepresentation.Role == nullptr)
+	FLiveLinkSourceSubjectRole CurrentSrcSubjRole = ValueAttribute.Get();
+	if (CurrentSrcSubjRole.Role == nullptr)
 	{
 		return LOCTEXT("NoValueText", "<none>");
 	}
-	return CurrentSubjectRepresentation.Role->GetDisplayNameText();
+	return CurrentSrcSubjRole.Role->GetDisplayNameText();
 }
 
 TSharedRef<SWidget> SLiveLinkSubjectRepresentationPicker::BuildMenu()
@@ -218,6 +262,24 @@ TSharedRef<SWidget> SLiveLinkSubjectRepresentationPicker::BuildMenu()
 	SelectedLiveLinkPreset.Reset();
 	BuildSubjectRepDataList();
 
+	TSharedRef<SHeaderRow> HeaderRow = SNew(SHeaderRow)
+		+ SHeaderRow::Column(SubjectUI::EnabledColumnName)
+		.ManualWidth(20.f)
+		.DefaultLabel(LOCTEXT("EnabledColumnHeaderName", ""));
+
+	if (bShowSource)
+	{
+		HeaderRow->AddColumn(SHeaderRow::Column(SubjectUI::SourceColumnName)
+			.DefaultLabel(LOCTEXT("SourceColumnHeaderName", "Source"))
+			.FillWidth(30));
+	}
+	HeaderRow->AddColumn(SHeaderRow::Column(SubjectUI::NameColumnName)
+		.FillWidth(bShowSource ? 50 : 60.f)
+		.DefaultLabel(LOCTEXT("SubjectColumnHeaderName", "Subject")));
+	HeaderRow->AddColumn(SHeaderRow::Column(SubjectUI::RoleColumnName)
+		.FillWidth(bShowSource ? 20 : 40.f)
+		.DefaultLabel(LOCTEXT("RoleColumnHeaderName", "Role")));
+
 	return SNew(SBox)
 		.Padding(0)
 		.WidthOverride(300.f)
@@ -282,24 +344,12 @@ TSharedRef<SWidget> SLiveLinkSubjectRepresentationPicker::BuildMenu()
 					SNew(SBorder)
 					.Padding(FMargin(4.0f, 4.0f))
 					[
-						SAssignNew(SubjectListView, SListView<FLiveLinkSubjectRepEntryPtr>)
+						SAssignNew(SubjectListView, SListView<FLiveLinkSubjectRepresentationPickerEntryPtr>)
 						.ListItemsSource(&SubjectRepData)
 						.SelectionMode(ESelectionMode::Single)
 						.OnGenerateRow(this, &SLiveLinkSubjectRepresentationPicker::MakeSubjectRepListViewWidget)
 						.OnSelectionChanged(this, &SLiveLinkSubjectRepresentationPicker::OnSubjectRepListSelectionChanged)
-						.HeaderRow
-						(
-							SNew(SHeaderRow)
-							+ SHeaderRow::Column(SubjectUI::EnabledColumnName)
-							.ManualWidth(20.f)
-							.DefaultLabel(LOCTEXT("EnabledColumnHeaderName", ""))
-							+ SHeaderRow::Column(SubjectUI::NameColumnName)
-							.FillWidth(60.f)
-							.DefaultLabel(LOCTEXT("SubjectColumnHeaderName", "Subject"))
-							+ SHeaderRow::Column(SubjectUI::RoleColumnName)
-							.FillWidth(40.f)
-							.DefaultLabel(LOCTEXT("RoleColumnHeaderName", "Role"))
-						)
+						.HeaderRow(HeaderRow)
 					]
 				]
 			]
@@ -337,21 +387,23 @@ bool SLiveLinkSubjectRepresentationPicker::HasCurrentPreset() const
 	return SelectedLiveLinkPreset.IsValid();
 }
 
-TSharedRef<ITableRow> SLiveLinkSubjectRepresentationPicker::MakeSubjectRepListViewWidget(FLiveLinkSubjectRepEntryPtr Entry, const TSharedRef<STableViewBase>& OwnerTable) const
+TSharedRef<ITableRow> SLiveLinkSubjectRepresentationPicker::MakeSubjectRepListViewWidget(FLiveLinkSubjectRepresentationPickerEntryPtr Entry, const TSharedRef<STableViewBase>& OwnerTable) const
 {
 	return SNew(SLiveLinkSubjectEntryRow, OwnerTable)
 		.Entry(Entry);
 }
 
-void SLiveLinkSubjectRepresentationPicker::OnSubjectRepListSelectionChanged(FLiveLinkSubjectRepEntryPtr Entry, ESelectInfo::Type SelectionType)
+void SLiveLinkSubjectRepresentationPicker::OnSubjectRepListSelectionChanged(FLiveLinkSubjectRepresentationPickerEntryPtr Entry, ESelectInfo::Type SelectionType)
 {
 	if (Entry.IsValid())
 	{
-		SetValue(Entry->SubjectRepresentation);
+		SetValue(Entry->SourceSubjectRole);
+		CachedSourceType = Entry->SourceType;
 	}
 	else
 	{
-		SetValue(FLiveLinkSubjectRepresentation());
+		SetValue(FLiveLinkSourceSubjectRole());
+		CachedSourceType = FText::GetEmpty();
 	}
 }
 
@@ -401,12 +453,14 @@ void SLiveLinkSubjectRepresentationPicker::NewPresetSelected(const FAssetData& A
 
 void SLiveLinkSubjectRepresentationPicker::OnComboTextCommitted(const FText& NewText, ETextCommit::Type InTextCommit)
 {
-	FLiveLinkSubjectRepresentation Representation;
-	Representation.Subject.Name = *NewText.ToString();
-	SetValue(Representation);
+	FLiveLinkSourceSubjectRole SourceSubjectRole;
+	SourceSubjectRole.Subject.Name = *NewText.ToString();
+	SetValue(SourceSubjectRole);
+
+	CachedSourceType = FText::GetEmpty();
 }
 
-void SLiveLinkSubjectRepresentationPicker::SetValue(const FLiveLinkSubjectRepresentation& InValue)
+void SLiveLinkSubjectRepresentationPicker::SetValue(const SLiveLinkSubjectRepresentationPicker::FLiveLinkSourceSubjectRole& InValue)
 {
 	if (OnValueChangedDelegate.IsBound())
 	{
@@ -433,13 +487,27 @@ void SLiveLinkSubjectRepresentationPicker::BuildSubjectRepDataList()
 	{
 		for (const FLiveLinkSubjectPreset& SubjectPreset : LiveLinkPresetPtr->GetSubjectPresets())
 		{
-			FLiveLinkSubjectRepresentation Representation;
-			Representation.Role = SubjectPreset.Role;
-			Representation.Subject = SubjectPreset.Key.SubjectName;
+			SLiveLinkSubjectRepresentationPicker::FLiveLinkSourceSubjectRole SrcSubRole;
+			SrcSubRole.Source = SubjectPreset.Key.Source;
+			SrcSubRole.Subject = SubjectPreset.Key.SubjectName;
+			SrcSubRole.Role = SubjectPreset.Role;
+			const FLiveLinkSourcePreset* FoundSourcePreset = nullptr;
+			if (bShowSource)
+			{
+				const FGuid& SourceGuid = SubjectPreset.Key.Source;
+				FoundSourcePreset = LiveLinkPresetPtr->GetSourcePresets().FindByPredicate([&SourceGuid](const FLiveLinkSourcePreset& SourcePreset) { return SourcePreset.Guid == SourceGuid; });
+			}
 
-			if (Representation.Role != nullptr && !Representation.Subject.IsNone())
+			if (SrcSubRole.Role != nullptr && !SrcSubRole.Subject.IsNone())
 			{
-				SubjectRepData.Add(MakeShared<FLiveLinkSubjectRepEntry>(Representation, SubjectPreset.bEnabled));
+				if (FoundSourcePreset)
+				{
+					SubjectRepData.Add(MakeShared<FLiveLinkSubjectRepresentationPickerEntry>(SrcSubRole, FoundSourcePreset->SourceType, SubjectPreset.bEnabled));
+				}
+				else
+				{
+					SubjectRepData.Add(MakeShared<FLiveLinkSubjectRepresentationPickerEntry>(SrcSubRole, FText::GetEmpty(), SubjectPreset.bEnabled));
+				}
 			}
 		}
 	}
@@ -448,29 +516,100 @@ void SLiveLinkSubjectRepresentationPicker::BuildSubjectRepDataList()
 		ILiveLinkClient& LiveLinkClient = IModularFeatures::Get().GetModularFeature<ILiveLinkClient>(ILiveLinkClient::ModularFeatureName);
 		TArray<FLiveLinkSubjectKey> SubjectKeys = LiveLinkClient.GetSubjects(true, true);
 
-		TArray<FLiveLinkSubjectName> UniqueSubjectName;
-		UniqueSubjectName.Reset(SubjectKeys.Num());
-		for(const FLiveLinkSubjectKey& Key : SubjectKeys)
+		if (bShowSource)
 		{
-			UniqueSubjectName.AddUnique(Key.SubjectName);
-		}
-		UniqueSubjectName.Sort(FNameLexicalLess());
+			TMap<FGuid, FText> UniqueSources;
+			for (const FLiveLinkSubjectKey& Key : SubjectKeys)
+			{
+				if (!UniqueSources.Contains(Key.Source))
+				{
+					UniqueSources.Add(Key.Source, LiveLinkClient.GetSourceType(Key.Source));
+				}
+			}
 
-		for (const FLiveLinkSubjectName& SubjectName : UniqueSubjectName)
+			for (const FLiveLinkSubjectKey& SubjectKey : SubjectKeys)
+			{
+				if (!SubjectKey.SubjectName.IsNone())
+				{
+					bool bEnabled = LiveLinkClient.IsSubjectEnabled(SubjectKey, true);
+					SLiveLinkSubjectRepresentationPicker::FLiveLinkSourceSubjectRole SrcSubRole;
+					SrcSubRole.Source = SubjectKey.Source;
+					SrcSubRole.Subject = SubjectKey.SubjectName;
+					SrcSubRole.Role = LiveLinkClient.GetSubjectRole(SubjectKey);
+					SubjectRepData.Add(MakeShared<FLiveLinkSubjectRepresentationPickerEntry>(SrcSubRole, UniqueSources[SubjectKey.Source], bEnabled));
+				}
+			}
+
+			struct FEntryLess
+			{
+				FORCEINLINE bool operator()(const FLiveLinkSubjectRepresentationPickerEntryPtr& A, const FLiveLinkSubjectRepresentationPickerEntryPtr& B) const
+				{
+					return A->SourceSubjectRole.Source < B->SourceSubjectRole.Source && A->SourceSubjectRole.Subject.Name.Compare(B->SourceSubjectRole.Subject.Name) < 0;
+				}
+			};
+			SubjectRepData.Sort(FEntryLess());
+		}
+		else
 		{
-			bool bEnabled = LiveLinkClient.IsSubjectEnabled(SubjectName);
-			FLiveLinkSubjectRepresentation Representation;
-			Representation.Subject = SubjectName;
-			Representation.Role = LiveLinkClient.GetSubjectRole(SubjectName);
+			TMap<FLiveLinkSubjectName, int32> UniqueSubjectName;
+			for(const FLiveLinkSubjectKey& Key : SubjectKeys)
+			{
+				if (!Key.SubjectName.IsNone())
+				{
+					int32* FoundValue = UniqueSubjectName.Find(Key.SubjectName);
+					if (FoundValue)
+					{
+						++(*FoundValue);
+					}
+					else
+					{
+						UniqueSubjectName.Add(Key.SubjectName, 1);
+					}
+				}
+			}
 
-			if (Representation.Role != nullptr && !Representation.Subject.IsNone())
+			for (const auto& Item : UniqueSubjectName)
 			{
-				SubjectRepData.Add(MakeShared<FLiveLinkSubjectRepEntry>(Representation, bEnabled));
+				bool bEnabled = LiveLinkClient.IsSubjectEnabled(Item.Key);
+
+				SLiveLinkSubjectRepresentationPicker::FLiveLinkSourceSubjectRole SrcSubRole;
+				SrcSubRole.Subject = Item.Key;
+				SrcSubRole.Role = LiveLinkClient.GetSubjectRole(Item.Key);
+
+				if (SrcSubRole.Role == nullptr && Item.Value != 1)
+				{
+					bool bIsFirstEntry = true;
+					for (const FLiveLinkSubjectKey& Key : SubjectKeys)
+					{
+						if (Key.SubjectName == Item.Key)
+						{
+							 TSubclassOf<ULiveLinkRole> NewRole = LiveLinkClient.GetSubjectRole(Key);
+							 if (!bIsFirstEntry && SrcSubRole.Role != NewRole)
+							 {
+								 SrcSubRole.Role = nullptr;
+								 break;
+							 }
+
+							 bIsFirstEntry = false;
+							 SrcSubRole.Role = NewRole;
+						}
+					}
+				}
+				SubjectRepData.Add(MakeShared<FLiveLinkSubjectRepresentationPickerEntry>(SrcSubRole, FText::GetEmpty(), bEnabled));
 			}
+
+			struct FEntryLess
+			{
+				FORCEINLINE bool operator()(const FLiveLinkSubjectRepresentationPickerEntryPtr& A, const FLiveLinkSubjectRepresentationPickerEntryPtr& B) const
+				{
+					return A->SourceSubjectRole.Subject.Name.Compare(B->SourceSubjectRole.Subject.Name) < 0;
+				}
+			};
+			SubjectRepData.Sort(FEntryLess());
 		}
 	}
 
-	TSharedPtr<SListView<FLiveLinkSubjectRepEntryPtr>> SubjectListViewPin = SubjectListView.Pin();
+	TSharedPtr<SListView<FLiveLinkSubjectRepresentationPickerEntryPtr>> SubjectListViewPin = SubjectListView.Pin();
 	if (SubjectListViewPin.IsValid())
 	{
 		SubjectListViewPin->RebuildList();
diff --git a/Engine/Plugins/Animation/LiveLink/Source/LiveLinkEditor/Public/SLiveLinkSubjectRepresentationPicker.h b/Engine/Plugins/Animation/LiveLink/Source/LiveLinkEditor/Public/SLiveLinkSubjectRepresentationPicker.h
--- a/Engine/Plugins/Animation/LiveLink/Source/LiveLinkEditor/Public/SLiveLinkSubjectRepresentationPicker.h
+++ b/Engine/Plugins/Animation/LiveLink/Source/LiveLinkEditor/Public/SLiveLinkSubjectRepresentationPicker.h
@@ -15,7 +15,7 @@
 #include "Widgets/Views/SListView.h"
 
 struct FAssetData;
-struct FLiveLinkSubjectRepEntry;
+struct FLiveLinkSubjectRepresentationPickerEntry;
 struct FSlateBrush;
 
 class FMenuBuilder;
@@ -23,7 +23,7 @@ class ITableRow;
 class SComboButton;
 class STableViewBase;
 
-typedef TSharedPtr<FLiveLinkSubjectRepEntry> FLiveLinkSubjectRepEntryPtr;
+typedef TSharedPtr<FLiveLinkSubjectRepresentationPickerEntry> FLiveLinkSubjectRepresentationPickerEntryPtr;
 
 
 /**
@@ -32,7 +32,28 @@ typedef TSharedPtr<FLiveLinkSubjectRepEntry> FLiveLinkSubjectRepEntryPtr;
 class LIVELINKEDITOR_API SLiveLinkSubjectRepresentationPicker : public SCompoundWidget
 {
 public:
-	DECLARE_DELEGATE_OneParam(FOnValueChanged, FLiveLinkSubjectRepresentation);
+	struct FLiveLinkSourceSubjectRole
+	{
+		FLiveLinkSourceSubjectRole() = default;
+		explicit FLiveLinkSourceSubjectRole(FLiveLinkSubjectRepresentation InSubjectRepresentation)
+			: Subject(InSubjectRepresentation.Subject), Role(InSubjectRepresentation.Role)
+		{}
+		explicit FLiveLinkSourceSubjectRole(FLiveLinkSubjectKey InSubjectKey)
+			: Source(InSubjectKey.Source), Subject(InSubjectKey.SubjectName)
+		{}
+		FLiveLinkSourceSubjectRole(FGuid InSource, FLiveLinkSubjectName InSubject, TSubclassOf<ULiveLinkRole> InRole)
+			: Source(InSource), Subject(InSubject), Role(InRole)
+		{}
+
+		FLiveLinkSubjectRepresentation ToSubjectRepresentation() const { return FLiveLinkSubjectRepresentation(Subject, Role); }
+		FLiveLinkSubjectKey ToSubjectKey() const { return FLiveLinkSubjectKey(Source, Subject); }
+
+		FGuid Source;
+		FLiveLinkSubjectName Subject;
+		TSubclassOf<ULiveLinkRole> Role;
+	};
+
+	DECLARE_DELEGATE_OneParam(FOnValueChanged, FLiveLinkSourceSubjectRole);
 
 	SLATE_BEGIN_ARGS(SLiveLinkSubjectRepresentationPicker)
 		: _ComboButtonStyle(&FCoreStyle::Get().GetWidgetStyle< FComboButtonStyle >("ComboButton"))
@@ -40,6 +61,7 @@ public:
 		, _ForegroundColor(FCoreStyle::Get().GetSlateColor("InvertedForeground"))
 		, _ContentPadding(FMargin(2.f, 0.f))
 		, _HasMultipleValues(false)
+		, _ShowSource(false)
 		, _ShowRole(false)
 		, _Font()
 	{}
@@ -57,7 +79,7 @@ public:
 		SLATE_ATTRIBUTE(FMargin, ContentPadding)
 	
 		/** Attribute used to retrieve the current value. */
-		SLATE_ATTRIBUTE(FLiveLinkSubjectRepresentation, Value)
+		SLATE_ATTRIBUTE(FLiveLinkSourceSubjectRole, Value)
 	
 		/** Delegate for handling when for when the current value changes. */
 		SLATE_EVENT(FOnValueChanged, OnValueChanged)
@@ -65,6 +87,9 @@ public:
 		/** Attribute used to retrieve whether the picker has multiple values. */
 		SLATE_ATTRIBUTE(bool, HasMultipleValues)
 
+		/** Attribute used to retrieve whether the picker should show the source name. */
+		SLATE_ARGUMENT(bool, ShowSource)
+
 		/** Attribute used to retrieve whether the picker should show roles. */
 		SLATE_ARGUMENT(bool, ShowRole)
 
@@ -81,9 +106,10 @@ public:
 	/**
 	 * Access the current value of this picker
 	 */
-	FLiveLinkSubjectRepresentation GetCurrentValue() const;
+	FLiveLinkSourceSubjectRole GetCurrentValue() const;
 
 private:
+	FText GetSourceNameValueText() const;
 	FText GetSubjectNameValueText() const;
 	const FSlateBrush* GetRoleIcon() const;
 	FText GetRoleText() const;
@@ -94,13 +120,13 @@ private:
 	FReply ClearCurrentPreset();
 	bool HasCurrentPreset() const;
 	
-	TSharedRef<ITableRow> MakeSubjectRepListViewWidget(FLiveLinkSubjectRepEntryPtr Entry, const TSharedRef<STableViewBase>& OwnerTable) const;
-	void OnSubjectRepListSelectionChanged(FLiveLinkSubjectRepEntryPtr Entry, ESelectInfo::Type SelectionType);
+	TSharedRef<ITableRow> MakeSubjectRepListViewWidget(FLiveLinkSubjectRepresentationPickerEntryPtr Entry, const TSharedRef<STableViewBase>& OwnerTable) const;
+	void OnSubjectRepListSelectionChanged(FLiveLinkSubjectRepresentationPickerEntryPtr Entry, ESelectInfo::Type SelectionType);
 
 	TSharedRef<SWidget> BuildPresetSubMenu();
 	void NewPresetSelected(const FAssetData& AssetData);
 	void OnComboTextCommitted(const FText& NewText, ETextCommit::Type InTextCommit);
-	void SetValue(const FLiveLinkSubjectRepresentation& InValue);
+	void SetValue(const FLiveLinkSourceSubjectRole& InValue);
 
 	void BuildSubjectRepDataList();
 
@@ -108,11 +134,13 @@ private:
 	TWeakObjectPtr<ULiveLinkPreset> SelectedLiveLinkPreset;
 	TWeakPtr<SComboButton> PickerComboButton;
 	TWeakPtr<SComboButton> SelectPresetComboButton;
-	TWeakPtr<SListView<FLiveLinkSubjectRepEntryPtr>> SubjectListView;
-	TArray<FLiveLinkSubjectRepEntryPtr> SubjectRepData;
+	TWeakPtr<SListView<FLiveLinkSubjectRepresentationPickerEntryPtr>> SubjectListView;
+	TArray<FLiveLinkSubjectRepresentationPickerEntryPtr> SubjectRepData;
 
-	TAttribute<FLiveLinkSubjectRepresentation> ValueAttribute;
+	FText CachedSourceType;
+	TAttribute<FLiveLinkSourceSubjectRole> ValueAttribute;
 	FOnValueChanged OnValueChangedDelegate;
 	TAttribute<bool> HasMultipleValuesAttribute;
+	bool bShowSource;
 	bool bShowRole;
 };
diff --git a/Engine/Source/Runtime/LiveLinkInterface/Public/LiveLinkPresetTypes.h b/Engine/Source/Runtime/LiveLinkInterface/Public/LiveLinkPresetTypes.h
--- a/Engine/Source/Runtime/LiveLinkInterface/Public/LiveLinkPresetTypes.h
+++ b/Engine/Source/Runtime/LiveLinkInterface/Public/LiveLinkPresetTypes.h
@@ -24,6 +24,10 @@ struct FLiveLinkSourcePreset
 
 	UPROPERTY()
 	ULiveLinkSourceSettings* Settings = nullptr;
+
+	/** The SourceType when the source was saved to a Preset. */
+	UPROPERTY()
+	FText SourceType;
 };
 
 
diff --git a/Engine/Source/Runtime/LiveLinkInterface/Public/LiveLinkRole.h b/Engine/Source/Runtime/LiveLinkInterface/Public/LiveLinkRole.h
--- a/Engine/Source/Runtime/LiveLinkInterface/Public/LiveLinkRole.h
+++ b/Engine/Source/Runtime/LiveLinkInterface/Public/LiveLinkRole.h
@@ -36,6 +36,11 @@ struct LIVELINKINTERFACE_API FLiveLinkSubjectRepresentation
 {
 	GENERATED_BODY()
 
+	FLiveLinkSubjectRepresentation() = default;
+	FLiveLinkSubjectRepresentation(const FLiveLinkSubjectName& InSubject, const TSubclassOf<ULiveLinkRole>& InRole)
+		: Subject(InSubject), Role(InRole)
+	{ }
+
 	UPROPERTY(EditAnywhere, BlueprintReadWrite, Category="Live Link")
 	FLiveLinkSubjectName Subject;
 
