Add comments to LiveLinkBasic Interpolation

#rnx

diff --git a/Engine/Plugins/Animation/LiveLink/Source/LiveLink/Private/InterpolationProcessor/LiveLinkBasicFrameInterpolateProcessor.cpp b/Engine/Plugins/Animation/LiveLink/Source/LiveLink/Private/InterpolationProcessor/LiveLinkBasicFrameInterpolateProcessor.cpp
--- a/Engine/Plugins/Animation/LiveLink/Source/LiveLink/Private/InterpolationProcessor/LiveLinkBasicFrameInterpolateProcessor.cpp
+++ b/Engine/Plugins/Animation/LiveLink/Source/LiveLink/Private/InterpolationProcessor/LiveLinkBasicFrameInterpolateProcessor.cpp
@@ -235,6 +235,7 @@ void ULiveLinkBasicFrameInterpolationProcessor::FLiveLinkBasicFrameInterpolation
 
 	const FLiveLinkFrameDataStruct& FrameWhenCanNotBlend = (InBlendWeight > 0.5f) ? FrameDataB : FrameDataA;
 
+	// Initialize the data (copy all properties, including those that will be interpolate later).
 	if (Options.bCopyClosestFrame)
 	{
 		OutBlendedFrameData.InitializeWith(FrameDataA.GetStruct(), FrameWhenCanNotBlend.GetBaseData());
@@ -248,8 +249,10 @@ void ULiveLinkBasicFrameInterpolationProcessor::FLiveLinkBasicFrameInterpolation
 		}
 	}
 
+	// Interpolate basic data
 	OutBlendedFrameData.GetBaseData()->WorldTime = FLiveLinkWorldTime(FMath::Lerp(FrameDataA.GetBaseData()->WorldTime.GetOffsettedTime(), FrameDataB.GetBaseData()->WorldTime.GetOffsettedTime(), InBlendWeight), 0.0);
 
+	// Interpolate Property Values
 	if (Options.bInterpolatePropertyValues)
 	{
 		const TArray<float>& PropertiesA = FrameDataA.GetBaseData()->PropertyValues;
@@ -264,11 +267,13 @@ void ULiveLinkBasicFrameInterpolationProcessor::FLiveLinkBasicFrameInterpolation
 			PropertiesResult[PropertyIndex] = FMath::Lerp(PropertiesA[PropertyIndex], PropertiesB[PropertyIndex], InBlendWeight);
 		}
 	}
-	else
+	// Copy the Property Values if they were not copied previously
+	else if (!Options.bCopyClosestFrame)
 	{
 		OutBlendedFrameData.GetBaseData()->PropertyValues = FrameWhenCanNotBlend.GetBaseData()->PropertyValues;
 	}
 
+	// Interpolate all properties with the tag "interp"
 	if (Options.bInterpolateInterpProperties)
 	{
 		LiveLinkInterpolation::Interpolate(FrameDataA.GetStruct(), true, InBlendWeight, FrameDataA.GetBaseData(), FrameDataB.GetBaseData(), OutBlendedFrameData.GetBaseData());
