- Now lerpin SceneTime from LiveLink frames when interpolation is used

diff --git a/Engine/Plugins/Animation/LiveLink/Source/LiveLink/Private/InterpolationProcessor/LiveLinkBasicFrameInterpolateProcessor.cpp b/Engine/Plugins/Animation/LiveLink/Source/LiveLink/Private/InterpolationProcessor/LiveLinkBasicFrameInterpolateProcessor.cpp
--- a/Engine/Plugins/Animation/LiveLink/Source/LiveLink/Private/InterpolationProcessor/LiveLinkBasicFrameInterpolateProcessor.cpp
+++ b/Engine/Plugins/Animation/LiveLink/Source/LiveLink/Private/InterpolationProcessor/LiveLinkBasicFrameInterpolateProcessor.cpp
@@ -249,8 +249,9 @@ void ULiveLinkBasicFrameInterpolationProcessor::FLiveLinkBasicFrameInterpolation
 		}
 	}
 
-	// Interpolate basic data
+	// Interpolate Time data
 	OutBlendedFrameData.GetBaseData()->WorldTime = FLiveLinkWorldTime(FMath::Lerp(FrameDataA.GetBaseData()->WorldTime.GetOffsettedTime(), FrameDataB.GetBaseData()->WorldTime.GetOffsettedTime(), InBlendWeight), 0.0);
+	OutBlendedFrameData.GetBaseData()->MetaData.SceneTime.Time = FFrameTime(FMath::Lerp(FrameDataA.GetBaseData()->MetaData.SceneTime.Time, FrameDataB.GetBaseData()->MetaData.SceneTime.Time, InBlendWeight));
 
 	// Interpolate Property Values
 	if (Options.bInterpolatePropertyValues)
diff --git a/Engine/Plugins/Animation/LiveLink/Source/LiveLink/Public/InterpolationProcessor/LiveLinkBasicFrameInterpolateProcessor.h b/Engine/Plugins/Animation/LiveLink/Source/LiveLink/Public/InterpolationProcessor/LiveLinkBasicFrameInterpolateProcessor.h
--- a/Engine/Plugins/Animation/LiveLink/Source/LiveLink/Public/InterpolationProcessor/LiveLinkBasicFrameInterpolateProcessor.h
+++ b/Engine/Plugins/Animation/LiveLink/Source/LiveLink/Public/InterpolationProcessor/LiveLinkBasicFrameInterpolateProcessor.h
@@ -28,7 +28,7 @@ public:
 		{
 			bool bInterpolatePropertyValues = true;
 			bool bCopyClosestFrame = true;
-			bool bCopyClosestMetaData = true; // only used if bCopyClosestFrame is false
+			bool bCopyClosestMetaData = true; // only used if bCopyClosestFrame is false. Does NOT apply to SceneTime, it will always be interpolated.
 			bool bInterpolateInterpProperties = true;
 		};
 
diff --git a/Engine/Source/Runtime/Core/Public/Misc/FrameTime.h b/Engine/Source/Runtime/Core/Public/Misc/FrameTime.h
--- a/Engine/Source/Runtime/Core/Public/Misc/FrameTime.h
+++ b/Engine/Source/Runtime/Core/Public/Misc/FrameTime.h
@@ -268,6 +268,10 @@ FORCEINLINE FFrameTime operator*(FFrameTime A, float Scalar)
 	return FFrameTime::FromDecimal(A.AsDecimal() * Scalar);
 }
 
+FORCEINLINE FFrameTime operator*(float Scalar, FFrameTime A)
+{
+	return FFrameTime::FromDecimal(A.AsDecimal() * Scalar);
+}
 
 FORCEINLINE FFrameTime operator/(FFrameTime A, float Scalar)
 {
