Timecodesync

- Adding a test map and an actor to validate synchronization of incoming data through livelink and sdi. Locations and colors are used to validate expected data sent from TimeSynchronizationWriter

diff --git a/Engine/Plugins/Animation/LiveLink/Source/LiveLink/Private/LiveLinkClient.cpp b/Engine/Plugins/Animation/LiveLink/Source/LiveLink/Private/LiveLinkClient.cpp
--- a/Engine/Plugins/Animation/LiveLink/Source/LiveLink/Private/LiveLinkClient.cpp
+++ b/Engine/Plugins/Animation/LiveLink/Source/LiveLink/Private/LiveLinkClient.cpp
@@ -1257,14 +1257,11 @@ void FLiveLinkClient::OnPropertyChanged(FGuid InEntryGuid, const FPropertyChange
 	}
 }
 
-ULiveLinkSourceSettings* FLiveLinkClient::GetSourceSettings(FGuid InEntryGuid) const
+ULiveLinkSourceSettings* FLiveLinkClient::GetSourceSettings(const FGuid& InEntryGuid) const
 {
 	if (const FLiveLinkCollectionSourceItem* SourceItem = Collection->FindSource(InEntryGuid))
 	{
-		if (!SourceItem->IsVirtualSource())
-		{
-			return SourceItem->Setting;
-		}
+		return SourceItem->Setting;
 	}
 	return nullptr;
 }
diff --git a/Engine/Plugins/Animation/LiveLink/Source/LiveLink/Public/LiveLinkClient.h b/Engine/Plugins/Animation/LiveLink/Source/LiveLink/Public/LiveLinkClient.h
--- a/Engine/Plugins/Animation/LiveLink/Source/LiveLink/Public/LiveLinkClient.h
+++ b/Engine/Plugins/Animation/LiveLink/Source/LiveLink/Public/LiveLinkClient.h
@@ -113,6 +113,7 @@ public:
 	virtual bool DoesSubjectSupportsRole(const FLiveLinkSubjectKey& SubjectKey, TSubclassOf<ULiveLinkRole> SupportedRole) const override;
 	virtual TArray<FLiveLinkTime> GetSubjectFrameTimes(const FLiveLinkSubjectKey& SubjectKey) const override;
 	virtual TArray<FLiveLinkTime> GetSubjectFrameTimes(FLiveLinkSubjectName SubjectName) const override;
+	virtual ULiveLinkSourceSettings* GetSourceSettings(const FGuid& SourceGuid) const override;
 	virtual UObject* GetSubjectSettings(const FLiveLinkSubjectKey& SubjectKey) const override;
 
 	virtual bool EvaluateFrameFromSource_AnyThread(const FLiveLinkSubjectKey& SubjectKey, TSubclassOf<ULiveLinkRole> Role, FLiveLinkSubjectFrameData& OutFrame) override;
@@ -149,8 +150,6 @@ public:
 
 	TArray<FGuid> GetDisplayableSources() const;
 
-	ULiveLinkSourceSettings* GetSourceSettings(FGuid InEntryGuid) const;
-
 	FLiveLinkSubjectTimeSyncData GetTimeSyncData(FLiveLinkSubjectName SubjectName);
 
 	FText GetSourceMachineName(FGuid EntryGuid) const;
diff --git a/Engine/Source/Runtime/LiveLinkInterface/Public/ILiveLinkClient.h b/Engine/Source/Runtime/LiveLinkInterface/Public/ILiveLinkClient.h
--- a/Engine/Source/Runtime/LiveLinkInterface/Public/ILiveLinkClient.h
+++ b/Engine/Source/Runtime/LiveLinkInterface/Public/ILiveLinkClient.h
@@ -238,6 +238,11 @@ public:
 	 */
 	virtual TArray<FLiveLinkTime> GetSubjectFrameTimes(FLiveLinkSubjectName SubjectName) const = 0;
 
+	/**
+	 * Get the Settings of this source.
+	 */
+	virtual ULiveLinkSourceSettings* GetSourceSettings(const FGuid& SourceKey) const = 0;
+
 	/**
 	 * Get the Settings of this subject.
 	 * @note If subject is a VirtualSubject, the VirtualSubject itself is returned.
