- Removing casts to FLiveLinkClient
- Fixed uninitialized variable

diff --git a/Engine/Plugins/Animation/LiveLink/Source/LiveLink/Private/LiveLinkMessageBusFinder.cpp b/Engine/Plugins/Animation/LiveLink/Source/LiveLink/Private/LiveLinkMessageBusFinder.cpp
--- a/Engine/Plugins/Animation/LiveLink/Source/LiveLink/Private/LiveLinkMessageBusFinder.cpp
+++ b/Engine/Plugins/Animation/LiveLink/Source/LiveLink/Private/LiveLinkMessageBusFinder.cpp
@@ -4,7 +4,7 @@
 
 #include "Engine/Engine.h"
 #include "Features/IModularFeatures.h"
-#include "LiveLinkClient.h"
+#include "ILiveLinkClient.h"
 #include "LiveLinkMessages.h"
 #include "LiveLinkMessageBusSource.h"
 #include "LiveLinkMessageBusSourceFactory.h"
@@ -90,7 +90,7 @@ void ULiveLinkMessageBusFinder::ConnectToProvider(UPARAM(ref) FProviderPollResul
 
 	if (ModularFeatures.IsModularFeatureAvailable(ILiveLinkClient::ModularFeatureName))
 	{
-		FLiveLinkClient* LiveLinkClient = &ModularFeatures.GetModularFeature<FLiveLinkClient>(ILiveLinkClient::ModularFeatureName);
+		ILiveLinkClient* LiveLinkClient = &ModularFeatures.GetModularFeature<ILiveLinkClient>(ILiveLinkClient::ModularFeatureName);
 		TSharedPtr<FLiveLinkMessageBusSource> NewSource = MakeShared<FLiveLinkMessageBusSource>(FText::FromString(Provider.Name), FText::FromString(Provider.MachineName), Provider.Address, Provider.MachineTimeOffset);
 		FGuid NewSourceGuid = LiveLinkClient->AddSource(NewSource);
 		if (NewSourceGuid.IsValid())
diff --git a/Engine/Plugins/Experimental/VirtualProductionUtilities/Source/VPUtilities/Private/VPFullScreenUserWidget.cpp b/Engine/Plugins/Experimental/VirtualProductionUtilities/Source/VPUtilities/Private/VPFullScreenUserWidget.cpp
--- a/Engine/Plugins/Experimental/VirtualProductionUtilities/Source/VPUtilities/Private/VPFullScreenUserWidget.cpp
+++ b/Engine/Plugins/Experimental/VirtualProductionUtilities/Source/VPUtilities/Private/VPFullScreenUserWidget.cpp
@@ -221,6 +221,7 @@ FVPFullScreenUserWidget_PostProcess::FVPFullScreenUserWidget_PostProcess()
 	: PostProcessMaterial(nullptr)
 	, PostProcessTintColorAndOpacity(FLinearColor::White)
 	, PostProcessOpacityFromTexture(1.0f)
+	, bWidgetDrawSize(false)
 	, WidgetDrawSize(FIntPoint(640, 360))
 	, bWindowFocusable(true)
 	, WindowVisibility(EWindowVisibility::SelfHitTestInvisible)
