Fix the registration of a LiveLinkClientReference.

diff --git a/Engine/Plugins/Animation/LiveLink/Source/LiveLink/Private/LiveLinkClientReference.cpp b/Engine/Plugins/Animation/LiveLink/Source/LiveLink/Private/LiveLinkClientReference.cpp
--- a/Engine/Plugins/Animation/LiveLink/Source/LiveLink/Private/LiveLinkClientReference.cpp
+++ b/Engine/Plugins/Animation/LiveLink/Source/LiveLink/Private/LiveLinkClientReference.cpp
@@ -15,6 +15,23 @@ FLiveLinkClientReference::FLiveLinkClientReference()
 	InitClient();
 }
 
+FLiveLinkClientReference::FLiveLinkClientReference(const FLiveLinkClientReference& Other)
+	: LiveLinkClient(Other.LiveLinkClient)
+{
+	IModularFeatures& ModularFeatures = IModularFeatures::Get();
+
+	ModularFeatures.OnModularFeatureRegistered().AddRaw(this, &FLiveLinkClientReference::OnLiveLinkClientRegistered);
+	ModularFeatures.OnModularFeatureUnregistered().AddRaw(this, &FLiveLinkClientReference::OnLiveLinkClientUnregistered);
+
+	InitClient();
+}
+
+FLiveLinkClientReference& FLiveLinkClientReference::operator=(const FLiveLinkClientReference& Other)
+{
+	LiveLinkClient = Other.LiveLinkClient;
+	return *this;
+}
+
 FLiveLinkClientReference::~FLiveLinkClientReference()
 {
 	IModularFeatures& ModularFeatures = IModularFeatures::Get();
diff --git a/Engine/Plugins/Animation/LiveLink/Source/LiveLink/Public/LiveLinkClientReference.h b/Engine/Plugins/Animation/LiveLink/Source/LiveLink/Public/LiveLinkClientReference.h
--- a/Engine/Plugins/Animation/LiveLink/Source/LiveLink/Public/LiveLinkClientReference.h
+++ b/Engine/Plugins/Animation/LiveLink/Source/LiveLink/Public/LiveLinkClientReference.h
@@ -13,6 +13,9 @@ struct LIVELINK_API FLiveLinkClientReference
 public:
 
 	FLiveLinkClientReference();
+	FLiveLinkClientReference(const FLiveLinkClientReference& Other);
+	FLiveLinkClientReference& operator=(const FLiveLinkClientReference& Other);
+	FLiveLinkClientReference(const FLiveLinkClientReference&& Other) = delete;
 	~FLiveLinkClientReference();
 
 	ILiveLinkClient* GetClient() const { return LiveLinkClient; }
diff --git a/Engine/Plugins/Animation/LiveLink/Source/LiveLinkEditor/Private/LiveLinkClientPanel.cpp b/Engine/Plugins/Animation/LiveLink/Source/LiveLinkEditor/Private/LiveLinkClientPanel.cpp
--- a/Engine/Plugins/Animation/LiveLink/Source/LiveLinkEditor/Private/LiveLinkClientPanel.cpp
+++ b/Engine/Plugins/Animation/LiveLink/Source/LiveLinkEditor/Private/LiveLinkClientPanel.cpp
@@ -211,27 +211,32 @@ public:
 		{
 			if (EntryPtr->IsVirtualSubject())
 			{
-				return	SNew(SButton)
-						.ButtonStyle(FEditorStyle::Get(), "HoverHintOnly")
-						.HAlign(HAlign_Center)
-						.VAlign(VAlign_Center)
-						.OnClicked(this, &SLiveLinkClientPanelSubjectRow::OnRemoveClicked)
-						.ToolTipText(LOCTEXT("RemoveVirtualSubject", "Remove selected live link virtual subject"))
-						.ContentPadding(0.f)
-						.ForegroundColor(FSlateColor::UseForeground())
-						.IsFocusable(false)
-						[
-							SNew(SImage)
-							.Image(FEditorStyle::GetBrush("PropertyWindow.Button_EmptyArray"))
-							.ColorAndOpacity(FSlateColor::UseForeground())
-						];
+				return SNew(SButton)
+					.ButtonStyle(FEditorStyle::Get(), "HoverHintOnly")
+					.HAlign(HAlign_Center)
+					.VAlign(VAlign_Center)
+					.OnClicked(this, &SLiveLinkClientPanelSubjectRow::OnRemoveClicked)
+					.ToolTipText(LOCTEXT("RemoveVirtualSubject", "Remove selected live link virtual subject"))
+					.ContentPadding(0.f)
+					.ForegroundColor(FSlateColor::UseForeground())
+					.IsFocusable(false)
+					[
+						SNew(SImage)
+						.Image(FEditorStyle::GetBrush("PropertyWindow.Button_EmptyArray"))
+						.ColorAndOpacity(FSlateColor::UseForeground())
+					];
 			}
 			else
 			{
-				return SNew(STextBlock)
-					.Font(FEditorStyle::Get().GetFontStyle("FontAwesome.8"))
-					.ColorAndOpacity(this, &SLiveLinkClientPanelSubjectRow::OnGetActivityColor)
-					.Text(FEditorFontGlyphs::Circle);
+				return SNew(SBox)
+					.HAlign(HAlign_Center)
+					.VAlign(VAlign_Center)
+					[
+						SNew(STextBlock)
+						.Font(FEditorStyle::Get().GetFontStyle("FontAwesome.8"))
+						.ColorAndOpacity(this, &SLiveLinkClientPanelSubjectRow::OnGetActivityColor)
+						.Text(FEditorFontGlyphs::Circle)
+					];
 			}
 		}
 
