Take Recorder: Various fixes for recording livelink data.

Added missing GetMovieSceneSection() to LiveLinkTrackRecorder so that SetSectionStartTimecode is called properly.
RecordFrame should ExpandToFrame the exact frame, not the lower bound + frame.
Fix drawing collapsed keys losing precision.

diff --git a/Engine/Plugins/Animation/LiveLink/Source/LiveLinkMovieScene/Private/MovieScene/MovieSceneLiveLinkSection.cpp b/Engine/Plugins/Animation/LiveLink/Source/LiveLinkMovieScene/Private/MovieScene/MovieSceneLiveLinkSection.cpp
--- a/Engine/Plugins/Animation/LiveLink/Source/LiveLinkMovieScene/Private/MovieScene/MovieSceneLiveLinkSection.cpp
+++ b/Engine/Plugins/Animation/LiveLink/Source/LiveLinkMovieScene/Private/MovieScene/MovieSceneLiveLinkSection.cpp
@@ -35,7 +35,7 @@ void UMovieSceneLiveLinkSection::SetMask(const TArray<bool>& InChannelMask)
 
 void UMovieSceneLiveLinkSection::RecordFrame(FFrameNumber InFrameNumber, const FLiveLinkFrameDataStruct& InFrameData)
 {
-	ExpandToFrame(GetInclusiveStartFrame() + InFrameNumber);
+	ExpandToFrame(InFrameNumber);
 
 	for (UMovieSceneLiveLinkSubSection* SubSection : SubSections)
 	{
diff --git a/Engine/Plugins/Animation/LiveLink/Source/LiveLinkSequencer/Public/MovieSceneLiveLinkTrackRecorder.h b/Engine/Plugins/Animation/LiveLink/Source/LiveLinkSequencer/Public/MovieSceneLiveLinkTrackRecorder.h
--- a/Engine/Plugins/Animation/LiveLink/Source/LiveLinkSequencer/Public/MovieSceneLiveLinkTrackRecorder.h
+++ b/Engine/Plugins/Animation/LiveLink/Source/LiveLinkSequencer/Public/MovieSceneLiveLinkTrackRecorder.h
@@ -11,6 +11,7 @@
 #include "Misc/Guid.h"
 #include "MovieScene.h"
 #include "TrackRecorders/IMovieSceneTrackRecorderFactory.h"
+#include "MovieScene/MovieSceneLiveLinkSection.h"
 
 #include "MovieSceneLiveLinkTrackRecorder.generated.h"
 
@@ -35,6 +36,7 @@ public:
 	virtual void RecordSampleImpl(const FQualifiedFrameTime& CurrentTime) override;
 	virtual void FinalizeTrackImpl() override;
 	virtual void SetSectionStartTimecodeImpl(const FTimecode& InSectionStartTimecode, const FFrameNumber& InSectionFirstFrame) override;
+	virtual UMovieSceneSection* GetMovieSceneSection() const override { return Cast<UMovieSceneSection>(MovieSceneSection.Get()); }
 	virtual void StopRecordingImpl() override;
 	virtual void SetSavedRecordingDirectory(const FString& InDirectory)
 	{
diff --git a/Engine/Source/Editor/Sequencer/Private/SSequencerSection.cpp b/Engine/Source/Editor/Sequencer/Private/SSequencerSection.cpp
--- a/Engine/Source/Editor/Sequencer/Private/SSequencerSection.cpp
+++ b/Engine/Source/Editor/Sequencer/Private/SSequencerSection.cpp
@@ -1411,7 +1411,7 @@ void SSequencerSection::PaintKeys( FSequencerSectionPainter& InPainter, const FW
 			// Generate the hull of frame numbers that contribute to this key so we can draw it enabled/disabled depending on whether it is outside of the valid range or not
 			TRange<FFrameNumber> KeyRange = TRange<FFrameNumber>::Empty();
 
-			float AverageKeyTime = 0.f;
+			double AverageKeyTime = 0.f;
 			int32 NumKeyTimes = 0;
 			int32 NumOverlaps = 0;
 
