Live Link: Discard recording samples if outside of the section start frame

diff --git a/Engine/Plugins/Animation/LiveLink/Source/LiveLinkSequencer/Private/SequencerRecorderSections/MovieSceneLiveLinkTrackRecorder.cpp b/Engine/Plugins/Animation/LiveLink/Source/LiveLinkSequencer/Private/SequencerRecorderSections/MovieSceneLiveLinkTrackRecorder.cpp
--- a/Engine/Plugins/Animation/LiveLink/Source/LiveLinkSequencer/Private/SequencerRecorderSections/MovieSceneLiveLinkTrackRecorder.cpp
+++ b/Engine/Plugins/Animation/LiveLink/Source/LiveLinkSequencer/Private/SequencerRecorderSections/MovieSceneLiveLinkTrackRecorder.cpp
@@ -252,7 +252,16 @@ void UMovieSceneLiveLinkTrackRecorder::RecordSampleImpl(const FQualifiedFrameTim
 					FrameNumber += RecordStartFrame;
 				}
 
-				MovieSceneSection->RecordFrame(FrameNumber, Frame);
+				// For clarity, only record values that are after the start frame since frames could have been buffered before recording started.
+				if (FrameNumber >= MovieSceneSection->GetInclusiveStartFrame())
+				{
+					MovieSceneSection->RecordFrame(FrameNumber, Frame);
+				}
+				else
+				{
+					const FFrameRate DisplayRate = MovieSceneSection->GetTypedOuter<UMovieScene>()->GetDisplayRate();
+					UE_LOG(LogLiveLinkSequencer, Warning, TEXT("Discarded buffered frame: %f outside of start frame: %f for subject '%s'."), ConvertFrameTime(FrameNumber, TickResolution, DisplayRate).AsDecimal(), ConvertFrameTime(MovieSceneSection->GetInclusiveStartFrame(), TickResolution, DisplayRate).AsDecimal(), *(SubjectName.ToString()));
+				}
 			}
 
 			//Empty out frames that were processed
