Live Link: Added veryverbose log for injested live link frame times

Enable with console command: log LogLiveLinkSequencer VeryVerbose

Sammple output:

LogLiveLinkSequencer: VeryVerbose: LiveLinkFrameTime: [22:28:30:06] at 30 fps for subject 'Character'.
LogLiveLinkSequencer: VeryVerbose: LiveLinkFrameTime: [22:28:30:06] at 30 fps for subject 'Character'.
L

diff --git a/Engine/Plugins/Animation/LiveLink/Source/LiveLinkSequencer/Private/SequencerRecorderSections/MovieSceneLiveLinkTrackRecorder.cpp b/Engine/Plugins/Animation/LiveLink/Source/LiveLinkSequencer/Private/SequencerRecorderSections/MovieSceneLiveLinkTrackRecorder.cpp
--- a/Engine/Plugins/Animation/LiveLink/Source/LiveLinkSequencer/Private/SequencerRecorderSections/MovieSceneLiveLinkTrackRecorder.cpp
+++ b/Engine/Plugins/Animation/LiveLink/Source/LiveLinkSequencer/Private/SequencerRecorderSections/MovieSceneLiveLinkTrackRecorder.cpp
@@ -216,7 +216,8 @@ void UMovieSceneLiveLinkTrackRecorder::RecordSampleImpl(const FQualifiedFrameTim
 	if (LiveLinkClient && MovieSceneSection.IsValid())
 	{
 		//we know all section have same tick resolution
-		const FFrameRate TickResolution = MovieSceneSection->GetTypedOuter<UMovieScene>()->GetTickResolution();
+		const FFrameRate TickResolution = MovieScene->GetTickResolution();
+		const FFrameRate DisplayRate = MovieScene->GetDisplayRate();
 		const FFrameNumber CurrentFrame = CurrentTime.ConvertTo(TickResolution).FloorToFrame();
 
 		bool bSyncedOrForced = bUseSourceTimecode || LiveLinkClient->IsSubjectTimeSynchronized(SubjectName);
@@ -244,12 +245,20 @@ void UMovieSceneLiveLinkTrackRecorder::RecordSampleImpl(const FQualifiedFrameTim
 
 					FFrameTime FrameTime = LiveLinkFrameTime.ConvertTo(TickResolution);
 					FrameNumber = FrameTime.FrameNumber;
+
+					UE_LOG(LogLiveLinkSequencer, VeryVerbose, TEXT("LiveLinkFrameTime: [%s.%.03f] at %s for subject '%s'."), 
+						   *(FTimecode::FromFrameNumber(LiveLinkFrameTime.Time.FrameNumber, LiveLinkFrameTime.Rate).ToString()), 
+						   LiveLinkFrameTime.Time.GetSubFrame(),
+						   *(LiveLinkFrameTime.Rate.ToPrettyText().ToString()),
+						   *(SubjectName.ToString()));
 				}
 				else
 				{
-					const double Second = Frame.GetBaseData()->WorldTime.GetOffsettedTime() - SecondsDiff;
+					const double Second = Frame.GetBaseData()->WorldTime.GetOffsettedTime() - SecondsDiff;				
 					FrameNumber = (Second * TickResolution).FloorToFrame();
 					FrameNumber += RecordStartFrame;
+
+					UE_LOG(LogLiveLinkSequencer, VeryVerbose, TEXT("LiveLinkFrameTime (Unsynced): %f, for subject '%s'."), Frame.GetBaseData()->WorldTime.GetOffsettedTime(), *(SubjectName.ToString()));
 				}
 
 				// For clarity, only record values that are after the start frame since frames could have been buffered before recording started.
@@ -259,8 +268,12 @@ void UMovieSceneLiveLinkTrackRecorder::RecordSampleImpl(const FQualifiedFrameTim
 				}
 				else
 				{
-					const FFrameRate DisplayRate = MovieSceneSection->GetTypedOuter<UMovieScene>()->GetDisplayRate();
-					UE_LOG(LogLiveLinkSequencer, Warning, TEXT("Discarded buffered frame: %f outside of start frame: %f for subject '%s'."), ConvertFrameTime(FrameNumber, TickResolution, DisplayRate).AsDecimal(), ConvertFrameTime(MovieSceneSection->GetInclusiveStartFrame(), TickResolution, DisplayRate).AsDecimal(), *(SubjectName.ToString()));
+					UE_LOG(LogLiveLinkSequencer, Warning, TEXT("Discarded buffered frame: [%s.%.03f] outside of start frame: [%s.%.03f] for subject '%s'."),
+						   *(FTimecode::FromFrameNumber(ConvertFrameTime(FrameNumber, TickResolution, DisplayRate).FrameNumber, DisplayRate).ToString()), 
+						   ConvertFrameTime(FrameNumber, TickResolution, DisplayRate).GetSubFrame(),
+						   *(FTimecode::FromFrameNumber(ConvertFrameTime(MovieSceneSection->GetInclusiveStartFrame(), TickResolution, DisplayRate).FrameNumber, DisplayRate)).ToString(), 
+						   ConvertFrameTime(MovieSceneSection->GetInclusiveStartFrame(), TickResolution, DisplayRate).GetSubFrame(),
+						   *(SubjectName.ToString()));
 				}
 			}
 
