Simplify access to the LiveLinkClient from LiveLinkClientReference.

diff --git a/Engine/Plugins/Animation/LiveLink/Source/LiveLink/Private/LiveLinkClientReference.cpp b/Engine/Plugins/Animation/LiveLink/Source/LiveLink/Private/LiveLinkClientReference.cpp
--- a/Engine/Plugins/Animation/LiveLink/Source/LiveLink/Private/LiveLinkClientReference.cpp
+++ b/Engine/Plugins/Animation/LiveLink/Source/LiveLink/Private/LiveLinkClientReference.cpp
@@ -1,68 +1,9 @@
 // Copyright Epic Games, Inc. All Rights Reserved.
 
 #include "LiveLinkClientReference.h"
-#include "Features/IModularFeatures.h"
-#include "ILiveLinkClient.h"
+#include "LiveLinkModule.h"
 
-FLiveLinkClientReference::FLiveLinkClientReference()
-	: LiveLinkClient(nullptr)
+ILiveLinkClient* FLiveLinkClientReference::GetClient()const
 {
-	IModularFeatures& ModularFeatures = IModularFeatures::Get();
-
-	ModularFeatures.OnModularFeatureRegistered().AddRaw(this, &FLiveLinkClientReference::OnLiveLinkClientRegistered);
-	ModularFeatures.OnModularFeatureUnregistered().AddRaw(this, &FLiveLinkClientReference::OnLiveLinkClientUnregistered);
-
-	InitClient();
-}
-
-FLiveLinkClientReference::FLiveLinkClientReference(const FLiveLinkClientReference& Other)
-	: LiveLinkClient(Other.LiveLinkClient)
-{
-	IModularFeatures& ModularFeatures = IModularFeatures::Get();
-
-	ModularFeatures.OnModularFeatureRegistered().AddRaw(this, &FLiveLinkClientReference::OnLiveLinkClientRegistered);
-	ModularFeatures.OnModularFeatureUnregistered().AddRaw(this, &FLiveLinkClientReference::OnLiveLinkClientUnregistered);
-
-	InitClient();
-}
-
-FLiveLinkClientReference& FLiveLinkClientReference::operator=(const FLiveLinkClientReference& Other)
-{
-	LiveLinkClient = Other.LiveLinkClient;
-	return *this;
-}
-
-FLiveLinkClientReference::~FLiveLinkClientReference()
-{
-	IModularFeatures& ModularFeatures = IModularFeatures::Get();
-
-	ModularFeatures.OnModularFeatureRegistered().RemoveAll(this);
-	ModularFeatures.OnModularFeatureUnregistered().RemoveAll(this);
-}
-
-void FLiveLinkClientReference::InitClient()
-{
-	IModularFeatures& ModularFeatures = IModularFeatures::Get();
-
-	if (ModularFeatures.IsModularFeatureAvailable(ILiveLinkClient::ModularFeatureName))
-	{
-		LiveLinkClient = &IModularFeatures::Get().GetModularFeature<ILiveLinkClient>(ILiveLinkClient::ModularFeatureName);
-	}
-}
-
-void FLiveLinkClientReference::OnLiveLinkClientRegistered(const FName& Type, class IModularFeature* ModularFeature)
-{
-	if (Type == ILiveLinkClient::ModularFeatureName && !LiveLinkClient)
-	{
-		LiveLinkClient = static_cast<ILiveLinkClient*>(ModularFeature);
-	}
-}
-
-void FLiveLinkClientReference::OnLiveLinkClientUnregistered(const FName& Type, class IModularFeature* ModularFeature)
-{
-	if (Type == ILiveLinkClient::ModularFeatureName && ModularFeature == LiveLinkClient)
-	{
-		LiveLinkClient = nullptr;
-		InitClient();
-	}
+	return FLiveLinkModule::StaticLiveLinkClient;
 }
\ No newline at end of file
diff --git a/Engine/Plugins/Animation/LiveLink/Source/LiveLink/Private/LiveLinkModule.cpp b/Engine/Plugins/Animation/LiveLink/Source/LiveLink/Private/LiveLinkModule.cpp
--- a/Engine/Plugins/Animation/LiveLink/Source/LiveLink/Private/LiveLinkModule.cpp
+++ b/Engine/Plugins/Animation/LiveLink/Source/LiveLink/Private/LiveLinkModule.cpp
@@ -1,112 +1,63 @@
 // Copyright Epic Games, Inc. All Rights Reserved.
 
-#include "ILiveLinkModule.h"
+#include "LiveLinkModule.h"
 
-#include "Features/IModularFeatures.h"
-#include "Interfaces/IPluginManager.h"
-#include "LiveLinkMotionController.h"
-
-#include "LiveLinkMessageBusDiscoveryManager.h"
-
-#include "LiveLinkClient.h"
+#include "HAL/PlatformAtomics.h"
 #include "LiveLinkLogInstance.h"
-#include "LiveLinkDebugCommand.h"
-#include "LiveLinkHeartbeatEmitter.h"
 #include "LiveLinkPreset.h"
 #include "LiveLinkSettings.h"
 
-#include "Styling/SlateStyle.h"
-
-
-/**
- * Implements the Messaging module.
- */
-
 #define LOCTEXT_NAMESPACE "LiveLinkModule"
 
+FLiveLinkClient* FLiveLinkModule::StaticLiveLinkClient = nullptr;
 
-class FLiveLinkModule
-	: public ILiveLinkModule
+FLiveLinkModule::FLiveLinkModule()
+	: LiveLinkClient()
+	, LiveLinkMotionController(LiveLinkClient)
+	, HeartbeatEmitter(MakeUnique<FLiveLinkHeartbeatEmitter>())
+	, DiscoveryManager(MakeUnique<FLiveLinkMessageBusDiscoveryManager>())
+	, LiveLinkDebugCommand(MakeUnique<FLiveLinkDebugCommand>(LiveLinkClient))
 {
-public:
-	FLiveLinkClient LiveLinkClient;
-	FLiveLinkMotionController LiveLinkMotionController;
-	TSharedPtr<FSlateStyleSet> StyleSet;
-	FString ContentDir;
-
-	FLiveLinkModule()
-		: LiveLinkClient()
-		, LiveLinkMotionController(LiveLinkClient)
-		, HeartbeatEmitter(MakeUnique<FLiveLinkHeartbeatEmitter>())
-		, DiscoveryManager(MakeUnique<FLiveLinkMessageBusDiscoveryManager>())
-		, LiveLinkDebugCommand(MakeUnique<FLiveLinkDebugCommand>(LiveLinkClient))
-	{
-		CreateStyle();
-	}
-
-	// IModuleInterface interface
-
-	virtual void StartupModule() override
-	{
-		FLiveLinkLogInstance::CreateInstance();
-		IModularFeatures::Get().RegisterModularFeature(FLiveLinkClient::ModularFeatureName, &LiveLinkClient);
-		LiveLinkMotionController.RegisterController();
-
-		if (ULiveLinkPreset* Preset = GetDefault<ULiveLinkSettings>()->DefaultLiveLinkPreset.LoadSynchronous())
-		{
-			Preset->ApplyToClient();
-		}
-	}
-
-	virtual void ShutdownModule() override
-	{
-		HeartbeatEmitter->Exit();
-		DiscoveryManager->Stop();
-		LiveLinkMotionController.UnregisterController();
-		IModularFeatures::Get().UnregisterModularFeature(FLiveLinkClient::ModularFeatureName, &LiveLinkClient);
-		FLiveLinkLogInstance::DestroyInstance();
-	}
+}
 
-	virtual bool SupportsDynamicReloading() override
-	{
-		return false;
-	}
+void FLiveLinkModule::StartupModule()
+{
+	FLiveLinkLogInstance::CreateInstance();
+	CreateStyle();
 
-	virtual TSharedPtr<FSlateStyleSet> GetStyle() override
-	{
-		return StyleSet;
-	}
+	StaticLiveLinkClient = &LiveLinkClient;
+	IModularFeatures::Get().RegisterModularFeature(FLiveLinkClient::ModularFeatureName, &LiveLinkClient);
+	LiveLinkMotionController.RegisterController();
 
-	virtual FLiveLinkHeartbeatEmitter& GetHeartbeatEmitter() override
+	if (ULiveLinkPreset* Preset = GetDefault<ULiveLinkSettings>()->DefaultLiveLinkPreset.LoadSynchronous())
 	{
-		return *HeartbeatEmitter;
+		Preset->ApplyToClient();
 	}
+}
 
-	virtual FLiveLinkMessageBusDiscoveryManager& GetMessageBusDiscoveryManager() override
-	{
-		return *DiscoveryManager;
-	}
-
-private:
-	void CreateStyle()
-	{
-		static FName LiveLinkStyle(TEXT("LiveLinkCoreStyle"));
-		StyleSet = MakeShared<FSlateStyleSet>(LiveLinkStyle);
+void FLiveLinkModule::ShutdownModule()
+{
+	HeartbeatEmitter->Exit();
+	DiscoveryManager->Stop();
+	LiveLinkMotionController.UnregisterController();
+	
+	IModularFeatures::Get().UnregisterModularFeature(FLiveLinkClient::ModularFeatureName, &LiveLinkClient);
+	StaticLiveLinkClient = nullptr;
 
-		ContentDir = IPluginManager::Get().FindPlugin(TEXT("LiveLink"))->GetContentDir();
+	FLiveLinkLogInstance::DestroyInstance();
+}
 
-		const FVector2D Icon16x16(16.0f, 16.0f);
+void FLiveLinkModule::CreateStyle()
+{
+	static FName LiveLinkStyle(TEXT("LiveLinkCoreStyle"));
+	StyleSet = MakeShared<FSlateStyleSet>(LiveLinkStyle);
 
-		StyleSet->Set("LiveLinkIcon", new FSlateImageBrush((ContentDir / TEXT("LiveLink_16x")) + TEXT(".png"), Icon16x16));
-	}
+	FString ContentDir = IPluginManager::Get().FindPlugin(TEXT("LiveLink"))->GetContentDir();
 
-private:
-	TUniquePtr<FLiveLinkHeartbeatEmitter> HeartbeatEmitter;
-	TUniquePtr<FLiveLinkMessageBusDiscoveryManager> DiscoveryManager;
+	const FVector2D Icon16x16(16.0f, 16.0f);
 
-	/** Handler for LiveLink debug command. */
-	TUniquePtr<FLiveLinkDebugCommand> LiveLinkDebugCommand;
-};
+	StyleSet->Set("LiveLinkIcon", new FSlateImageBrush((ContentDir / TEXT("LiveLink_16x")) + TEXT(".png"), Icon16x16));
+}
 
 IMPLEMENT_MODULE(FLiveLinkModule, LiveLink);
 
diff --git a/Engine/Plugins/Animation/LiveLink/Source/LiveLink/Private/LiveLinkModule.h b/Engine/Plugins/Animation/LiveLink/Source/LiveLink/Private/LiveLinkModule.h
new file mode 100644
--- /dev/null
+++ b/Engine/Plugins/Animation/LiveLink/Source/LiveLink/Private/LiveLinkModule.h
@@ -0,0 +1,57 @@
+// Copyright Epic Games, Inc. All Rights Reserved.
+
+#pragma once
+
+#include "ILiveLinkModule.h"
+
+#include "Features/IModularFeatures.h"
+#include "Interfaces/IPluginManager.h"
+#include "LiveLinkMotionController.h"
+
+#include "LiveLinkMessageBusDiscoveryManager.h"
+
+#include "LiveLinkClient.h"
+#include "LiveLinkDebugCommand.h"
+#include "LiveLinkHeartbeatEmitter.h"
+
+#include "Styling/SlateStyle.h"
+
+/**
+ * Implements the Messaging module.
+ */
+
+ struct FLiveLinkClientReference;
+
+class FLiveLinkModule : public ILiveLinkModule
+{
+public:
+	FLiveLinkModule();
+
+	//~ Begin IModuleInterface interface
+	virtual void StartupModule() override;
+	virtual void ShutdownModule() override;
+	virtual bool SupportsDynamicReloading() override { return false; }
+	//~ End IModuleInterface interface
+
+	//~ Begin ILiveLinkModule interface
+	virtual FLiveLinkHeartbeatEmitter& GetHeartbeatEmitter() override { return *HeartbeatEmitter; }
+	virtual FLiveLinkMessageBusDiscoveryManager& GetMessageBusDiscoveryManager() override { return *DiscoveryManager; }
+	virtual TSharedPtr<FSlateStyleSet> GetStyle() override { return StyleSet; }
+	//~ End ILiveLinkModule interface
+
+private:
+	void CreateStyle();
+
+private:
+	friend FLiveLinkClientReference;
+	static FLiveLinkClient* StaticLiveLinkClient;
+
+	FLiveLinkClient LiveLinkClient;
+	FLiveLinkMotionController LiveLinkMotionController;
+
+	TSharedPtr<FSlateStyleSet> StyleSet;
+
+	TUniquePtr<FLiveLinkHeartbeatEmitter> HeartbeatEmitter;
+	TUniquePtr<FLiveLinkMessageBusDiscoveryManager> DiscoveryManager;
+	TUniquePtr<FLiveLinkDebugCommand> LiveLinkDebugCommand;
+};
diff --git a/Engine/Plugins/Animation/LiveLink/Source/LiveLink/Public/LiveLinkClientReference.h b/Engine/Plugins/Animation/LiveLink/Source/LiveLink/Public/LiveLinkClientReference.h
--- a/Engine/Plugins/Animation/LiveLink/Source/LiveLink/Public/LiveLinkClientReference.h
+++ b/Engine/Plugins/Animation/LiveLink/Source/LiveLink/Public/LiveLinkClientReference.h
@@ -7,25 +7,9 @@
 class ILiveLinkClient;
 class IModularFeature;
 
-// References the live link client modular feature and handles add/remove
+/** References the live link client */
 struct LIVELINK_API FLiveLinkClientReference
 {
 public:
-
-	FLiveLinkClientReference();
-	FLiveLinkClientReference(const FLiveLinkClientReference& Other);
-	FLiveLinkClientReference& operator=(const FLiveLinkClientReference& Other);
-	FLiveLinkClientReference(const FLiveLinkClientReference&& Other) = delete;
-	~FLiveLinkClientReference();
-
-	ILiveLinkClient* GetClient() const { return LiveLinkClient; }
-
-private:
-	void InitClient();
-
-	// Handlers for modular features coming and going
-	void OnLiveLinkClientRegistered(const FName& Type, class IModularFeature* ModularFeature);
-	void OnLiveLinkClientUnregistered(const FName& Type, class IModularFeature* ModularFeature);
-
-	ILiveLinkClient* LiveLinkClient;
+	ILiveLinkClient* GetClient() const;
 };
\ No newline at end of file
