Simplify access to the LiveLinkClient from LiveLinkClientReference.

diff --git a/Engine/Plugins/Animation/LiveLink/Source/LiveLink/Private/LiveLinkClientReference.cpp b/Engine/Plugins/Animation/LiveLink/Source/LiveLink/Private/LiveLinkClientReference.cpp
--- a/Engine/Plugins/Animation/LiveLink/Source/LiveLink/Private/LiveLinkClientReference.cpp
+++ b/Engine/Plugins/Animation/LiveLink/Source/LiveLink/Private/LiveLinkClientReference.cpp
@@ -1,6 +1,7 @@
 // Copyright Epic Games, Inc. All Rights Reserved.
 
 #include "LiveLinkClientReference.h"
+#include "ILiveLinkClient.h"
 #include "LiveLinkModule.h"
 
 ILiveLinkClient* FLiveLinkClientReference::GetClient()const
diff --git a/Engine/Plugins/Animation/LiveLink/Source/LiveLink/Private/LiveLinkModule.cpp b/Engine/Plugins/Animation/LiveLink/Source/LiveLink/Private/LiveLinkModule.cpp
--- a/Engine/Plugins/Animation/LiveLink/Source/LiveLink/Private/LiveLinkModule.cpp
+++ b/Engine/Plugins/Animation/LiveLink/Source/LiveLink/Private/LiveLinkModule.cpp
@@ -2,7 +2,6 @@
 
 #include "LiveLinkModule.h"
 
-#include "HAL/PlatformAtomics.h"
 #include "LiveLinkLogInstance.h"
 #include "LiveLinkPreset.h"
 #include "LiveLinkSettings.h"
diff --git a/Engine/Plugins/Animation/LiveLink/Source/LiveLink/Private/LiveLinkModule.h b/Engine/Plugins/Animation/LiveLink/Source/LiveLink/Private/LiveLinkModule.h
--- a/Engine/Plugins/Animation/LiveLink/Source/LiveLink/Private/LiveLinkModule.h
+++ b/Engine/Plugins/Animation/LiveLink/Source/LiveLink/Private/LiveLinkModule.h
@@ -5,15 +5,11 @@
 #include "ILiveLinkModule.h"
 
 #include "Features/IModularFeatures.h"
-#include "Interfaces/IPluginManager.h"
-#include "LiveLinkMotionController.h"
-
-#include "LiveLinkMessageBusDiscoveryManager.h"
-
 #include "LiveLinkClient.h"
 #include "LiveLinkDebugCommand.h"
 #include "LiveLinkHeartbeatEmitter.h"
-
+#include "LiveLinkMessageBusDiscoveryManager.h"
+#include "LiveLinkMotionController.h"
 #include "Styling/SlateStyle.h"
 
 /**
@@ -54,4 +50,4 @@ private:
 	TUniquePtr<FLiveLinkHeartbeatEmitter> HeartbeatEmitter;
 	TUniquePtr<FLiveLinkMessageBusDiscoveryManager> DiscoveryManager;
 	TUniquePtr<FLiveLinkDebugCommand> LiveLinkDebugCommand;
-};
+};
\ No newline at end of file
diff --git a/Engine/Plugins/Animation/LiveLink/Source/LiveLink/Public/LiveLinkClientReference.h b/Engine/Plugins/Animation/LiveLink/Source/LiveLink/Public/LiveLinkClientReference.h
--- a/Engine/Plugins/Animation/LiveLink/Source/LiveLink/Public/LiveLinkClientReference.h
+++ b/Engine/Plugins/Animation/LiveLink/Source/LiveLink/Public/LiveLinkClientReference.h
@@ -5,7 +5,6 @@
 #include "CoreMinimal.h"
 
 class ILiveLinkClient;
-class IModularFeature;
 
 /** References the live link client */
 struct LIVELINK_API FLiveLinkClientReference
