Fix compile error with LiveLink plugin
#rnx

diff --git a/Engine/Plugins/Animation/LiveLink/Source/LiveLink/Private/LiveLinkModule.cpp b/Engine/Plugins/Animation/LiveLink/Source/LiveLink/Private/LiveLinkModule.cpp
--- a/Engine/Plugins/Animation/LiveLink/Source/LiveLink/Private/LiveLinkModule.cpp
+++ b/Engine/Plugins/Animation/LiveLink/Source/LiveLink/Private/LiveLinkModule.cpp
@@ -5,6 +5,7 @@
 #include "LiveLinkLogInstance.h"
 #include "LiveLinkPreset.h"
 #include "LiveLinkSettings.h"
+#include "Interfaces/IPluginManager.h"
 
 #define LOCTEXT_NAMESPACE "LiveLinkModule"
 
