LiveLink: Added option to discard livelink samples if their timecode occurs before the start of recording. Defaults to true.

#rnx

diff --git a/Engine/Plugins/Animation/LiveLink/Source/LiveLinkSequencer/Private/SequencerRecorderSections/MovieSceneLiveLinkTrackRecorder.cpp b/Engine/Plugins/Animation/LiveLink/Source/LiveLinkSequencer/Private/SequencerRecorderSections/MovieSceneLiveLinkTrackRecorder.cpp
--- a/Engine/Plugins/Animation/LiveLink/Source/LiveLinkSequencer/Private/SequencerRecorderSections/MovieSceneLiveLinkTrackRecorder.cpp
+++ b/Engine/Plugins/Animation/LiveLink/Source/LiveLinkSequencer/Private/SequencerRecorderSections/MovieSceneLiveLinkTrackRecorder.cpp
@@ -37,12 +37,13 @@ static void AlwaysRecordLiveLinkTimecodeSinkFunction()
 
 static FAutoConsoleVariableSink CVarAlwaysRecoredLiveLinkSink(FConsoleCommandDelegate::CreateStatic(&AlwaysRecordLiveLinkTimecodeSinkFunction));
 
-void UMovieSceneLiveLinkTrackRecorder::CreateTrack(UMovieScene* InMovieScene, const FName& InSubjectName, bool bInSaveSubjectSettings, bool bInAlwaysUseTimecoe, UMovieSceneTrackRecorderSettings* InSettingsObject)
+void UMovieSceneLiveLinkTrackRecorder::CreateTrack(UMovieScene* InMovieScene, const FName& InSubjectName, bool bInSaveSubjectSettings, bool bInAlwaysUseTimecode, bool bInDiscardSamplesBeforeStart, UMovieSceneTrackRecorderSettings* InSettingsObject)
 {
 	MovieScene = InMovieScene;
 	SubjectName = InSubjectName;
 	bSaveSubjectSettings = bInSaveSubjectSettings;
-	bUseSourceTimecode = bInAlwaysUseTimecoe;
+	bUseSourceTimecode = bInAlwaysUseTimecode;
+	bDiscardSamplesBeforeStart = bInDiscardSamplesBeforeStart;
 	CreateTracks();
 }
 
@@ -262,7 +263,7 @@ void UMovieSceneLiveLinkTrackRecorder::RecordSampleImpl(const FQualifiedFrameTim
 				}
 
 				// For clarity, only record values that are after the start frame since frames could have been buffered before recording started.
-				if (FrameNumber >= MovieSceneSection->GetInclusiveStartFrame())
+				if (FrameNumber >= MovieSceneSection->GetInclusiveStartFrame() || !bDiscardSamplesBeforeStart)
 				{
 					MovieSceneSection->RecordFrame(FrameNumber, Frame);
 				}
diff --git a/Engine/Plugins/Animation/LiveLink/Source/LiveLinkSequencer/Private/TakeRecorderSource/TakeRecorderLiveLinkSource.cpp b/Engine/Plugins/Animation/LiveLink/Source/LiveLinkSequencer/Private/TakeRecorderSource/TakeRecorderLiveLinkSource.cpp
--- a/Engine/Plugins/Animation/LiveLink/Source/LiveLinkSequencer/Private/TakeRecorderSource/TakeRecorderLiveLinkSource.cpp
+++ b/Engine/Plugins/Animation/LiveLink/Source/LiveLinkSequencer/Private/TakeRecorderSource/TakeRecorderLiveLinkSource.cpp
@@ -24,6 +24,7 @@ UTakeRecorderLiveLinkSource::UTakeRecorderLiveLinkSource(const FObjectInitialize
 	, SubjectName(NAME_None)
 	, bSaveSubjectSettings(true)
 	, bUseSourceTimecode(true)
+	, bDiscardSamplesBeforeStart(true)
 {
 	TrackTint = FColor(74, 108, 164);
 }
@@ -32,7 +33,7 @@ TArray<UTakeRecorderSource*> UTakeRecorderLiveLinkSource::PreRecording(class ULe
 {
 	UMovieScene* MovieScene = InSequence->GetMovieScene();
 	TrackRecorder = NewObject<UMovieSceneLiveLinkTrackRecorder>();
-	TrackRecorder->CreateTrack(MovieScene, SubjectName, bSaveSubjectSettings, bUseSourceTimecode, nullptr);
+	TrackRecorder->CreateTrack(MovieScene, SubjectName, bSaveSubjectSettings, bUseSourceTimecode, bDiscardSamplesBeforeStart, nullptr);
 
 	return TArray<UTakeRecorderSource*>();
 }
diff --git a/Engine/Plugins/Animation/LiveLink/Source/LiveLinkSequencer/Private/TakeRecorderSource/TakeRecorderLiveLinkSource.h b/Engine/Plugins/Animation/LiveLink/Source/LiveLinkSequencer/Private/TakeRecorderSource/TakeRecorderLiveLinkSource.h
--- a/Engine/Plugins/Animation/LiveLink/Source/LiveLinkSequencer/Private/TakeRecorderSource/TakeRecorderLiveLinkSource.h
+++ b/Engine/Plugins/Animation/LiveLink/Source/LiveLinkSequencer/Private/TakeRecorderSource/TakeRecorderLiveLinkSource.h
@@ -71,6 +71,10 @@ public:
 	UPROPERTY(EditAnywhere, BlueprintReadWrite, Category = "Subjects")
 	bool bUseSourceTimecode;
 
+	/** If true discard livelink samples with timecode that occurs before the start of recording*/
+	UPROPERTY(EditAnywhere, BlueprintReadWrite, Category = "Subjects", meta=(EditCondition = bUseSourceTimecode))
+	bool bDiscardSamplesBeforeStart;
+
 	/**
 	* The master track recorder we created.
 	*/
diff --git a/Engine/Plugins/Animation/LiveLink/Source/LiveLinkSequencer/Public/MovieSceneLiveLinkTrackRecorder.h b/Engine/Plugins/Animation/LiveLink/Source/LiveLinkSequencer/Public/MovieSceneLiveLinkTrackRecorder.h
--- a/Engine/Plugins/Animation/LiveLink/Source/LiveLinkSequencer/Public/MovieSceneLiveLinkTrackRecorder.h
+++ b/Engine/Plugins/Animation/LiveLink/Source/LiveLinkSequencer/Public/MovieSceneLiveLinkTrackRecorder.h
@@ -47,7 +47,7 @@ public:
 
 public:
 	//we don't call UMovieSceneTrackRecorder::CreateTrack or CreateTrackImpl since that expects an  ObjectToRecord and a GUID which isn't needed.
-	void CreateTrack(UMovieScene* InMovieScene, const FName& InSubjectName, bool bInSaveSubjectSettings, bool bInAlwaysUseTimecoe, UMovieSceneTrackRecorderSettings* InSettingsObject);
+	void CreateTrack(UMovieScene* InMovieScene, const FName& InSubjectName, bool bInSaveSubjectSettings, bool bInAlwaysUseTimecode, bool bDiscardSamplesBeforeStart, UMovieSceneTrackRecorderSettings* InSettingsObject);
 	void AddContentsToFolder(UMovieSceneFolder* InFolder);
 	void SetReduceKeys(bool bInReduce) { bReduceKeys = bInReduce; }
 
@@ -70,6 +70,9 @@ private:
 	/** Whether or not we use timecode time or world time*/
 	bool bUseSourceTimecode;
 
+	/** Whether to discard livelink samples with timecode that occurs before the start of recording*/
+	bool bDiscardSamplesBeforeStart;
+
 	/** Role of the subject we will record*/
 	TSubclassOf<ULiveLinkRole> SubjectRole;
 
