Merging CL# 11619735

- Fixing default livelink preset not loading when depending on other module. Moved to PostEngineInit instead.

diff --git a/Engine/Plugins/Animation/LiveLink/Source/LiveLink/Private/LiveLinkClientReference.cpp b/Engine/Plugins/Animation/LiveLink/Source/LiveLink/Private/LiveLinkClientReference.cpp
--- a/Engine/Plugins/Animation/LiveLink/Source/LiveLink/Private/LiveLinkClientReference.cpp
+++ b/Engine/Plugins/Animation/LiveLink/Source/LiveLink/Private/LiveLinkClientReference.cpp
@@ -1,68 +1,9 @@
 // Copyright Epic Games, Inc. All Rights Reserved.
 
 #include "LiveLinkClientReference.h"
-#include "Features/IModularFeatures.h"
-#include "ILiveLinkClient.h"
+#include "LiveLinkModule.h"
 
-FLiveLinkClientReference::FLiveLinkClientReference()
-	: LiveLinkClient(nullptr)
+ILiveLinkClient* FLiveLinkClientReference::GetClient()const
 {
-	IModularFeatures& ModularFeatures = IModularFeatures::Get();
-
-	ModularFeatures.OnModularFeatureRegistered().AddRaw(this, &FLiveLinkClientReference::OnLiveLinkClientRegistered);
-	ModularFeatures.OnModularFeatureUnregistered().AddRaw(this, &FLiveLinkClientReference::OnLiveLinkClientUnregistered);
-
-	InitClient();
-}
-
-FLiveLinkClientReference::FLiveLinkClientReference(const FLiveLinkClientReference& Other)
-	: LiveLinkClient(Other.LiveLinkClient)
-{
-	IModularFeatures& ModularFeatures = IModularFeatures::Get();
-
-	ModularFeatures.OnModularFeatureRegistered().AddRaw(this, &FLiveLinkClientReference::OnLiveLinkClientRegistered);
-	ModularFeatures.OnModularFeatureUnregistered().AddRaw(this, &FLiveLinkClientReference::OnLiveLinkClientUnregistered);
-
-	InitClient();
-}
-
-FLiveLinkClientReference& FLiveLinkClientReference::operator=(const FLiveLinkClientReference& Other)
-{
-	LiveLinkClient = Other.LiveLinkClient;
-	return *this;
-}
-
-FLiveLinkClientReference::~FLiveLinkClientReference()
-{
-	IModularFeatures& ModularFeatures = IModularFeatures::Get();
-
-	ModularFeatures.OnModularFeatureRegistered().RemoveAll(this);
-	ModularFeatures.OnModularFeatureUnregistered().RemoveAll(this);
-}
-
-void FLiveLinkClientReference::InitClient()
-{
-	IModularFeatures& ModularFeatures = IModularFeatures::Get();
-
-	if (ModularFeatures.IsModularFeatureAvailable(ILiveLinkClient::ModularFeatureName))
-	{
-		LiveLinkClient = &IModularFeatures::Get().GetModularFeature<ILiveLinkClient>(ILiveLinkClient::ModularFeatureName);
-	}
-}
-
-void FLiveLinkClientReference::OnLiveLinkClientRegistered(const FName& Type, class IModularFeature* ModularFeature)
-{
-	if (Type == ILiveLinkClient::ModularFeatureName && !LiveLinkClient)
-	{
-		LiveLinkClient = static_cast<ILiveLinkClient*>(ModularFeature);
-	}
-}
-
-void FLiveLinkClientReference::OnLiveLinkClientUnregistered(const FName& Type, class IModularFeature* ModularFeature)
-{
-	if (Type == ILiveLinkClient::ModularFeatureName && ModularFeature == LiveLinkClient)
-	{
-		LiveLinkClient = nullptr;
-		InitClient();
-	}
+	return FLiveLinkModule::LiveLinkClient_AnyThread;
 }
\ No newline at end of file
diff --git a/Engine/Plugins/Animation/LiveLink/Source/LiveLink/Private/LiveLinkModule.cpp b/Engine/Plugins/Animation/LiveLink/Source/LiveLink/Private/LiveLinkModule.cpp
--- a/Engine/Plugins/Animation/LiveLink/Source/LiveLink/Private/LiveLinkModule.cpp
+++ b/Engine/Plugins/Animation/LiveLink/Source/LiveLink/Private/LiveLinkModule.cpp
@@ -1,112 +1,72 @@
 // Copyright Epic Games, Inc. All Rights Reserved.
 
-#include "ILiveLinkModule.h"
+#include "LiveLinkModule.h"
 
-#include "Features/IModularFeatures.h"
-#include "Interfaces/IPluginManager.h"
-#include "LiveLinkMotionController.h"
-
-#include "LiveLinkMessageBusDiscoveryManager.h"
-
-#include "LiveLinkClient.h"
+#include "HAL/PlatformAtomics.h"
 #include "LiveLinkLogInstance.h"
-#include "LiveLinkDebugCommand.h"
-#include "LiveLinkHeartbeatEmitter.h"
 #include "LiveLinkPreset.h"
 #include "LiveLinkSettings.h"
-
-#include "Styling/SlateStyle.h"
-
-
-/**
- * Implements the Messaging module.
- */
+#include "Misc/CoreDelegates.h"
 
 #define LOCTEXT_NAMESPACE "LiveLinkModule"
 
+FLiveLinkClient* FLiveLinkModule::LiveLinkClient_AnyThread = nullptr;
 
-class FLiveLinkModule
-	: public ILiveLinkModule
+FLiveLinkModule::FLiveLinkModule()
+	: LiveLinkClient()
+	, LiveLinkMotionController(LiveLinkClient)
+	, HeartbeatEmitter(MakeUnique<FLiveLinkHeartbeatEmitter>())
+	, DiscoveryManager(MakeUnique<FLiveLinkMessageBusDiscoveryManager>())
+	, LiveLinkDebugCommand(MakeUnique<FLiveLinkDebugCommand>(LiveLinkClient))
 {
-public:
-	FLiveLinkClient LiveLinkClient;
-	FLiveLinkMotionController LiveLinkMotionController;
-	TSharedPtr<FSlateStyleSet> StyleSet;
-	FString ContentDir;
-
-	FLiveLinkModule()
-		: LiveLinkClient()
-		, LiveLinkMotionController(LiveLinkClient)
-		, HeartbeatEmitter(MakeUnique<FLiveLinkHeartbeatEmitter>())
-		, DiscoveryManager(MakeUnique<FLiveLinkMessageBusDiscoveryManager>())
-		, LiveLinkDebugCommand(MakeUnique<FLiveLinkDebugCommand>(LiveLinkClient))
-	{
-		CreateStyle();
-	}
+}
 
-	// IModuleInterface interface
+void FLiveLinkModule::StartupModule()
+{
+	FLiveLinkLogInstance::CreateInstance();
+	CreateStyle();
 
-	virtual void StartupModule() override
-	{
-		FLiveLinkLogInstance::CreateInstance();
-		IModularFeatures::Get().RegisterModularFeature(FLiveLinkClient::ModularFeatureName, &LiveLinkClient);
-		LiveLinkMotionController.RegisterController();
-
-		if (ULiveLinkPreset* Preset = GetDefault<ULiveLinkSettings>()->DefaultLiveLinkPreset.LoadSynchronous())
-		{
-			Preset->ApplyToClient();
-		}
-	}
+	FPlatformAtomics::InterlockedExchangePtr((void**)&LiveLinkClient_AnyThread, &LiveLinkClient);
+	IModularFeatures::Get().RegisterModularFeature(FLiveLinkClient::ModularFeatureName, &LiveLinkClient);
+	LiveLinkMotionController.RegisterController();
 
-	virtual void ShutdownModule() override
-	{
-		HeartbeatEmitter->Exit();
-		DiscoveryManager->Stop();
-		LiveLinkMotionController.UnregisterController();
-		IModularFeatures::Get().UnregisterModularFeature(FLiveLinkClient::ModularFeatureName, &LiveLinkClient);
-		FLiveLinkLogInstance::DestroyInstance();
-	}
-
-	virtual bool SupportsDynamicReloading() override
-	{
-		return false;
-	}
+	//Register for engine initialization completed so we can load default preset if any. Presets could depend on plugins loaded at a later stage.
+	FCoreDelegates::OnFEngineLoopInitComplete.AddRaw(this, &FLiveLinkModule::OnEngineLoopInitComplete);
+}
 
-	virtual TSharedPtr<FSlateStyleSet> GetStyle() override
-	{
-		return StyleSet;
-	}
+void FLiveLinkModule::ShutdownModule()
+{
+	FCoreDelegates::OnFEngineLoopInitComplete.RemoveAll(this);
 
-	virtual FLiveLinkHeartbeatEmitter& GetHeartbeatEmitter() override
-	{
-		return *HeartbeatEmitter;
-	}
+	HeartbeatEmitter->Exit();
+	DiscoveryManager->Stop();
+	LiveLinkMotionController.UnregisterController();
 
-	virtual FLiveLinkMessageBusDiscoveryManager& GetMessageBusDiscoveryManager() override
-	{
-		return *DiscoveryManager;
-	}
+	IModularFeatures::Get().UnregisterModularFeature(FLiveLinkClient::ModularFeatureName, &LiveLinkClient);
+	FPlatformAtomics::InterlockedExchangePtr((void**)&LiveLinkClient_AnyThread, nullptr);
 
-private:
-	void CreateStyle()
-	{
-		static FName LiveLinkStyle(TEXT("LiveLinkCoreStyle"));
-		StyleSet = MakeShared<FSlateStyleSet>(LiveLinkStyle);
+	FLiveLinkLogInstance::DestroyInstance();
+}
 
-		ContentDir = IPluginManager::Get().FindPlugin(TEXT("LiveLink"))->GetContentDir();
+void FLiveLinkModule::CreateStyle()
+{
+	static FName LiveLinkStyle(TEXT("LiveLinkCoreStyle"));
+	StyleSet = MakeShared<FSlateStyleSet>(LiveLinkStyle);
 
-		const FVector2D Icon16x16(16.0f, 16.0f);
+	FString ContentDir = IPluginManager::Get().FindPlugin(TEXT("LiveLink"))->GetContentDir();
 
-		StyleSet->Set("LiveLinkIcon", new FSlateImageBrush((ContentDir / TEXT("LiveLink_16x")) + TEXT(".png"), Icon16x16));
-	}
+	const FVector2D Icon16x16(16.0f, 16.0f);
 
-private:
-	TUniquePtr<FLiveLinkHeartbeatEmitter> HeartbeatEmitter;
-	TUniquePtr<FLiveLinkMessageBusDiscoveryManager> DiscoveryManager;
+	StyleSet->Set("LiveLinkIcon", new FSlateImageBrush((ContentDir / TEXT("LiveLink_16x")) + TEXT(".png"), Icon16x16));
+}
 
-	/** Handler for LiveLink debug command. */
-	TUniquePtr<FLiveLinkDebugCommand> LiveLinkDebugCommand;
-};
+void FLiveLinkModule::OnEngineLoopInitComplete()
+{
+	if (ULiveLinkPreset* Preset = GetDefault<ULiveLinkSettings>()->DefaultLiveLinkPreset.LoadSynchronous())
+	{
+		Preset->ApplyToClient();
+	}
+}
 
 IMPLEMENT_MODULE(FLiveLinkModule, LiveLink);
 
diff --git a/Engine/Plugins/Animation/LiveLink/Source/LiveLink/Private/LiveLinkModule.h b/Engine/Plugins/Animation/LiveLink/Source/LiveLink/Private/LiveLinkModule.h
new file mode 100644
--- /dev/null
+++ b/Engine/Plugins/Animation/LiveLink/Source/LiveLink/Private/LiveLinkModule.h
@@ -0,0 +1,58 @@
+// Copyright Epic Games, Inc. All Rights Reserved.
+
+#pragma once
+
+#include "ILiveLinkModule.h"
+
+#include "Features/IModularFeatures.h"
+#include "Interfaces/IPluginManager.h"
+#include "LiveLinkMotionController.h"
+
+#include "LiveLinkMessageBusDiscoveryManager.h"
+
+#include "LiveLinkClient.h"
+#include "LiveLinkDebugCommand.h"
+#include "LiveLinkHeartbeatEmitter.h"
+
+#include "Styling/SlateStyle.h"
+
+/**
+ * Implements the Messaging module.
+ */
+
+ struct FLiveLinkClientReference;
+
+class FLiveLinkModule : public ILiveLinkModule
+{
+public:
+	FLiveLinkModule();
+
+	//~ Begin IModuleInterface interface
+	virtual void StartupModule() override;
+	virtual void ShutdownModule() override;
+	virtual bool SupportsDynamicReloading() override { return false; }
+	//~ End IModuleInterface interface
+
+	//~ Begin ILiveLinkModule interface
+	virtual FLiveLinkHeartbeatEmitter& GetHeartbeatEmitter() override { return *HeartbeatEmitter; }
+	virtual FLiveLinkMessageBusDiscoveryManager& GetMessageBusDiscoveryManager() override { return *DiscoveryManager; }
+	virtual TSharedPtr<FSlateStyleSet> GetStyle() override { return StyleSet; }
+	//~ End ILiveLinkModule interface
+
+private:
+	void CreateStyle();
+	void OnEngineLoopInitComplete();
+
+private:
+	friend FLiveLinkClientReference;
+	static FLiveLinkClient* LiveLinkClient_AnyThread;
+
+	FLiveLinkClient LiveLinkClient;
+	FLiveLinkMotionController LiveLinkMotionController;
+
+	TSharedPtr<FSlateStyleSet> StyleSet;
+
+	TUniquePtr<FLiveLinkHeartbeatEmitter> HeartbeatEmitter;
+	TUniquePtr<FLiveLinkMessageBusDiscoveryManager> DiscoveryManager;
+	TUniquePtr<FLiveLinkDebugCommand> LiveLinkDebugCommand;
+};
diff --git a/Engine/Plugins/Animation/LiveLink/Source/LiveLink/Public/LiveLinkClientReference.h b/Engine/Plugins/Animation/LiveLink/Source/LiveLink/Public/LiveLinkClientReference.h
--- a/Engine/Plugins/Animation/LiveLink/Source/LiveLink/Public/LiveLinkClientReference.h
+++ b/Engine/Plugins/Animation/LiveLink/Source/LiveLink/Public/LiveLinkClientReference.h
@@ -7,25 +7,9 @@
 class ILiveLinkClient;
 class IModularFeature;
 
-// References the live link client modular feature and handles add/remove
+/** References the live link client */
 struct LIVELINK_API FLiveLinkClientReference
 {
 public:
-
-	FLiveLinkClientReference();
-	FLiveLinkClientReference(const FLiveLinkClientReference& Other);
-	FLiveLinkClientReference& operator=(const FLiveLinkClientReference& Other);
-	FLiveLinkClientReference(const FLiveLinkClientReference&& Other) = delete;
-	~FLiveLinkClientReference();
-
-	ILiveLinkClient* GetClient() const { return LiveLinkClient; }
-
-private:
-	void InitClient();
-
-	// Handlers for modular features coming and going
-	void OnLiveLinkClientRegistered(const FName& Type, class IModularFeature* ModularFeature);
-	void OnLiveLinkClientUnregistered(const FName& Type, class IModularFeature* ModularFeature);
-
-	ILiveLinkClient* LiveLinkClient;
+	ILiveLinkClient* GetClient() const;
 };
\ No newline at end of file
diff --git a/Engine/Plugins/VirtualProduction/TimedDataMonitor/Source/TimedDataMonitorEditor/Private/STimedDataMonitorPanel.cpp b/Engine/Plugins/VirtualProduction/TimedDataMonitor/Source/TimedDataMonitorEditor/Private/STimedDataMonitorPanel.cpp
--- a/Engine/Plugins/VirtualProduction/TimedDataMonitor/Source/TimedDataMonitorEditor/Private/STimedDataMonitorPanel.cpp
+++ b/Engine/Plugins/VirtualProduction/TimedDataMonitor/Source/TimedDataMonitorEditor/Private/STimedDataMonitorPanel.cpp
@@ -362,23 +362,45 @@ void STimedDataMonitorPanel::Construct(const FArguments& InArgs)
 				.FillWidth(0.5f)
 				.Padding(0.f, 0.f, 10.f, 0.f)
 				[
-					SNew(SBorder)
-					.HAlign(HAlign_Center)
-					.VAlign(VAlign_Center)
-					.Padding(FMargin(4.f, 4.f, 4.f, 4.f))
+					SNew(SVerticalBox)
+					+ SVerticalBox::Slot()
+					.AutoHeight()
+					[
+						SNew(STextBlock)
+						.Text(LOCTEXT("GenlockLabel", "Genlock"))
+					]
+					+ SVerticalBox::Slot()
+					.AutoHeight()
 					[
-						SAssignNew(TimedDataGenlockWidget, STimedDataGenlock, SharedThis<STimedDataMonitorPanel>(this))
+						SNew(SBorder)
+						.HAlign(HAlign_Center)
+						.VAlign(VAlign_Center)
+						.Padding(FMargin(4.f, 4.f, 4.f, 4.f))
+						[
+							SAssignNew(TimedDataGenlockWidget, STimedDataGenlock, SharedThis<STimedDataMonitorPanel>(this))
+						]
 					]
 				]
 				+ SHorizontalBox::Slot()
 				.FillWidth(0.5f)
 				.Padding(0.f)
 				[
-					SNew(SBorder)
-					.HAlign(HAlign_Center)
-					.VAlign(VAlign_Center)
+					SNew(SVerticalBox)
+					+ SVerticalBox::Slot()
+					.AutoHeight()
+					[
+						SNew(STextBlock)
+						.Text(LOCTEXT("TimecodeLabel", "Timecode Provider"))
+					]
+					+ SVerticalBox::Slot()
+					.AutoHeight()
 					[
-						SAssignNew(TimedDataTimecodeWidget, STimedDataTimecodeProvider, SharedThis<STimedDataMonitorPanel>(this))
+						SNew(SBorder)
+						.HAlign(HAlign_Center)
+						.VAlign(VAlign_Center)
+						[
+							SAssignNew(TimedDataTimecodeWidget, STimedDataTimecodeProvider, SharedThis<STimedDataMonitorPanel>(this))
+						]
 					]
 				]
 			]
