Merging livelink fixes from dev-vp

diff --git a/Engine/Plugins/Animation/LiveLink/Source/LiveLink/Private/ClockOffsetEstimatorRamp.cpp b/Engine/Plugins/Animation/LiveLink/Source/LiveLink/Private/ClockOffsetEstimatorRamp.cpp
new file mode 100644
--- /dev/null
+++ b/Engine/Plugins/Animation/LiveLink/Source/LiveLink/Private/ClockOffsetEstimatorRamp.cpp
@@ -0,0 +1,87 @@
+// Copyright Epic Games, Inc. All Rights Reserved.
+
+#include "ClockOffsetEstimatorRamp.h"
+
+
+double FClockOffsetEstimatorRamp::GetEstimatedOffset() const
+{
+	return EstimatedOffset;
+}
+
+void FClockOffsetEstimatorRamp::UpdateEstimation(double Atime, double Btime)
+{
+	const double CurrentOffset = Btime - Atime;
+	const double Error = EstimatedOffset - CurrentOffset;
+
+	// If our estimation is too far, increase our counter to satisfy margin of acceptance before clamping the estimation.
+	if (FMath::Abs(Error) > MaxAllowedEstimationError)
+	{
+		++ConsecutiveBigErrorCount;
+		if (ConsecutiveBigErrorCount > ConsecutiveTooBigErrorRequired)
+		{
+			EstimatedOffset = CurrentOffset;
+			return;
+		}
+	}
+	else
+	{
+		ConsecutiveBigErrorCount = 0;
+	}
+
+	// Avoid overshoot if estimation is within a correction step from current offset.
+	if (FMath::Abs(Error) < CorrectionStep)
+	{
+		EstimatedOffset = CurrentOffset;
+		return;
+	}
+
+	// Apply correction step
+	if (EstimatedOffset < CurrentOffset)
+	{
+		EstimatedOffset += CorrectionStep;
+	}
+	else
+	{
+		EstimatedOffset -= CorrectionStep;
+	}
+}
+
+void FClockOffsetEstimatorRamp::SetRequiredConsecutiveBigError(int32 NewCount)
+{
+	// Only allow positive value.
+	if (!ensure(NewCount >= 0))
+	{
+		return;
+	}
+
+	ConsecutiveTooBigErrorRequired = NewCount;
+}
+
+void FClockOffsetEstimatorRamp::SetMaxAllowedEstimationError(double NewValue)
+{
+	// Only allow positive value because will only compare to absolute values.
+	if (!ensure(NewValue > 0.0))
+	{
+		return;
+	}
+
+	MaxAllowedEstimationError = NewValue;
+}
+
+void FClockOffsetEstimatorRamp::SetCorrectionStep(double NewValue)
+{
+	// Avoid control system correcting in the wrong direction.
+	if (!ensure(NewValue > 0.0))
+	{
+		return;
+	}
+
+	CorrectionStep = NewValue;
+}
+
+void FClockOffsetEstimatorRamp::Reset()
+{
+	EstimatedOffset = 0.0;
+	ConsecutiveBigErrorCount = ConsecutiveTooBigErrorRequired;
+}
+
diff --git a/Engine/Plugins/Animation/LiveLink/Source/LiveLink/Private/ClockOffsetEstimatorRamp.h b/Engine/Plugins/Animation/LiveLink/Source/LiveLink/Private/ClockOffsetEstimatorRamp.h
new file mode 100644
--- /dev/null
+++ b/Engine/Plugins/Animation/LiveLink/Source/LiveLink/Private/ClockOffsetEstimatorRamp.h
@@ -0,0 +1,52 @@
+// Copyright Epic Games, Inc. All Rights Reserved.
+
+#pragma once
+
+#include "CoreMinimal.h"
+
+/**
+ * Continually estimates the time offset between two live time sources.
+ *
+ * Samples of clock times corresponding to the same instant are provided,
+ * and it will update its estimation of the time difference between the clocks.
+ */
+class FClockOffsetEstimatorRamp
+{
+public:
+
+	/** Get current estimated offset between the two tracked clocks */
+	double GetEstimatedOffset() const;
+
+	/** Update the current time offset estimation with a new data sample */
+	void UpdateEstimation(double Atime, double Btime);
+
+	/** Set required consecutive big errors to clamp */
+	void SetRequiredConsecutiveBigError(int32 NewCount);
+
+	/** Set maximum allowed estimation error before looking to clamp */
+	void SetMaxAllowedEstimationError(double NewValue);
+
+	/** Sets incremental correction step */
+	void SetCorrectionStep(double NewValue);
+
+	/** Puts back the estimator in its initial state */
+	void Reset();
+
+protected:
+
+	/** Current estimated offset between two clocks */
+	double EstimatedOffset = 0.0;
+	
+
+	/** Maximum error tolerated before considering clamping to current offset */
+	double MaxAllowedEstimationError = 0.25;
+
+	/** Correction step to apply, positive or negative depending on error sign */
+	double CorrectionStep = 250e-6;
+
+	/** Consecutive time a too big estimation error is required before clamping */
+	int32 ConsecutiveTooBigErrorRequired = 5;
+	
+	/** Current consecutive big error count. Initialized at required count so first pass snaps */
+	int32 ConsecutiveBigErrorCount = 5;
+};
\ No newline at end of file
diff --git a/Engine/Plugins/Animation/LiveLink/Source/LiveLink/Private/LiveLinkClient.cpp b/Engine/Plugins/Animation/LiveLink/Source/LiveLink/Private/LiveLinkClient.cpp
--- a/Engine/Plugins/Animation/LiveLink/Source/LiveLink/Private/LiveLinkClient.cpp
+++ b/Engine/Plugins/Animation/LiveLink/Source/LiveLink/Private/LiveLinkClient.cpp
@@ -608,6 +608,17 @@ void FLiveLinkClient::PushSubjectFrameData_AnyThread(const FLiveLinkSubjectKey&
 		}
 		else
 		{
+			//Stamp arrival time of each packet to track clock difference
+			if (SubjectFrame.FrameData.GetBaseData())
+			{
+				SubjectFrame.FrameData.GetBaseData()->ArrivalTime.WorldTime = FPlatformTime::Seconds();
+				const TOptional<FQualifiedFrameTime>& CurrentTime = FApp::GetCurrentFrameTime();
+				if (CurrentTime.IsSet())
+				{
+					SubjectFrame.FrameData.GetBaseData()->ArrivalTime.SceneTime = *CurrentTime;
+				}
+			}
+			
 			{
 				FScopeLock BroadcastLock(&SubjectFrameReceivedHandleseCriticalSection);
 				if (const FSubjectFramesReceivedHandles* Handles = SubjectFrameReceivedHandles.Find(InSubjectKey))
@@ -615,6 +626,7 @@ void FLiveLinkClient::PushSubjectFrameData_AnyThread(const FLiveLinkSubjectKey&
 					Handles->OnFrameDataReceived.Broadcast(SubjectFrame.FrameData);
 				}
 			}
+			
 			SubjectFrameToPush.Add(MoveTemp(SubjectFrame));
 		}
 	}
diff --git a/Engine/Plugins/Animation/LiveLink/Source/LiveLink/Private/LiveLinkMessageBusSource.cpp b/Engine/Plugins/Animation/LiveLink/Source/LiveLink/Private/LiveLinkMessageBusSource.cpp
--- a/Engine/Plugins/Animation/LiveLink/Source/LiveLink/Private/LiveLinkMessageBusSource.cpp
+++ b/Engine/Plugins/Animation/LiveLink/Source/LiveLink/Private/LiveLinkMessageBusSource.cpp
@@ -205,7 +205,7 @@ void FLiveLinkMessageBusSource::InternalHandleMessage(const TSharedRef<IMessageC
 		FLiveLinkFrameDataStruct DataStruct(MessageTypeInfo);
 		const FLiveLinkBaseFrameData* Message = reinterpret_cast<const FLiveLinkBaseFrameData*>(Context->GetMessage());
 		DataStruct.InitializeWith(MessageTypeInfo, Message);
-		DataStruct.GetBaseData()->WorldTime = FLiveLinkWorldTime(Message->WorldTime.GetOffsettedTime(), MachineTimeOffset);
+		DataStruct.GetBaseData()->WorldTime = Message->WorldTime.GetOffsettedTime();
 		Client->PushSubjectFrameData_AnyThread(SubjectKey, MoveTemp(DataStruct));
 	}
 }
diff --git a/Engine/Plugins/Animation/LiveLink/Source/LiveLink/Private/LiveLinkSubject.cpp b/Engine/Plugins/Animation/LiveLink/Source/LiveLink/Private/LiveLinkSubject.cpp
--- a/Engine/Plugins/Animation/LiveLink/Source/LiveLink/Private/LiveLinkSubject.cpp
+++ b/Engine/Plugins/Animation/LiveLink/Source/LiveLink/Private/LiveLinkSubject.cpp
@@ -37,7 +37,10 @@ void FLiveLinkSubject::Initialize(FLiveLinkSubjectKey InSubjectKey, TSubclassOf<
 	Role = InRole;
 
 	FrameData.Reset();
+	ReceivedOrderedFrames.Empty();
 	ResetBufferStats();
+	EngineTimeClockOffset.Reset();
+	TimecodeClockOffset.Reset();
 
 	if (TSharedPtr<FLiveLinkTimedDataInput> TimedDataGroupPin = TimedDataGroup.Pin())
 	{
@@ -63,10 +66,7 @@ void FLiveLinkSubject::Update()
 	{
 		const int32 NumberOfFrameToRemove = FrameData.Num() - CachedSettings.BufferSettings.MaxNumberOfFrameToBuffered;
 		const int32 Count = CachedSettings.BufferSettings.bKeepAtLeastOneFrame && FrameData.Num() == NumberOfFrameToRemove ? NumberOfFrameToRemove-1 : NumberOfFrameToRemove;
-		if (Count > 0)
-		{
-			FrameData.RemoveAt(0, Count, false);
-		}
+		RemoveFrames(Count);
 	}
 
 	if (GetMode() == ELiveLinkSourceMode::EngineTime)
@@ -89,10 +89,7 @@ void FLiveLinkSubject::Update()
 			if (FrameIndex - 1 >= 0)
 			{
 				const int32 Count = CachedSettings.BufferSettings.bKeepAtLeastOneFrame && FrameData.Num() == FrameIndex ? FrameIndex - 1 : FrameIndex;
-				if (Count > 0)
-				{
-					FrameData.RemoveAt(0, Count, false);
-				}
+				RemoveFrames(Count);
 			}
 		}
 	}
@@ -102,14 +99,22 @@ void FLiveLinkSubject::Update()
 		{
 			if (FApp::GetCurrentFrameTime().IsSet())
 			{
+				//Take the current Engine Timecode and substract all the offset in subject's frame space (clock offset if any, frame offset and the desired frame distance threshold)
 				const FQualifiedFrameTime CurrentSyncTime = FApp::GetCurrentFrameTime().GetValue();
+
+				FFrameTime ClockOffset;
+				if (CachedSettings.BufferSettings.bUseTimecodeSmoothLatest)
+				{
+					ClockOffset = CachedSettings.BufferSettings.TimecodeFrameRate.AsFrameTime(TimecodeClockOffset.GetEstimatedOffset());
+				}
+
 				const FFrameTime CurrentFrameTimeInFrameSpace = CurrentSyncTime.ConvertTo(CachedSettings.BufferSettings.TimecodeFrameRate);
+				const FFrameTime FrameTimeThresholdFrameSpace = CurrentFrameTimeInFrameSpace - (ClockOffset + FFrameTime::FromDecimal(CachedSettings.BufferSettings.TimecodeFrameOffset) + CachedSettings.BufferSettings.ValidTimecodeFrame);
 				int32 FrameIndex = 0;
 				for (const FLiveLinkFrameDataStruct& SourceFrameData : FrameData)
 				{
-					FFrameTime UsedFrameTime = CurrentFrameTimeInFrameSpace - FFrameTime::FromDecimal(CachedSettings.BufferSettings.TimecodeFrameOffset) - CachedSettings.BufferSettings.ValidTimecodeFrame;
-					FFrameTime FrameTime = SourceFrameData.GetBaseData()->MetaData.SceneTime.Time;
-					if (FrameTime > UsedFrameTime)
+					const FFrameTime FrameTime = SourceFrameData.GetBaseData()->MetaData.SceneTime.Time;
+					if (FrameTime > FrameTimeThresholdFrameSpace)
 					{
 						break;
 					}
@@ -119,10 +124,7 @@ void FLiveLinkSubject::Update()
 				if (FrameIndex - 1 >= 0)
 				{
 					const int32 Count = CachedSettings.BufferSettings.bKeepAtLeastOneFrame && FrameData.Num() == FrameIndex ? FrameIndex - 1 : FrameIndex;
-					if (Count > 0)
-					{
-						FrameData.RemoveAt(0, Count, false);
-					}
+					RemoveFrames(Count);
 				}
 			}
 		}
@@ -140,19 +142,33 @@ void FLiveLinkSubject::Update()
 		{
 			if (FApp::GetCurrentFrameTime().IsSet())
 			{
-				bSnapshotIsValid = GetFrameAtSceneTime(FApp::GetCurrentFrameTime().GetValue(), FrameSnapshot);
+				//If Timecode smooth latest is enabled, use the clock offset to adjust read time. 
+				//@note For now, snapshot created will have the source time. We are not overriding its time with the read time. If a use case arise for this, we'll adjust.
+				const FQualifiedFrameTime EngineTimecode = FApp::GetCurrentFrameTime().GetValue();
+				
+				FFrameTime ClockOffset;
+				if (CachedSettings.BufferSettings.bUseTimecodeSmoothLatest)
+				{
+					ClockOffset = EngineTimecode.Rate.AsFrameTime(TimecodeClockOffset.GetEstimatedOffset());
+				}
+
+				const FFrameTime FrameOffset = FQualifiedFrameTime(FFrameTime::FromDecimal(CachedSettings.BufferSettings.TimecodeFrameOffset), CachedSettings.BufferSettings.TimecodeFrameRate).ConvertTo(EngineTimecode.Rate);
+				const FFrameTime ReadTime = EngineTimecode.Time - (ClockOffset + FrameOffset);
+				const FQualifiedFrameTime LookupQFrameTime = FQualifiedFrameTime(ReadTime, EngineTimecode.Rate);
+				bSnapshotIsValid = GetFrameAtSceneTime(LookupQFrameTime, FrameSnapshot);
 			}
 			else
 			{
-				static const FName NAME_InvalidRole = "LiveLinkSubject_NoCurrentFrameTime";
-				FLiveLinkLog::WarningOnce(NAME_InvalidRole, SubjectKey, TEXT("Can't evaluate frame for subject '%s'. The engine doesn't have a timecode value set."), *SubjectKey.SubjectName.ToString());
+				static const FName NAME_InvalidTime = "LiveLinkSubject_NoCurrentFrameTime";
+				FLiveLinkLog::WarningOnce(NAME_InvalidTime, SubjectKey, TEXT("Can't evaluate frame for subject '%s'. The engine doesn't have a timecode value set."), *SubjectKey.SubjectName.ToString());
 			}
 		}
 		break;
 
 		case ELiveLinkSourceMode::EngineTime:
 		{
-			bSnapshotIsValid = GetFrameAtWorldTime(FApp::GetCurrentTime(), FrameSnapshot);
+			const double ReadTime = FApp::GetCurrentTime() - CachedSettings.BufferSettings.EngineTimeOffset;
+			bSnapshotIsValid = GetFrameAtWorldTime(ReadTime, FrameSnapshot);
 			break;
 		}
 
@@ -167,6 +183,12 @@ void FLiveLinkSubject::Update()
 		// Invalidate the snapshot
 		FrameSnapshot.FrameData.Reset();
 	}
+	else
+	{
+		//Extra verbose logging about snapshot just created
+		const FTimecode FrameTC = FTimecode::FromFrameNumber(FrameSnapshot.FrameData.GetBaseData()->MetaData.SceneTime.Time.GetFrame(), FrameSnapshot.FrameData.GetBaseData()->MetaData.SceneTime.Rate);
+		UE_LOG(LogLiveLink, VeryVerbose, TEXT("Subject '%s' created snapshot with frame ID %d - Timecode:[%s.%0.3f] - SourceTime: %0.4f, Offset: %0.6f, CorrectedTime: %0.4f"), *SubjectKey.SubjectName.ToString(), FrameSnapshot.FrameData.GetBaseData()->FrameId, *FrameTC.ToString(), FrameSnapshot.FrameData.GetBaseData()->MetaData.SceneTime.Time.GetSubFrame(), FrameSnapshot.FrameData.GetBaseData()->WorldTime.GetSourceTime(), FrameSnapshot.FrameData.GetBaseData()->WorldTime.GetOffset(), FrameSnapshot.FrameData.GetBaseData()->WorldTime.GetOffsettedTime());
+	}
 }
 
 void FLiveLinkSubject::ClearFrames()
@@ -174,6 +196,7 @@ void FLiveLinkSubject::ClearFrames()
 	FrameSnapshot.StaticData.Reset();
 	FrameSnapshot.FrameData.Reset();
 	FrameData.Reset();
+	ReceivedOrderedFrames.Empty();
 }
 
 bool FLiveLinkSubject::HasValidFrameSnapshot() const
@@ -218,15 +241,17 @@ bool FLiveLinkSubject::EvaluateFrameAtWorldTime(double InWorldTime, TSubclassOf<
 	bool bSuccess = false;
 	if (FrameData.Num() != 0)
 	{
+		// Adjust desired read time based on offset settings
+		const double ReadTime = InWorldTime - CachedSettings.BufferSettings.EngineTimeOffset;
 		if (Role == InDesiredRole || Role->IsChildOf(InDesiredRole))
 		{
-			GetFrameAtWorldTime(InWorldTime, OutFrame);
+			GetFrameAtWorldTime(ReadTime, OutFrame);
 			bSuccess = true;
 		}
 		else if (SupportsRole(InDesiredRole))
 		{
 			FLiveLinkSubjectFrameData TmpFrameData;
-			GetFrameAtWorldTime(InWorldTime, TmpFrameData);
+			GetFrameAtWorldTime(ReadTime, TmpFrameData);
 			bSuccess = Translate(this, InDesiredRole, TmpFrameData.StaticData, TmpFrameData.FrameData, OutFrame);
 		}
 		else
@@ -265,15 +290,25 @@ bool FLiveLinkSubject::EvaluateFrameAtSceneTime(const FQualifiedFrameTime& InSce
 	bool bSuccess = false;
 	if (FrameData.Num() != 0)
 	{
+		// Adjust desired read time based on offset settings and clock offset if any
+		FFrameTime ClockOffset;
+		if (CachedSettings.BufferSettings.bUseTimecodeSmoothLatest)
+		{
+			ClockOffset = InSceneTime.Rate.AsFrameTime(TimecodeClockOffset.GetEstimatedOffset());
+		}
+
+		const FFrameTime FrameOffset = FQualifiedFrameTime(FFrameTime::FromDecimal(CachedSettings.BufferSettings.TimecodeFrameOffset), CachedSettings.BufferSettings.TimecodeFrameRate).ConvertTo(InSceneTime.Rate);
+		const FFrameTime ReadTime = InSceneTime.Time - (ClockOffset + FrameOffset);
+		const FQualifiedFrameTime LookupQFrameTime = FQualifiedFrameTime(ReadTime, InSceneTime.Rate);
 		if (Role == InDesiredRole || Role->IsChildOf(InDesiredRole))
 		{
-			GetFrameAtSceneTime(InSceneTime, OutFrame);
+			GetFrameAtSceneTime(LookupQFrameTime, OutFrame);
 			bSuccess = true;
 		}
 		else if (SupportsRole(InDesiredRole))
 		{
 			FLiveLinkSubjectFrameData TmpFrameData;
-			GetFrameAtSceneTime(InSceneTime, TmpFrameData);
+			GetFrameAtSceneTime(LookupQFrameTime, TmpFrameData);
 			bSuccess = Translate(this, InDesiredRole, TmpFrameData.StaticData, TmpFrameData.FrameData, OutFrame);
 		}
 		else
@@ -320,6 +355,20 @@ void FLiveLinkSubject::AddFrameData(FLiveLinkFrameDataStruct&& InFrameData)
 		return;
 	}
 
+	// Before adding the new frame, test to see if we are going to increase the buffer size
+	const bool bRemoveFrame = FrameData.Num() >= CachedSettings.BufferSettings.MaxNumberOfFrameToBuffered;
+	if (bRemoveFrame)
+	{
+		RemoveFrames(1);
+	}
+
+	//Update clock trackers
+	UpdateClockTracker(*InFrameData.GetBaseData());
+
+	//Update this frame clock offset with the newly updated one. This will make it transparent everywhere we used the offsetted time.
+	//Done before selecting a frame offset to have the most accurate offset
+	InFrameData.GetBaseData()->WorldTime.SetClockOffset(EngineTimeClockOffset.GetEstimatedOffset());
+	
 	int32 FrameIndex = INDEX_NONE;
 	switch (CachedSettings.SourceMode)
 	{
@@ -334,34 +383,29 @@ void FLiveLinkSubject::AddFrameData(FLiveLinkFrameDataStruct&& InFrameData)
 		FrameIndex = FindNewFrame_Latest(InFrameData.GetBaseData()->WorldTime);
 		break;
 	}
-	
-	if (FrameIndex >= 0)
+
+	//Make sure frame should be inserted. 
+	if(FrameIndex >= 0)
 	{
-		// Before adding the new frame, test to see if we are going to increase the buffer size
-		const bool bRemoveFrame = FrameData.Num() >= CachedSettings.BufferSettings.MaxNumberOfFrameToBuffered;
-		if (bRemoveFrame)
+		for (ULiveLinkFramePreProcessor::FWorkerSharedPtr PreProcessor : FramePreProcessors)
 		{
-			--FrameIndex;
+			PreProcessor->PreProcessFrame(InFrameData);
 		}
-
-		// It's possible the new frame is the frame we want to remove
-		if (FrameIndex >= 0)
+		
+		//Assign identifier to incoming frame and insert it where it belongs
+		const FLiveLinkFrameIdentifier ThisFrameIdentifier = NextIdentifier++;
+		ReceivedOrderedFrames.Enqueue(ThisFrameIdentifier);
+		InFrameData.GetBaseData()->FrameId = ThisFrameIdentifier;
+		
+		//Extra verbose logging about new frame being added
+		const FTimecode FrameTC = FTimecode::FromFrameNumber(InFrameData.GetBaseData()->MetaData.SceneTime.Time.GetFrame(), InFrameData.GetBaseData()->MetaData.SceneTime.Rate);
+		UE_LOG(LogLiveLink, VeryVerbose, TEXT("Subject '%s' adding frame ID %d at index %d - Timecode:[%s.%0.3f] - SourceTime: %0.4f, Offset: %0.6f, CorrectedTime: %0.4f"), *SubjectKey.SubjectName.ToString(), ThisFrameIdentifier, FrameIndex, *FrameTC.ToString(), InFrameData.GetBaseData()->MetaData.SceneTime.Time.GetSubFrame(), InFrameData.GetBaseData()->WorldTime.GetSourceTime(), InFrameData.GetBaseData()->WorldTime.GetOffset(), InFrameData.GetBaseData()->WorldTime.GetOffsettedTime());
+		
+		FrameData.Insert(MoveTemp(InFrameData), FrameIndex);
+
+		if (CachedSettings.BufferSettings.bGenerateSubFrame && CachedSettings.SourceMode == ELiveLinkSourceMode::Timecode)
 		{
-			for (ULiveLinkFramePreProcessor::FWorkerSharedPtr PreProcessor : FramePreProcessors)
-			{
-				PreProcessor->PreProcessFrame(InFrameData);
-			}
-
-			if (bRemoveFrame)
-			{
-				FrameData.RemoveAt(0);
-			}
-			FrameData.Insert(MoveTemp(InFrameData), FrameIndex);
-
-			if (CachedSettings.BufferSettings.bGenerateSubFrame && CachedSettings.SourceMode == ELiveLinkSourceMode::Timecode)
-			{
-				AdjustSubFrame_SceneTime(FrameIndex);
-			}
+			AdjustSubFrame_SceneTime(FrameIndex);
 		}
 	}
 	else 
@@ -618,11 +662,10 @@ bool FLiveLinkSubject::GetFrameAtWorldTime_Closest(const double InSeconds, FLive
 	bool bOverflowDetected = false;
 	bool bUnderflowDetected = false;
 	bool bBuiltFrame = false;
-	const double ReadTime = (InSeconds) - CachedSettings.BufferSettings.EngineTimeOffset;
 	for (int32 FrameIndex = FrameData.Num() - 1; FrameIndex >= 0; --FrameIndex)
 	{
 		const double Time = FrameData[FrameIndex].GetBaseData()->WorldTime.GetOffsettedTime();
-		if (Time <= ReadTime)
+		if (Time <= InSeconds)
 		{
 			if (FrameIndex == FrameData.Num() - 1)
 			{
@@ -631,13 +674,13 @@ bool FLiveLinkSubject::GetFrameAtWorldTime_Closest(const double InSeconds, FLive
 				bBuiltFrame = true;
 
 				//If we tried to read above our buffer, stamp an overflow
-				bUnderflowDetected = !FMath::IsNearlyEqual(Time, ReadTime);
+				bUnderflowDetected = !FMath::IsNearlyEqual(Time, InSeconds);
 				break;
 			}
 			else
 			{
 				const double NextTime = FrameData[FrameIndex + 1].GetBaseData()->WorldTime.GetOffsettedTime();
-				const float BlendWeight = (ReadTime - NextTime) / (NextTime - Time);
+				const float BlendWeight = (InSeconds - NextTime) / (NextTime - Time);
 				int32 CopyIndex = (BlendWeight > 0.5f) ? FrameIndex : FrameIndex + 1;
 				OutFrame.FrameData.InitializeWith(FrameData[CopyIndex].GetStruct(), FrameData[CopyIndex].GetBaseData());
 				bBuiltFrame = true;
@@ -666,8 +709,8 @@ bool FLiveLinkSubject::GetFrameAtWorldTime_Closest(const double InSeconds, FLive
 		}
 
 		FTimedDataInputEvaluationData EvaluationData;
-		EvaluationData.DistanceToNewestSampleSeconds = FrameData[FrameData.Num()-1].GetBaseData()->WorldTime.GetOffsettedTime() - ReadTime;
-		EvaluationData.DistanceToOldestSampleSeconds = ReadTime - FrameData[0].GetBaseData()->WorldTime.GetOffsettedTime();
+		EvaluationData.DistanceToNewestSampleSeconds = FrameData[FrameData.Num()-1].GetBaseData()->WorldTime.GetOffsettedTime() - InSeconds;
+		EvaluationData.DistanceToOldestSampleSeconds = InSeconds - FrameData[0].GetBaseData()->WorldTime.GetOffsettedTime();
 		UpdateEvaluationData(EvaluationData);
 	}
 
@@ -679,8 +722,7 @@ bool FLiveLinkSubject::GetFrameAtWorldTime_Interpolated(const double InSeconds,
 	check(FrameData.Num() != 0);
 
 	FLiveLinkInterpolationInfo InterpolationInfo;
-	const double ReadTime = InSeconds - CachedSettings.BufferSettings.EngineTimeOffset;
-	FrameInterpolationProcessor->Interpolate(ReadTime, StaticData, FrameData, OutFrame, InterpolationInfo);
+	FrameInterpolationProcessor->Interpolate(InSeconds, StaticData, FrameData, OutFrame, InterpolationInfo);
 
 	VerifyInterpolationInfo(InterpolationInfo);
 
@@ -721,10 +763,7 @@ bool FLiveLinkSubject::GetFrameAtSceneTime_Closest(const FQualifiedFrameTime& In
 	bool bOverflowDetected = false;
 	bool bBuiltFrame = false;
 
-	const FFrameTime FrameOffset = FQualifiedFrameTime(FFrameTime::FromDecimal(CachedSettings.BufferSettings.TimecodeFrameOffset), CachedSettings.BufferSettings.TimecodeFrameRate).ConvertTo(InTimeInEngineFrameRate.Rate);
-	const FFrameTime ReadTime = InTimeInEngineFrameRate.Time - FrameOffset;
-	const FQualifiedFrameTime LookupQFrameTime = FQualifiedFrameTime(ReadTime, InTimeInEngineFrameRate.Rate);
-	const double TimeInSeconds = LookupQFrameTime.AsSeconds();
+	const double TimeInSeconds = InTimeInEngineFrameRate.AsSeconds();
 	for (int32 FrameIndex = FrameData.Num() - 1; FrameIndex >= 0; --FrameIndex)
 	{
 		const double FrameASeconds = FrameData[FrameIndex].GetBaseData()->MetaData.SceneTime.AsSeconds();
@@ -786,10 +825,7 @@ bool FLiveLinkSubject::GetFrameAtSceneTime_Interpolated(const FQualifiedFrameTim
 	check(FrameData.Num() != 0);
 
 	FLiveLinkInterpolationInfo InterpolationInfo;
-	const FFrameTime FrameOffset = FQualifiedFrameTime(FFrameTime::FromDecimal(CachedSettings.BufferSettings.TimecodeFrameOffset), CachedSettings.BufferSettings.TimecodeFrameRate).ConvertTo(InTimeInEngineFrameRate.Rate);
-	const FFrameTime ReadTime = InTimeInEngineFrameRate.Time - FrameOffset;
-	const FQualifiedFrameTime LookupQFrameTime = FQualifiedFrameTime(ReadTime, InTimeInEngineFrameRate.Rate);
-	FrameInterpolationProcessor->Interpolate(LookupQFrameTime, StaticData, FrameData, OutFrame, InterpolationInfo);
+	FrameInterpolationProcessor->Interpolate(InTimeInEngineFrameRate, StaticData, FrameData, OutFrame, InterpolationInfo);
 
 	if (IsBufferStatsEnabled())
 	{
@@ -894,6 +930,34 @@ void FLiveLinkSubject::UpdateEvaluationData(const FTimedDataInputEvaluationData&
 	EvaluationStatistics.LastEvaluationData = EvaluationData;
 }
 
+void FLiveLinkSubject::UpdateClockTracker(const FLiveLinkBaseFrameData& NewFrame)
+{
+	EngineTimeClockOffset.UpdateEstimation(NewFrame.WorldTime.GetSourceTime(), NewFrame.ArrivalTime.WorldTime);
+	TimecodeClockOffset.UpdateEstimation(NewFrame.MetaData.SceneTime.AsSeconds(), NewFrame.ArrivalTime.SceneTime.AsSeconds());
+}
+
+void FLiveLinkSubject::RemoveFrames(int32 InCount)
+{
+	for (int32 i = 0; i < InCount; ++i)
+	{
+		FLiveLinkFrameIdentifier FrameToRemove = INDEX_NONE;
+		ReceivedOrderedFrames.Dequeue(FrameToRemove);
+		const int32 FrameToRemoveIndex = FrameData.IndexOfByPredicate([FrameToRemove](const FLiveLinkFrameDataStruct& Other) { return Other.GetBaseData()->FrameId == FrameToRemove; });
+		if (FrameToRemoveIndex != INDEX_NONE)
+		{
+			const int32 Count = 1;
+			const bool bAllowShrinking = false;
+			FrameData.RemoveAt(FrameToRemoveIndex, Count, bAllowShrinking);
+		}
+		else
+		{
+			FLiveLinkLog::Error(TEXT("Subject '%s': Could not find frame with identifier %d. Buffers are out of sync, resetting them."), *SubjectKey.SubjectName.ToString(), FrameToRemove);
+			ReceivedOrderedFrames.Empty();
+			FrameData.Reset();
+		}
+	}
+}
+
 void FLiveLinkSubject::SetStaticData(TSubclassOf<ULiveLinkRole> InRole, FLiveLinkStaticDataStruct&& InStaticData)
 {
 	check(IsInGameThread());
@@ -909,6 +973,7 @@ void FLiveLinkSubject::SetStaticData(TSubclassOf<ULiveLinkRole> InRole, FLiveLin
 	{
 		//Set initial blending processor to the role's default one. User will be able to modify it afterwards.
 		FrameData.Reset();
+		ReceivedOrderedFrames.Empty();
 		StaticData = MoveTemp(InStaticData);
 	}
 	else
@@ -927,9 +992,11 @@ void FLiveLinkSubject::CacheSettings(ULiveLinkSourceSettings* SourceSetting, ULi
 		const bool bSourceModeChanged = SourceSetting->Mode != CachedSettings.SourceMode;
 		const bool bTimecodeFrameRateChanged = SourceSetting->Mode == ELiveLinkSourceMode::Timecode && SourceSetting->BufferSettings.TimecodeFrameRate != CachedSettings.BufferSettings.TimecodeFrameRate;
 		const bool bGenerateSubFrameChanged = SourceSetting->Mode == ELiveLinkSourceMode::Timecode && SourceSetting->BufferSettings.bGenerateSubFrame != CachedSettings.BufferSettings.bGenerateSubFrame;
-		if (bSourceModeChanged || bTimecodeFrameRateChanged || bGenerateSubFrameChanged)
+		const bool bTimecodeModeChanged = SourceSetting->Mode == ELiveLinkSourceMode::Timecode && SourceSetting->BufferSettings.bUseTimecodeSmoothLatest != CachedSettings.BufferSettings.bUseTimecodeSmoothLatest;
+		if (bSourceModeChanged || bTimecodeFrameRateChanged || bGenerateSubFrameChanged || bTimecodeModeChanged)
 		{
 			FrameData.Reset();
+			ReceivedOrderedFrames.Empty();
 		}
 
 		CachedSettings.SourceMode = SourceSetting->Mode;
@@ -1039,7 +1106,15 @@ FTimedDataChannelSampleTime FLiveLinkSubject::GetOldestDataTime() const
 {
 	if (FrameData.Num() > 0)
 	{
-		return FTimedDataChannelSampleTime(FrameData[0].GetBaseData()->WorldTime.GetOffsettedTime(), FrameData[0].GetBaseData()->MetaData.SceneTime);
+		//For timed data, adjust frame timings to match with clock offset to have valid timing diagrams
+		FFrameTime TimecodeClockOffsetTime;
+		if (CachedSettings.BufferSettings.bUseTimecodeSmoothLatest)
+		{
+			TimecodeClockOffsetTime = FrameData[0].GetBaseData()->MetaData.SceneTime.Rate.AsFrameTime(TimecodeClockOffset.GetEstimatedOffset());
+		}
+
+		const FQualifiedFrameTime OldestSceneTime = FQualifiedFrameTime(FrameData[0].GetBaseData()->MetaData.SceneTime.Time + TimecodeClockOffsetTime, FrameData[0].GetBaseData()->MetaData.SceneTime.Rate);
+		return FTimedDataChannelSampleTime(FrameData[0].GetBaseData()->WorldTime.GetOffsettedTime(), OldestSceneTime);
 	}
 	return FTimedDataChannelSampleTime();
 }
@@ -1048,18 +1123,34 @@ FTimedDataChannelSampleTime FLiveLinkSubject::GetNewestDataTime() const
 {
 	if (FrameData.Num() > 0)
 	{
-		return FTimedDataChannelSampleTime(FrameData.Last().GetBaseData()->WorldTime.GetOffsettedTime(), FrameData.Last().GetBaseData()->MetaData.SceneTime);
+		//For timed data, adjust frame timings to match with clock offset to have valid timing diagrams
+		FFrameTime TimecodeClockOffsetTime;
+		if (CachedSettings.BufferSettings.bUseTimecodeSmoothLatest)
+		{
+			TimecodeClockOffsetTime = FrameData.Last().GetBaseData()->MetaData.SceneTime.Rate.AsFrameTime(TimecodeClockOffset.GetEstimatedOffset());
+		}
+
+		const FQualifiedFrameTime NewestSceneTime = FQualifiedFrameTime(FrameData.Last().GetBaseData()->MetaData.SceneTime.Time + TimecodeClockOffsetTime, FrameData.Last().GetBaseData()->MetaData.SceneTime.Rate);
+		return FTimedDataChannelSampleTime(FrameData.Last().GetBaseData()->WorldTime.GetOffsettedTime(), NewestSceneTime);
 	}
 	return FTimedDataChannelSampleTime();
 }
 
 TArray<FTimedDataChannelSampleTime> FLiveLinkSubject::GetDataTimes() const
 {
+	//For timed data, adjust frame timings to match with clock offset to have valid timing diagrams
+	FFrameTime TimecodeClockOffsetTime;
+	if (CachedSettings.BufferSettings.bUseTimecodeSmoothLatest)
+	{
+		TimecodeClockOffsetTime = FrameData[0].GetBaseData()->MetaData.SceneTime.Rate.AsFrameTime(TimecodeClockOffset.GetEstimatedOffset());;
+	}
+	
 	TArray<FTimedDataChannelSampleTime> Result;
 	Result.Reset(FrameData.Num());
 	for (const FLiveLinkFrameDataStruct& Data : FrameData)
 	{
-		Result.Emplace(Data.GetBaseData()->WorldTime.GetOffsettedTime(), Data.GetBaseData()->MetaData.SceneTime);
+		const FQualifiedFrameTime AdjustedSceneTime = FQualifiedFrameTime(Data.GetBaseData()->MetaData.SceneTime.Time + TimecodeClockOffsetTime, Data.GetBaseData()->MetaData.SceneTime.Rate);
+		Result.Emplace(Data.GetBaseData()->WorldTime.GetOffsettedTime(), AdjustedSceneTime);
 	}
 	return Result;
 }
diff --git a/Engine/Plugins/Animation/LiveLink/Source/LiveLink/Private/LiveLinkSubject.h b/Engine/Plugins/Animation/LiveLink/Source/LiveLink/Private/LiveLinkSubject.h
--- a/Engine/Plugins/Animation/LiveLink/Source/LiveLink/Private/LiveLinkSubject.h
+++ b/Engine/Plugins/Animation/LiveLink/Source/LiveLink/Private/LiveLinkSubject.h
@@ -4,9 +4,11 @@
 
 #include "ILiveLinkSubject.h"
 
-#include "LiveLinkClient.h"
+#include "ClockOffsetEstimatorRamp.h"
+#include "Containers/Queue.h"
 #include "ILiveLinkSource.h"
 #include "ITimedDataInput.h"
+#include "LiveLinkClient.h"
 #include "LiveLinkFrameInterpolationProcessor.h"
 #include "LiveLinkFramePreProcessor.h"
 #include "LiveLinkFrameTranslator.h"
@@ -117,12 +119,12 @@ private:
 	// Reorder frame with the same timecode and create subframes
 	void AdjustSubFrame_SceneTime(int32 FrameIndex);
 
-	// Populate OutFrame with a frame based off of the supplied time and our own offsets
+	// Populate OutFrame with a frame based off of the supplied time (pre offsetted)
 	bool GetFrameAtWorldTime(const double InSeconds, FLiveLinkSubjectFrameData& OutFrame);
 	bool GetFrameAtWorldTime_Closest(const double InSeconds, FLiveLinkSubjectFrameData& OutFrame);
 	bool GetFrameAtWorldTime_Interpolated(const double InSeconds, FLiveLinkSubjectFrameData& OutFrame);
 
-	// Populate OutFrame with a frame based off of the supplied scene time.
+	// Populate OutFrame with a frame based off of the supplied scene time (pre offsetted).
 	bool GetFrameAtSceneTime(const FQualifiedFrameTime& InSceneTime, FLiveLinkSubjectFrameData& OutFrame);
 	bool GetFrameAtSceneTime_Closest(const FQualifiedFrameTime& InSceneTime, FLiveLinkSubjectFrameData& OutFrame);
 	bool GetFrameAtSceneTime_Interpolated(const FQualifiedFrameTime& InSceneTime, FLiveLinkSubjectFrameData& OutFrame);
@@ -141,6 +143,12 @@ private:
 	void IncreaseBufferOverFlowStat();
 	void UpdateEvaluationData(const FTimedDataInputEvaluationData& EvaluationData);
 
+	//Tracks clock difference between subject time and engine time.
+	void UpdateClockTracker(const FLiveLinkBaseFrameData& NewFrame);
+
+	//Remove frames from our buffer - based on receiving order
+	void RemoveFrames(int32 Count);
+
 protected:
 	// The role the subject was build with
 	TSubclassOf<ULiveLinkRole> Role;
@@ -177,6 +185,12 @@ private:
 	// Frames added to the subject
 	TArray<FLiveLinkFrameDataStruct> FrameData;
 
+	// Contains identifier of each frame in the order they were received
+	TQueue<FLiveLinkFrameIdentifier> ReceivedOrderedFrames;
+
+	// Next identifier to assign to next received frame
+	FLiveLinkFrameIdentifier NextIdentifier = 0;
+
 	// Current frame snapshot of the evaluation
 	FLiveLinkSubjectFrameData FrameSnapshot;
 
@@ -197,7 +211,11 @@ private:
 
 	// Some stats compiled by the subject.
 	FSubjectEvaluationStatistics EvaluationStatistics;
-	
+
+	// Keep track of offsets between subject time and engine time
+	FClockOffsetEstimatorRamp EngineTimeClockOffset;
+	FClockOffsetEstimatorRamp TimecodeClockOffset;
+
 	/** 
 	 * Evaluation can be done on any thread so we need to protect statistic logging 
 	 * Some stats requires more than atomic sized vars so a critical section is used to protect when necessary
diff --git a/Engine/Plugins/Animation/LiveLink/Source/LiveLink/Public/LiveLinkTimecodeProvider.h b/Engine/Plugins/Animation/LiveLink/Source/LiveLink/Public/LiveLinkTimecodeProvider.h
--- a/Engine/Plugins/Animation/LiveLink/Source/LiveLink/Public/LiveLinkTimecodeProvider.h
+++ b/Engine/Plugins/Animation/LiveLink/Source/LiveLink/Public/LiveLinkTimecodeProvider.h
@@ -73,7 +73,7 @@ private:
 	FLiveLinkSubjectKey SubjectKey;
 
 	/** How to evaluate the timecode. */
-	UPROPERTY(EditAnywhere, Category = Timecode, meta = (EditCondition = "!bUseLatest"))
+	UPROPERTY(EditAnywhere, Category = Timecode)
 	ELiveLinkTimecodeProviderEvaluationType Evaluation;
 
 	UPROPERTY(EditAnywhere, Category = Timecode, meta=(InlineEditConditionToggle))
diff --git a/Engine/Plugins/Animation/LiveLink/Source/LiveLinkEditor/Private/LiveLinkSourceSettingsDetailCustomization.cpp b/Engine/Plugins/Animation/LiveLink/Source/LiveLinkEditor/Private/LiveLinkSourceSettingsDetailCustomization.cpp
--- a/Engine/Plugins/Animation/LiveLink/Source/LiveLinkEditor/Private/LiveLinkSourceSettingsDetailCustomization.cpp
+++ b/Engine/Plugins/Animation/LiveLink/Source/LiveLinkEditor/Private/LiveLinkSourceSettingsDetailCustomization.cpp
@@ -43,6 +43,7 @@ void FLiveLinkSourceSettingsDetailCustomization::CustomizeDetails(IDetailLayoutB
 		CategoryBuilder.AddProperty(BufferSettingsPropertyHandle->GetChildHandle(GET_MEMBER_NAME_CHECKED(FLiveLinkSourceBufferManagementSettings, TimecodeFrameOffset)))
 			.DisplayName(LOCTEXT("TimecodeFrameOffsetDisplayName", "Offset"));
 		CategoryBuilder.AddProperty(BufferSettingsPropertyHandle->GetChildHandle(GET_MEMBER_NAME_CHECKED(FLiveLinkSourceBufferManagementSettings, TimecodeFrameRate)));
+		CategoryBuilder.AddProperty(BufferSettingsPropertyHandle->GetChildHandle(GET_MEMBER_NAME_CHECKED(FLiveLinkSourceBufferManagementSettings, bUseTimecodeSmoothLatest)));
 
 		CategoryBuilder.AddProperty(BufferSettingsPropertyHandle->GetChildHandle(GET_MEMBER_NAME_CHECKED(FLiveLinkSourceBufferManagementSettings, bKeepAtLeastOneFrame))
 			, EPropertyLocation::Advanced);
diff --git a/Engine/Plugins/VirtualProduction/TimedDataMonitor/Source/TimedDataMonitor/Private/TimedDataMonitorCalibration.cpp b/Engine/Plugins/VirtualProduction/TimedDataMonitor/Source/TimedDataMonitor/Private/TimedDataMonitorCalibration.cpp
--- a/Engine/Plugins/VirtualProduction/TimedDataMonitor/Source/TimedDataMonitor/Private/TimedDataMonitorCalibration.cpp
+++ b/Engine/Plugins/VirtualProduction/TimedDataMonitor/Source/TimedDataMonitor/Private/TimedDataMonitorCalibration.cpp
@@ -220,7 +220,7 @@ namespace TimedDataCalibration
 			check(ChannelItem);
 
 			TArray<FTimedDataChannelSampleTime> AllSamplesTimes = ChannelItem->GetDataTimes();
-			const double AverageSecondsBetweenSample = CalculateAverageInDeltaTimeBetweenSample(ETimedDataInputEvaluationType::Timecode, AllSamplesTimes);
+			const double AverageSecondsBetweenSample = CalculateAverageInDeltaTimeBetweenSample(EvaluationType, AllSamplesTimes);
 
 			double STDExtraTime = 0.0;
 			if (NumberOfSTD != 0)
@@ -673,11 +673,9 @@ FTimedDataMonitorTimeCorrectionResult FTimedDataMonitorCalibration::ApplyTimeCor
 		return Result;
 	}
 
+	//Don't touch inputs with evaluation type set to none. Just consider it a success
 	if (EvaluationType == ETimedDataInputEvaluationType::None)
 	{
-		// Set the evaluation offset of everyone to 0
-		TimedDataMonitorSubsystem->SetInputEvaluationOffsetInSeconds(InputIdentifier, 0.0);
-
 		Result.ReturnCode = ETimedDataMonitorTimeCorrectionReturnCode::Succeeded;
 		return Result;
 	}
diff --git a/Engine/Plugins/VirtualProduction/TimedDataMonitor/Source/TimedDataMonitor/Private/TimedDataMonitorSubsystem.cpp b/Engine/Plugins/VirtualProduction/TimedDataMonitor/Source/TimedDataMonitor/Private/TimedDataMonitorSubsystem.cpp
--- a/Engine/Plugins/VirtualProduction/TimedDataMonitor/Source/TimedDataMonitor/Private/TimedDataMonitorSubsystem.cpp
+++ b/Engine/Plugins/VirtualProduction/TimedDataMonitor/Source/TimedDataMonitor/Private/TimedDataMonitorSubsystem.cpp
@@ -620,12 +620,12 @@ ETimedDataMonitorEvaluationState UTimedDataMonitorSubsystem::GetInputEvaluationS
 {
 	BuildSourcesListIfNeeded();
 
-	ETimedDataMonitorEvaluationState WorstState = ETimedDataMonitorEvaluationState::NoSample;
+	ETimedDataMonitorEvaluationState WorstState = ETimedDataMonitorEvaluationState::Disabled;
 	if (const FTimeDataInputItem* InputItem = InputMap.Find(Identifier))
 	{
-		if (InputItem->ChannelIdentifiers.Num() > 0)
+		// Only consider inputs with a synchronizable evaluation type
+		if (InputItem->ChannelIdentifiers.Num() > 0 && GetInputEvaluationType(Identifier) != ETimedDataInputEvaluationType::None)
 		{
-			WorstState = ETimedDataMonitorEvaluationState::Disabled;
 			for (const FTimedDataMonitorChannelIdentifier& ChannelIdentifier : InputItem->ChannelIdentifiers)
 			{
 				const ETimedDataMonitorEvaluationState ChannelState = GetChannelEvaluationState(ChannelIdentifier);
diff --git a/Engine/Plugins/VirtualProduction/TimedDataMonitor/Source/TimedDataMonitorEditor/Private/STimingDiagramWidget.cpp b/Engine/Plugins/VirtualProduction/TimedDataMonitor/Source/TimedDataMonitorEditor/Private/STimingDiagramWidget.cpp
--- a/Engine/Plugins/VirtualProduction/TimedDataMonitor/Source/TimedDataMonitorEditor/Private/STimingDiagramWidget.cpp
+++ b/Engine/Plugins/VirtualProduction/TimedDataMonitor/Source/TimedDataMonitorEditor/Private/STimingDiagramWidget.cpp
@@ -225,7 +225,6 @@ public:
 	double OldestSigma = 0.0;
 	double NewestSigma = 0.0;
 	int32 NumberOfSigma = 3;
-	ETimedDataInputEvaluationType EvaluationType;
 	FSlateFontInfo FontInfo;
 };
 
@@ -275,7 +274,7 @@ void STimingDiagramWidget::UpdateCachedValue()
 	UTimedDataMonitorSubsystem* TimedDataMonitorSubsystem = GEngine->GetEngineSubsystem<UTimedDataMonitorSubsystem>();
 	check(TimedDataMonitorSubsystem);
 
-	GraphicWidget->EvaluationType = TimedDataMonitorSubsystem->GetInputEvaluationType(InputIdentifier);
+	const ETimedDataInputEvaluationType EvaluationType = TimedDataMonitorSubsystem->GetInputEvaluationType(InputIdentifier);
 
 	double EvaluationOffset = TimedDataMonitorSubsystem->GetInputEvaluationOffsetInSeconds(InputIdentifier);
 	double MinSampleTime = 0.0;
@@ -283,8 +282,8 @@ void STimingDiagramWidget::UpdateCachedValue()
 	double NewestMean = 0.0;
 	if (bIsInput)
 	{
-		MinSampleTime = TimedDataMonitorSubsystem->GetInputOldestDataTime(InputIdentifier).AsSeconds(GraphicWidget->EvaluationType);
-		MaxSampleTime = TimedDataMonitorSubsystem->GetInputNewestDataTime(InputIdentifier).AsSeconds(GraphicWidget->EvaluationType);
+		MinSampleTime = TimedDataMonitorSubsystem->GetInputOldestDataTime(InputIdentifier).AsSeconds(EvaluationType);
+		MaxSampleTime = TimedDataMonitorSubsystem->GetInputNewestDataTime(InputIdentifier).AsSeconds(EvaluationType);
 		GraphicWidget->OldestMean = TimedDataMonitorSubsystem->GetInputEvaluationDistanceToOldestSampleMean(InputIdentifier);
 		GraphicWidget->NewestMean = TimedDataMonitorSubsystem->GetInputEvaluationDistanceToNewestSampleMean(InputIdentifier);
 		GraphicWidget->OldestSigma = TimedDataMonitorSubsystem->GetInputEvaluationDistanceToOldestSampleStandardDeviation(InputIdentifier);
@@ -292,8 +291,8 @@ void STimingDiagramWidget::UpdateCachedValue()
 	}
 	else
 	{
-		MinSampleTime = TimedDataMonitorSubsystem->GetChannelOldestDataTime(ChannelIdentifier).AsSeconds(GraphicWidget->EvaluationType);
-		MaxSampleTime = TimedDataMonitorSubsystem->GetChannelNewestDataTime(ChannelIdentifier).AsSeconds(GraphicWidget->EvaluationType);
+		MinSampleTime = TimedDataMonitorSubsystem->GetChannelOldestDataTime(ChannelIdentifier).AsSeconds(EvaluationType);
+		MaxSampleTime = TimedDataMonitorSubsystem->GetChannelNewestDataTime(ChannelIdentifier).AsSeconds(EvaluationType);
 		GraphicWidget->OldestMean = TimedDataMonitorSubsystem->GetChannelEvaluationDistanceToOldestSampleMean(ChannelIdentifier);
 		GraphicWidget->NewestMean = TimedDataMonitorSubsystem->GetChannelEvaluationDistanceToNewestSampleMean(ChannelIdentifier);
 		GraphicWidget->OldestSigma = TimedDataMonitorSubsystem->GetChannelEvaluationDistanceToOldestSampleStandardDeviation(ChannelIdentifier);
@@ -302,7 +301,7 @@ void STimingDiagramWidget::UpdateCachedValue()
 	GraphicWidget->MinSampleTime = MinSampleTime + EvaluationOffset;
 	GraphicWidget->MaxSampleTime = MaxSampleTime + EvaluationOffset;
 
-	GraphicWidget->EvaluationTime = UTimedDataMonitorSubsystem::GetEvaluationTime(GraphicWidget->EvaluationType);
+	GraphicWidget->EvaluationTime = UTimedDataMonitorSubsystem::GetEvaluationTime(EvaluationType);
 }
 
 
diff --git a/Engine/Source/Runtime/LiveLinkInterface/Private/LiveLinkSourceSettings.cpp b/Engine/Source/Runtime/LiveLinkInterface/Private/LiveLinkSourceSettings.cpp
--- a/Engine/Source/Runtime/LiveLinkInterface/Private/LiveLinkSourceSettings.cpp
+++ b/Engine/Source/Runtime/LiveLinkInterface/Private/LiveLinkSourceSettings.cpp
@@ -18,7 +18,8 @@ bool ULiveLinkSourceSettings::CanEditChange(const FProperty* InProperty) const
 		if (InProperty->GetFName() == GET_MEMBER_NAME_CHECKED(FLiveLinkSourceBufferManagementSettings, TimecodeFrameOffset)
 			|| InProperty->GetFName() == GET_MEMBER_NAME_CHECKED(FLiveLinkSourceBufferManagementSettings, TimecodeFrameRate)
 			|| InProperty->GetFName() == GET_MEMBER_NAME_CHECKED(FLiveLinkSourceBufferManagementSettings, ValidTimecodeFrame)
-			|| InProperty->GetFName() == GET_MEMBER_NAME_CHECKED(FLiveLinkSourceBufferManagementSettings, bValidTimecodeFrameEnabled))
+			|| InProperty->GetFName() == GET_MEMBER_NAME_CHECKED(FLiveLinkSourceBufferManagementSettings, bValidTimecodeFrameEnabled)
+			|| InProperty->GetFName() == GET_MEMBER_NAME_CHECKED(FLiveLinkSourceBufferManagementSettings, bUseTimecodeSmoothLatest))
 		{
 			return Mode == ELiveLinkSourceMode::Timecode;
 		}
diff --git a/Engine/Source/Runtime/LiveLinkInterface/Public/LiveLinkSourceSettings.h b/Engine/Source/Runtime/LiveLinkInterface/Public/LiveLinkSourceSettings.h
--- a/Engine/Source/Runtime/LiveLinkInterface/Public/LiveLinkSourceSettings.h
+++ b/Engine/Source/Runtime/LiveLinkInterface/Public/LiveLinkSourceSettings.h
@@ -58,6 +58,14 @@ struct FLiveLinkSourceBufferManagementSettings
 	UPROPERTY(EditAnywhere, Category = "Settings")
 	bool bGenerateSubFrame = false;
 
+	/** When evaluating with timecode, align source timecode using a continuous clock offset to do a smooth latest 
+	 * This means that even if engine Timecode and source Timecode are not aligned, the offset between both clocks
+	 * will be tracked to keep them aligned. With an additionnal offset, 1.5 is a good number, you can evaluate
+	 * your subject using the latest frame by keeping just enough margin to have a smooth playback and avoid starving.
+	 */
+	UPROPERTY(EditAnywhere, Category = "Settings")
+	bool bUseTimecodeSmoothLatest = false;
+
 	/**
 	 * What is the source frame rate.
 	 * When the refresh rate of the source is bigger than the timecode frame rate, LiveLink will try to generate sub frame numbers.
diff --git a/Engine/Source/Runtime/LiveLinkInterface/Public/LiveLinkTypes.h b/Engine/Source/Runtime/LiveLinkInterface/Public/LiveLinkTypes.h
--- a/Engine/Source/Runtime/LiveLinkInterface/Public/LiveLinkTypes.h
+++ b/Engine/Source/Runtime/LiveLinkInterface/Public/LiveLinkTypes.h
@@ -100,6 +100,7 @@ public:
 	FLiveLinkWorldTime(const double InTime)
 		: Time(InTime)
 	{
+		//Initialize offset with an instantaneous offset, to be corrected at a later stage from a continuous calculation
 		Offset = FPlatformTime::Seconds() - InTime;
 	}
 
@@ -109,14 +110,37 @@ public:
 	{
 	}
 
-	double GetOffsettedTime() const { return Time + Offset; }
+	// Returns the raw time received from the sender
+	double GetSourceTime() const 
+	{
+		return Time;
+	}
+
+	// Returns offset between source time and engine time
+	double GetOffset() const
+	{
+		return Offset;
+	}
+
+	// Returns the time + the adjustment between source clock and engine clock to have a time comparable to engine one
+	double GetOffsettedTime() const 
+	{ 
+		return Time + Offset; 
+	}
+
+	// Adjust clock offset with a better evaluation of the difference between source clock and engine clock
+	void SetClockOffset(double ClockOffset) 
+	{
+		Offset = ClockOffset;
+	}
 
 private:
-	// Time for this frame. Used during interpolation. If this goes backwards we will dump already stored frames.
+	// SourceTime for this frame. Used during interpolation and to compute a running clock offset
 	UPROPERTY(meta = (IgnoreForMemberInitializationTest))
 	double Time;
 
 	// Value calculated on create to represent the different between the source time and client time
+	// Can also be updated afterwards if a better continuous offset is calculated
 	UPROPERTY()
 	double Offset;
 };
@@ -154,6 +178,8 @@ struct LIVELINKINTERFACE_API FLiveLinkMetaData
 	FQualifiedFrameTime SceneTime;
 };
 
+// Identifier assigned to each incoming frame once in the pipeline
+using FLiveLinkFrameIdentifier = int32;
 
 /**
  * Base data structure for each frame coming in for a subject
@@ -176,6 +202,12 @@ struct LIVELINKINTERFACE_API FLiveLinkBaseFrameData
 	UPROPERTY(EditAnywhere, BlueprintReadWrite, Category="LiveLink")
 	TArray<float> PropertyValues;
 
+	/** Time at which this frame was received - populated once inside pipeline */
+	FLiveLinkTime ArrivalTime;
+
+	/** This frame identifier - populated once inside pipeline*/
+	FLiveLinkFrameIdentifier FrameId = INDEX_NONE;
+
 	/** Return the LiveLinkTime struct constructed from this frame data. */
 	FLiveLinkTime GetLiveLinkTime() const { return FLiveLinkTime(WorldTime.GetOffsettedTime(), MetaData.SceneTime); }
 };
