- Fixing LiveLink MessageBus not loading its preset correctly

diff --git a/Engine/Plugins/Animation/LiveLink/Source/LiveLink/Private/LiveLinkClient.cpp b/Engine/Plugins/Animation/LiveLink/Source/LiveLink/Private/LiveLinkClient.cpp
--- a/Engine/Plugins/Animation/LiveLink/Source/LiveLink/Private/LiveLinkClient.cpp
+++ b/Engine/Plugins/Animation/LiveLink/Source/LiveLink/Private/LiveLinkClient.cpp
@@ -318,7 +318,20 @@ bool FLiveLinkClient::CreateSource(const FLiveLinkSourcePreset& InSourcePreset)
 	}
 
 	Data.Source = Source;
-	Setting = Data.Setting = DuplicateObject<ULiveLinkSourceSettings>(InSourcePreset.Settings, GetTransientPackage());
+
+	//In case a source has changed its source settings class, instead of duplicating, create the right one and copy previous properties
+	UClass* SourceSettingsClass = Source->GetSettingsClass().Get();
+	if (SourceSettingsClass && SourceSettingsClass != InSourcePreset.Settings->GetClass())
+	{
+		FLiveLinkLog::Info(TEXT("Creating Source '%s' from Preset: Settings class '%s' is not what is expected ('%s'). Updating to new class."), *InSourcePreset.SourceType.ToString(), *InSourcePreset.Settings->GetClass()->GetName(), *SourceSettingsClass->GetName());
+		Setting = NewObject<ULiveLinkSourceSettings>(GetTransientPackage(), SourceSettingsClass);
+		UEngine::CopyPropertiesForUnrelatedObjects(InSourcePreset.Settings, Setting);
+		Data.Setting = Setting;
+	}
+	else
+	{
+		Setting = Data.Setting = DuplicateObject<ULiveLinkSourceSettings>(InSourcePreset.Settings, GetTransientPackage());
+	}
 
 	Collection->AddSource(MoveTemp(Data));
 	Source->ReceiveClient(this, InSourcePreset.Guid);
@@ -1636,4 +1649,4 @@ void FLiveLinkClient_Base_DEPRECATED::RemoveSourceFromSubjectWhiteList(FName Sub
 	SetSubjectEnabled({ SourceGuid, SubjectName}, false);
 }
 
-PRAGMA_ENABLE_DEPRECATION_WARNINGS
\ No newline at end of file
+PRAGMA_ENABLE_DEPRECATION_WARNINGS
diff --git a/Engine/Plugins/Animation/LiveLink/Source/LiveLink/Private/LiveLinkMessageBusSource.cpp b/Engine/Plugins/Animation/LiveLink/Source/LiveLink/Private/LiveLinkMessageBusSource.cpp
--- a/Engine/Plugins/Animation/LiveLink/Source/LiveLink/Private/LiveLinkMessageBusSource.cpp
+++ b/Engine/Plugins/Animation/LiveLink/Source/LiveLink/Private/LiveLinkMessageBusSource.cpp
@@ -28,7 +28,6 @@ FLiveLinkMessageBusSource::FLiveLinkMessageBusSource(const FText& InSourceType,
 
 void FLiveLinkMessageBusSource::InitializeSettings(ULiveLinkSourceSettings* Settings)
 {
-	Settings->Mode = GetDefault<ULiveLinkSettings>()->DefaultMessageBusSourceMode;
 }
 
 void FLiveLinkMessageBusSource::ReceiveClient(ILiveLinkClient* InClient, FGuid InSourceGuid)
diff --git a/Engine/Plugins/Animation/LiveLink/Source/LiveLink/Private/LiveLinkMessageBusSource.h b/Engine/Plugins/Animation/LiveLink/Source/LiveLink/Private/LiveLinkMessageBusSource.h
--- a/Engine/Plugins/Animation/LiveLink/Source/LiveLink/Private/LiveLinkMessageBusSource.h
+++ b/Engine/Plugins/Animation/LiveLink/Source/LiveLink/Private/LiveLinkMessageBusSource.h
@@ -6,6 +6,7 @@
 
 #include "HAL/ThreadSafeBool.h"
 #include "IMessageContext.h"
+#include "LiveLinkMessageBusSourceSettings.h"
 #include "LiveLinkRole.h"
 #include "MessageEndpoint.h"
 
@@ -34,6 +35,8 @@ public:
 	virtual FText GetSourceType() const override { return SourceType; }
 	virtual FText GetSourceMachineName() const override { return SourceMachineName; }
 	virtual FText GetSourceStatus() const override;
+
+	virtual TSubclassOf<ULiveLinkSourceSettings> GetSettingsClass() const override { return ULiveLinkMessageBusSourceSettings::StaticClass(); }
 	//~ End ILiveLinkSource interface
 
 private:
diff --git a/Engine/Plugins/Animation/LiveLink/Source/LiveLink/Private/LiveLinkMessageBusSourceSettings.cpp b/Engine/Plugins/Animation/LiveLink/Source/LiveLink/Private/LiveLinkMessageBusSourceSettings.cpp
new file mode 100644
--- /dev/null
+++ b/Engine/Plugins/Animation/LiveLink/Source/LiveLink/Private/LiveLinkMessageBusSourceSettings.cpp
@@ -0,0 +1,11 @@
+// Copyright Epic Games, Inc. All Rights Reserved.
+
+#include "LiveLinkMessageBusSourceSettings.h"
+
+#include "LiveLinkSettings.h"
+#include "UObject/UObjectGlobals.h"
+
+ULiveLinkMessageBusSourceSettings::ULiveLinkMessageBusSourceSettings()
+{
+	Mode = GetDefault<ULiveLinkSettings>()->DefaultMessageBusSourceMode;
+}
diff --git a/Engine/Plugins/Animation/LiveLink/Source/LiveLink/Private/LiveLinkMessageBusSourceSettings.h b/Engine/Plugins/Animation/LiveLink/Source/LiveLink/Private/LiveLinkMessageBusSourceSettings.h
new file mode 100644
--- /dev/null
+++ b/Engine/Plugins/Animation/LiveLink/Source/LiveLink/Private/LiveLinkMessageBusSourceSettings.h
@@ -0,0 +1,24 @@
+// Copyright Epic Games, Inc. All Rights Reserved.
+
+#pragma once
+
+#include "LiveLinkSourceSettings.h"
+
+#include "LiveLinkMessageBusSourceSettings.generated.h"
+
+
+
+
+/**
+ * Settings for LiveLinkMessageBusSource.
+ * Used to apply default Evaluation mode from project settings when constructed
+ */
+UCLASS()
+class LIVELINK_API ULiveLinkMessageBusSourceSettings : public ULiveLinkSourceSettings
+{
+public:
+	GENERATED_BODY()
+
+public:
+	ULiveLinkMessageBusSourceSettings();
+};
