- Fixed livelink struct enum property binding writing to bad memory after fproperty refactor

#rnx

diff --git a/Engine/Plugins/Animation/LiveLink/Source/LiveLinkMovieScene/Private/MovieScene/MovieSceneLiveLinkStructPropertyBindings.cpp b/Engine/Plugins/Animation/LiveLink/Source/LiveLinkMovieScene/Private/MovieScene/MovieSceneLiveLinkStructPropertyBindings.cpp
--- a/Engine/Plugins/Animation/LiveLink/Source/LiveLinkMovieScene/Private/MovieScene/MovieSceneLiveLinkStructPropertyBindings.cpp
+++ b/Engine/Plugins/Animation/LiveLink/Source/LiveLinkMovieScene/Private/MovieScene/MovieSceneLiveLinkStructPropertyBindings.cpp
@@ -104,7 +104,7 @@ void FLiveLinkStructPropertyBindings::SetCurrentValueForEnumAt(int32 InIndex, co
 			if (FEnumProperty* EnumProperty = CastFieldChecked<FEnumProperty>(Property))
 			{
 				FNumericProperty* UnderlyingProperty = EnumProperty->GetUnderlyingProperty();
-				void* ValueAddr = EnumProperty->ContainerPtrToValuePtr<void>(InSourceAddress);
+				void* ValueAddr = FoundProperty.GetPropertyAddress<void>(InSourceAddress, InIndex);	
 				UnderlyingProperty->SetIntPropertyValue(ValueAddr, InValue);
 			}
 			else
