- Fixing time alignment in livelink by having the offset part of the source instead of each subjects

diff --git a/Engine/Plugins/Animation/LiveLink/Source/LiveLink/Private/ClockOffsetEstimatorRamp.h b/Engine/Plugins/Animation/LiveLink/Source/LiveLink/Private/ClockOffsetEstimatorRamp.h
--- a/Engine/Plugins/Animation/LiveLink/Source/LiveLink/Private/ClockOffsetEstimatorRamp.h
+++ b/Engine/Plugins/Animation/LiveLink/Source/LiveLink/Private/ClockOffsetEstimatorRamp.h
@@ -42,7 +42,7 @@ protected:
 	double MaxAllowedEstimationError = 0.25;
 
 	/** Correction step to apply, positive or negative depending on error sign */
-	double CorrectionStep = 250e-6;
+	double CorrectionStep = 100e-6;
 
 	/** Consecutive time a too big estimation error is required before clamping */
 	int32 ConsecutiveTooBigErrorRequired = 5;
diff --git a/Engine/Plugins/Animation/LiveLink/Source/LiveLink/Private/InterpolationProcessor/LiveLinkBasicFrameInterpolateProcessor.cpp b/Engine/Plugins/Animation/LiveLink/Source/LiveLink/Private/InterpolationProcessor/LiveLinkBasicFrameInterpolateProcessor.cpp
--- a/Engine/Plugins/Animation/LiveLink/Source/LiveLink/Private/InterpolationProcessor/LiveLinkBasicFrameInterpolateProcessor.cpp
+++ b/Engine/Plugins/Animation/LiveLink/Source/LiveLink/Private/InterpolationProcessor/LiveLinkBasicFrameInterpolateProcessor.cpp
@@ -251,8 +251,10 @@ void ULiveLinkBasicFrameInterpolationProcessor::FLiveLinkBasicFrameInterpolation
 	}
 
 	// Interpolate Time data
-	OutBlendedFrameData.GetBaseData()->WorldTime = FLiveLinkWorldTime(FMath::Lerp(FrameDataA.GetBaseData()->WorldTime.GetOffsettedTime(), FrameDataB.GetBaseData()->WorldTime.GetOffsettedTime(), InBlendWeight), 0.0);
+	OutBlendedFrameData.GetBaseData()->WorldTime = FLiveLinkWorldTime(FMath::Lerp(FrameDataA.GetBaseData()->WorldTime.GetSourceTime(), FrameDataB.GetBaseData()->WorldTime.GetSourceTime(), InBlendWeight), 0.0);
 	OutBlendedFrameData.GetBaseData()->MetaData.SceneTime.Time = FFrameTime(FMath::Lerp(FrameDataA.GetBaseData()->MetaData.SceneTime.Time, FrameDataB.GetBaseData()->MetaData.SceneTime.Time, InBlendWeight));
+	OutBlendedFrameData.GetBaseData()->ArrivalTime.WorldTime = FMath::Lerp(FrameDataA.GetBaseData()->ArrivalTime.WorldTime, FrameDataB.GetBaseData()->ArrivalTime.WorldTime, InBlendWeight);
+	OutBlendedFrameData.GetBaseData()->ArrivalTime.SceneTime.Time = FFrameTime(FMath::Lerp(FrameDataA.GetBaseData()->ArrivalTime.SceneTime.Time, FrameDataB.GetBaseData()->ArrivalTime.SceneTime.Time, InBlendWeight));
 
 	// Interpolate Property Values
 	if (Options.bInterpolatePropertyValues)
@@ -284,16 +286,34 @@ void ULiveLinkBasicFrameInterpolationProcessor::FLiveLinkBasicFrameInterpolation
 
 double ULiveLinkBasicFrameInterpolationProcessor::FLiveLinkBasicFrameInterpolationProcessorWorker::GetBlendFactor(double InTime, const FLiveLinkFrameDataStruct& FrameDataA, const FLiveLinkFrameDataStruct& FrameDataB)
 {
-	const double FrameATime = FrameDataA.GetBaseData()->WorldTime.GetOffsettedTime();
-	const double FrameBTime = FrameDataB.GetBaseData()->WorldTime.GetOffsettedTime();
-	return (InTime - FrameATime) / (FrameBTime - FrameATime);
+	const double FrameATime = FrameDataA.GetBaseData()->WorldTime.GetSourceTime();
+	const double FrameBTime = FrameDataB.GetBaseData()->WorldTime.GetSourceTime();
+	
+	const double Divider = FrameBTime - FrameATime;
+	if (!FMath::IsNearlyZero(Divider))
+	{
+		return (InTime - FrameATime) / Divider;
+	}
+	else
+	{
+		return 1.0;
+	}
 }
 
 double ULiveLinkBasicFrameInterpolationProcessor::FLiveLinkBasicFrameInterpolationProcessorWorker::GetBlendFactor(FQualifiedFrameTime InTime, const FLiveLinkFrameDataStruct& FrameDataA, const FLiveLinkFrameDataStruct& FrameDataB)
 {
 	const double FrameASeconds = FrameDataA.GetBaseData()->MetaData.SceneTime.AsSeconds();
 	const double FrameBSeconds = FrameDataB.GetBaseData()->MetaData.SceneTime.AsSeconds();
-	return (InTime.AsSeconds() - FrameASeconds) / (FrameBSeconds - FrameASeconds);
+
+	const double Divider = FrameBSeconds - FrameASeconds;
+	if (!FMath::IsNearlyZero(Divider))
+	{
+		return (InTime.AsSeconds() - FrameASeconds) / Divider;
+	}
+	else
+	{
+		return 1.0;
+	}
 }
 
  bool ULiveLinkBasicFrameInterpolationProcessor::FLiveLinkBasicFrameInterpolationProcessorWorker::FindInterpolateIndex(double InTime, const TArray<FLiveLinkFrameDataStruct>& InSourceFrames, int32& OutFrameIndexA, int32& OutFrameIndexB, FLiveLinkInterpolationInfo& OutInterpolationInfo)
@@ -303,19 +323,19 @@ double ULiveLinkBasicFrameInterpolationProcessor::FLiveLinkBasicFrameInterpolati
 		return false;
 	}
 
-	OutInterpolationInfo.ExpectedEvaluationDistanceFromNewestSeconds = InSourceFrames.Last().GetBaseData()->WorldTime.GetOffsettedTime() - InTime;
-	OutInterpolationInfo.ExpectedEvaluationDistanceFromOldestSeconds = InTime - InSourceFrames[0].GetBaseData()->WorldTime.GetOffsettedTime();
+	OutInterpolationInfo.ExpectedEvaluationDistanceFromNewestSeconds = InSourceFrames.Last().GetBaseData()->WorldTime.GetSourceTime() - InTime;
+	OutInterpolationInfo.ExpectedEvaluationDistanceFromOldestSeconds = InTime - InSourceFrames[0].GetBaseData()->WorldTime.GetSourceTime();
 
 	for (int32 FrameIndex = InSourceFrames.Num() - 1; FrameIndex >= 0; --FrameIndex)
 	{
 		const FLiveLinkFrameDataStruct& SourceFrameData = InSourceFrames[FrameIndex];
-		if (SourceFrameData.GetBaseData()->WorldTime.GetOffsettedTime() <= InTime)
+		if (SourceFrameData.GetBaseData()->WorldTime.GetSourceTime() <= InTime)
 		{
 			if (FrameIndex == InSourceFrames.Num() - 1)
 			{
 				OutFrameIndexA = FrameIndex;
 				OutFrameIndexB = FrameIndex;
-				OutInterpolationInfo.bOverflowDetected = !FMath::IsNearlyEqual(InTime, SourceFrameData.GetBaseData()->WorldTime.GetOffsettedTime());
+				OutInterpolationInfo.bOverflowDetected = !FMath::IsNearlyEqual(InTime, SourceFrameData.GetBaseData()->WorldTime.GetSourceTime());
 				return true;
 			}
 			else
diff --git a/Engine/Plugins/Animation/LiveLink/Source/LiveLink/Private/LiveLinkClient.cpp b/Engine/Plugins/Animation/LiveLink/Source/LiveLink/Private/LiveLinkClient.cpp
--- a/Engine/Plugins/Animation/LiveLink/Source/LiveLink/Private/LiveLinkClient.cpp
+++ b/Engine/Plugins/Animation/LiveLink/Source/LiveLink/Private/LiveLinkClient.cpp
@@ -622,17 +622,7 @@ void FLiveLinkClient::PushSubjectFrameData_AnyThread(const FLiveLinkSubjectKey&
 		}
 		else
 		{
-			//Stamp arrival time of each packet to track clock difference
-			if (SubjectFrame.FrameData.GetBaseData())
-			{
-				SubjectFrame.FrameData.GetBaseData()->ArrivalTime.WorldTime = FPlatformTime::Seconds();
-				const TOptional<FQualifiedFrameTime>& CurrentTime = FApp::GetCurrentFrameTime();
-				if (CurrentTime.IsSet())
-				{
-					SubjectFrame.FrameData.GetBaseData()->ArrivalTime.SceneTime = *CurrentTime;
-				}
-			}
-			
+
 			{
 				FScopeLock BroadcastLock(&SubjectFrameReceivedHandleseCriticalSection);
 				if (const FSubjectFramesReceivedHandles* Handles = SubjectFrameReceivedHandles.Find(InSubjectKey))
@@ -658,7 +648,7 @@ void FLiveLinkClient::PushSubjectFrameData_Internal(FPendingSubjectFrame&& Subje
 
 	check(Collection);
 
-	const FLiveLinkCollectionSourceItem* SourceItem = Collection->FindSource(SubjectFrameData.SubjectKey.Source);
+	FLiveLinkCollectionSourceItem* SourceItem = Collection->FindSource(SubjectFrameData.SubjectKey.Source);
 	if (SourceItem == nullptr || SourceItem->bPendingKill)
 	{
 		return;
@@ -707,11 +697,29 @@ void FLiveLinkClient::PushSubjectFrameData_Internal(FPendingSubjectFrame&& Subje
 		return;
 	}
 
+	//Stamp arrival time of each packet to track clock difference when it is effectively added to the stash. 
+	//Doing it in the Add_AnyThread would mean that we stamp it up to 1 frame time behind, causing the offset to always be 1 frame behind
+	//and requiring 2.5 frames or so to have a valid smooth offset
+	if (SubjectFrameData.FrameData.GetBaseData())
+	{
+		SubjectFrameData.FrameData.GetBaseData()->ArrivalTime.WorldTime = FPlatformTime::Seconds();
+		const TOptional<FQualifiedFrameTime>& CurrentTime = FApp::GetCurrentFrameTime();
+		if (CurrentTime.IsSet())
+		{
+			SubjectFrameData.FrameData.GetBaseData()->ArrivalTime.SceneTime = *CurrentTime;
+		}
+	}
+	
+	//Let source data know about this new frame to get latest clock offset
+	SourceItem->TimedData->ProcessNewFrameTimingInfo(*SubjectFrameData.FrameData.GetBaseData());
+
 	if (const FSubjectFramesAddedHandles* Handles = SubjectFrameAddedHandles.Find(SubjectFrameData.SubjectKey.SubjectName))
 	{
 		Handles->OnFrameDataAdded.Broadcast(SubjectItem->Key, Role, SubjectFrameData.FrameData);
 	}
 
+
+	//Finally, add the new frame to the subject. After this point, the frame data is unusable, it has been moved!
 	LinkSubject->AddFrameData(MoveTemp(SubjectFrameData.FrameData));
 }
 
diff --git a/Engine/Plugins/Animation/LiveLink/Source/LiveLink/Private/LiveLinkSettings.cpp b/Engine/Plugins/Animation/LiveLink/Source/LiveLink/Private/LiveLinkSettings.cpp
--- a/Engine/Plugins/Animation/LiveLink/Source/LiveLink/Private/LiveLinkSettings.cpp
+++ b/Engine/Plugins/Animation/LiveLink/Source/LiveLink/Private/LiveLinkSettings.cpp
@@ -11,7 +11,8 @@ FLiveLinkRoleProjectSetting::FLiveLinkRoleProjectSetting()
 
 
 ULiveLinkSettings::ULiveLinkSettings()
-	: DefaultMessageBusSourceMode(ELiveLinkSourceMode::EngineTime)
+	: ClockOffsetCorrectionStep(100e-6)
+	, DefaultMessageBusSourceMode(ELiveLinkSourceMode::EngineTime)
 	, MessageBusPingRequestFrequency(1.0)
 	, MessageBusHeartbeatFrequency(1.0)
 	, MessageBusHeartbeatTimeout(2.0)
diff --git a/Engine/Plugins/Animation/LiveLink/Source/LiveLink/Private/LiveLinkSubject.cpp b/Engine/Plugins/Animation/LiveLink/Source/LiveLink/Private/LiveLinkSubject.cpp
--- a/Engine/Plugins/Animation/LiveLink/Source/LiveLink/Private/LiveLinkSubject.cpp
+++ b/Engine/Plugins/Animation/LiveLink/Source/LiveLink/Private/LiveLinkSubject.cpp
@@ -39,8 +39,6 @@ void FLiveLinkSubject::Initialize(FLiveLinkSubjectKey InSubjectKey, TSubclassOf<
 	FrameData.Reset();
 	ReceivedOrderedFrames.Empty();
 	ResetBufferStats();
-	EngineTimeClockOffset.Reset();
-	TimecodeClockOffset.Reset();
 
 	if (TSharedPtr<FLiveLinkTimedDataInput> TimedDataGroupPin = TimedDataGroup.Pin())
 	{
@@ -73,13 +71,13 @@ void FLiveLinkSubject::Update()
 	{
 		if (CachedSettings.BufferSettings.bValidEngineTimeEnabled)
 		{
-			double ValidEngineTime = FApp::GetCurrentTime() - CachedSettings.BufferSettings.EngineTimeOffset - CachedSettings.BufferSettings.ValidEngineTime;
+			//Engine time threshold is made from current time minus all desired offset (user offset, validity offset, clock offset)
+			const double ValidEngineTime = FApp::GetCurrentTime() - CachedSettings.BufferSettings.EngineTimeOffset - CachedSettings.BufferSettings.ValidEngineTime - CachedSettings.BufferSettings.EngineTimeClockOffset;
 			int32 FrameIndex = 0;
 			for (const FLiveLinkFrameDataStruct& SourceFrameData : FrameData)
 			{
-				double FrameTime = SourceFrameData.GetBaseData()->WorldTime.GetOffsettedTime();
-				double OffsetTime = ValidEngineTime;
-				if (FrameTime > OffsetTime)
+				const double FrameTime = SourceFrameData.GetBaseData()->WorldTime.GetSourceTime();
+				if (FrameTime > ValidEngineTime)
 				{
 					break;
 				}
@@ -105,7 +103,7 @@ void FLiveLinkSubject::Update()
 				FFrameTime ClockOffset;
 				if (CachedSettings.BufferSettings.bUseTimecodeSmoothLatest)
 				{
-					ClockOffset = CachedSettings.BufferSettings.TimecodeFrameRate.AsFrameTime(TimecodeClockOffset.GetEstimatedOffset());
+					ClockOffset = CachedSettings.BufferSettings.TimecodeFrameRate.AsFrameTime(CachedSettings.BufferSettings.TimecodeClockOffset);
 				}
 
 				const FFrameTime CurrentFrameTimeInFrameSpace = CurrentSyncTime.ConvertTo(CachedSettings.BufferSettings.TimecodeFrameRate);
@@ -149,7 +147,7 @@ void FLiveLinkSubject::Update()
 				FFrameTime ClockOffset;
 				if (CachedSettings.BufferSettings.bUseTimecodeSmoothLatest)
 				{
-					ClockOffset = EngineTimecode.Rate.AsFrameTime(TimecodeClockOffset.GetEstimatedOffset());
+					ClockOffset = EngineTimecode.Rate.AsFrameTime(CachedSettings.BufferSettings.TimecodeClockOffset);
 				}
 
 				const FFrameTime FrameOffset = FQualifiedFrameTime(FFrameTime::FromDecimal(CachedSettings.BufferSettings.TimecodeFrameOffset), CachedSettings.BufferSettings.TimecodeFrameRate).ConvertTo(EngineTimecode.Rate);
@@ -167,7 +165,8 @@ void FLiveLinkSubject::Update()
 
 		case ELiveLinkSourceMode::EngineTime:
 		{
-			const double ReadTime = FApp::GetCurrentTime() - CachedSettings.BufferSettings.EngineTimeOffset;
+			//Compute the evaluation time using all offsets (user offset + clock offset)
+			const double ReadTime = FApp::GetCurrentTime() - CachedSettings.BufferSettings.EngineTimeOffset - CachedSettings.BufferSettings.EngineTimeClockOffset;
 			bSnapshotIsValid = GetFrameAtWorldTime(ReadTime, FrameSnapshot);
 			break;
 		}
@@ -187,7 +186,7 @@ void FLiveLinkSubject::Update()
 	{
 		//Extra verbose logging about snapshot just created
 		const FTimecode FrameTC = FTimecode::FromFrameNumber(FrameSnapshot.FrameData.GetBaseData()->MetaData.SceneTime.Time.GetFrame(), FrameSnapshot.FrameData.GetBaseData()->MetaData.SceneTime.Rate);
-		UE_LOG(LogLiveLink, VeryVerbose, TEXT("Subject '%s' created snapshot with frame ID %d - Timecode:[%s.%0.3f] - SourceTime: %0.4f, Offset: %0.6f, CorrectedTime: %0.4f"), *SubjectKey.SubjectName.ToString(), FrameSnapshot.FrameData.GetBaseData()->FrameId, *FrameTC.ToString(), FrameSnapshot.FrameData.GetBaseData()->MetaData.SceneTime.Time.GetSubFrame(), FrameSnapshot.FrameData.GetBaseData()->WorldTime.GetSourceTime(), FrameSnapshot.FrameData.GetBaseData()->WorldTime.GetOffset(), FrameSnapshot.FrameData.GetBaseData()->WorldTime.GetOffsettedTime());
+		UE_LOG(LogLiveLink, Verbose, TEXT("Subject '%s' created snapshot with frame ID %d - Timecode:[%s.%0.3f] - SourceTime: %0.4f, Offset: %0.6f, CorrectedTime: %0.4f"), *SubjectKey.SubjectName.ToString(), FrameSnapshot.FrameData.GetBaseData()->FrameId, *FrameTC.ToString(), FrameSnapshot.FrameData.GetBaseData()->MetaData.SceneTime.Time.GetSubFrame(), FrameSnapshot.FrameData.GetBaseData()->WorldTime.GetSourceTime(), FrameSnapshot.FrameData.GetBaseData()->WorldTime.GetOffset(), FrameSnapshot.FrameData.GetBaseData()->WorldTime.GetOffsettedTime());
 	}
 }
 
@@ -210,7 +209,8 @@ TArray<FLiveLinkTime> FLiveLinkSubject::GetFrameTimes() const
 	Result.Reset(FrameData.Num());
 	for (const FLiveLinkFrameDataStruct& Data : FrameData)
 	{
-		Result.Emplace(Data.GetBaseData()->WorldTime.GetOffsettedTime(), Data.GetBaseData()->MetaData.SceneTime);
+		//When giving out frame times, include the latest Source clock offset
+		Result.Emplace(Data.GetBaseData()->WorldTime.GetSourceTime() + CachedSettings.BufferSettings.EngineTimeClockOffset, Data.GetBaseData()->MetaData.SceneTime);
 	}
 	return Result;
 }
@@ -241,8 +241,8 @@ bool FLiveLinkSubject::EvaluateFrameAtWorldTime(double InWorldTime, TSubclassOf<
 	bool bSuccess = false;
 	if (FrameData.Num() != 0)
 	{
-		// Adjust desired read time based on offset settings
-		const double ReadTime = InWorldTime - CachedSettings.BufferSettings.EngineTimeOffset;
+		// Adjust desired read time based on offset settings (User desired offset and continuous clock offset)
+		const double ReadTime = InWorldTime - CachedSettings.BufferSettings.EngineTimeOffset - CachedSettings.BufferSettings.EngineTimeClockOffset;
 		if (Role == InDesiredRole || Role->IsChildOf(InDesiredRole))
 		{
 			GetFrameAtWorldTime(ReadTime, OutFrame);
@@ -294,7 +294,7 @@ bool FLiveLinkSubject::EvaluateFrameAtSceneTime(const FQualifiedFrameTime& InSce
 		FFrameTime ClockOffset;
 		if (CachedSettings.BufferSettings.bUseTimecodeSmoothLatest)
 		{
-			ClockOffset = InSceneTime.Rate.AsFrameTime(TimecodeClockOffset.GetEstimatedOffset());
+			ClockOffset = InSceneTime.Rate.AsFrameTime(CachedSettings.BufferSettings.TimecodeClockOffset);
 		}
 
 		const FFrameTime FrameOffset = FQualifiedFrameTime(FFrameTime::FromDecimal(CachedSettings.BufferSettings.TimecodeFrameOffset), CachedSettings.BufferSettings.TimecodeFrameRate).ConvertTo(InSceneTime.Rate);
@@ -362,13 +362,6 @@ void FLiveLinkSubject::AddFrameData(FLiveLinkFrameDataStruct&& InFrameData)
 		RemoveFrames(1);
 	}
 
-	//Update clock trackers
-	UpdateClockTracker(*InFrameData.GetBaseData());
-
-	//Update this frame clock offset with the newly updated one. This will make it transparent everywhere we used the offsetted time.
-	//Done before selecting a frame offset to have the most accurate offset
-	InFrameData.GetBaseData()->WorldTime.SetClockOffset(EngineTimeClockOffset.GetEstimatedOffset());
-	
 	int32 FrameIndex = INDEX_NONE;
 	switch (CachedSettings.SourceMode)
 	{
@@ -399,7 +392,7 @@ void FLiveLinkSubject::AddFrameData(FLiveLinkFrameDataStruct&& InFrameData)
 		
 		//Extra verbose logging about new frame being added
 		const FTimecode FrameTC = FTimecode::FromFrameNumber(InFrameData.GetBaseData()->MetaData.SceneTime.Time.GetFrame(), InFrameData.GetBaseData()->MetaData.SceneTime.Rate);
-		UE_LOG(LogLiveLink, VeryVerbose, TEXT("Subject '%s' adding frame ID %d at index %d - Timecode:[%s.%0.3f] - SourceTime: %0.4f, Offset: %0.6f, CorrectedTime: %0.4f"), *SubjectKey.SubjectName.ToString(), ThisFrameIdentifier, FrameIndex, *FrameTC.ToString(), InFrameData.GetBaseData()->MetaData.SceneTime.Time.GetSubFrame(), InFrameData.GetBaseData()->WorldTime.GetSourceTime(), InFrameData.GetBaseData()->WorldTime.GetOffset(), InFrameData.GetBaseData()->WorldTime.GetOffsettedTime());
+		UE_LOG(LogLiveLink, Verbose, TEXT("Subject '%s' adding frame ID %d at index %d - Timecode:[%s.%0.3f] - SourceTime: %0.4f, Offset: %0.6f, CorrectedTime: %0.4f"), *SubjectKey.SubjectName.ToString(), ThisFrameIdentifier, FrameIndex, *FrameTC.ToString(), InFrameData.GetBaseData()->MetaData.SceneTime.Time.GetSubFrame(), InFrameData.GetBaseData()->WorldTime.GetSourceTime(), InFrameData.GetBaseData()->WorldTime.GetOffset(), InFrameData.GetBaseData()->WorldTime.GetOffsettedTime());
 		
 		FrameData.Insert(MoveTemp(InFrameData), FrameIndex);
 
@@ -421,9 +414,10 @@ int32 FLiveLinkSubject::FindNewFrame_WorldTime(const FLiveLinkWorldTime& WorldTi
 {
 	if (CachedSettings.BufferSettings.bValidEngineTimeEnabled)
 	{
-		const double ValidEngineTime = FApp::GetCurrentTime() - CachedSettings.BufferSettings.EngineTimeOffset - CachedSettings.BufferSettings.ValidEngineTime;
-		const double WorldOffsettedTime = WorldTime.GetOffsettedTime();
-		if (WorldOffsettedTime < ValidEngineTime)
+		//Offset the reading position with clock offset and use source time directly
+		const double ValidEngineTime = FApp::GetCurrentTime() - CachedSettings.BufferSettings.EngineTimeOffset - CachedSettings.BufferSettings.ValidEngineTime - CachedSettings.BufferSettings.EngineTimeClockOffset;
+		const double SourceTime = WorldTime.GetSourceTime();
+		if (SourceTime < ValidEngineTime)
 		{
 			static const FName NAME_InvalidWorldTime = "LiveLinkSubject_InvalidWorldTIme";
 			FLiveLinkLog::WarningOnce(NAME_InvalidWorldTime, SubjectKey, TEXT("Trying to add a frame in which the world time has a value too low compare to the engine's time. Do you have an invalid offset? The Subject is '%s'."), *SubjectKey.SubjectName.ToString());
@@ -435,14 +429,15 @@ int32 FLiveLinkSubject::FindNewFrame_WorldTime(const FLiveLinkWorldTime& WorldTi
 
 int32 FLiveLinkSubject::FindNewFrame_WorldTimeInternal(const FLiveLinkWorldTime& WorldTime) const
 {
+	//When inserting our new frame, use the source time of each packet. 
 	int32 FrameIndex = FrameData.Num() - 1;
-	const double NewFrameOffsettedTime = WorldTime.GetOffsettedTime();
+	const double NewFrameSourceTime = WorldTime.GetSourceTime();
 	for (; FrameIndex >= 0; --FrameIndex)
 	{
-		const double FrameOffsettedTime = FrameData[FrameIndex].GetBaseData()->WorldTime.GetOffsettedTime();
-		if (FrameOffsettedTime <= NewFrameOffsettedTime)
+		const double FrameSourceTime = FrameData[FrameIndex].GetBaseData()->WorldTime.GetSourceTime();
+		if (FrameSourceTime <= NewFrameSourceTime)
 		{
-			if (FMath::IsNearlyEqual(FrameOffsettedTime, NewFrameOffsettedTime))
+			if (FMath::IsNearlyEqual(FrameSourceTime, NewFrameSourceTime))
 			{
 				static const FName NAME_SameWorldTime = "LiveLinkSubject_SameWorldTime";
 				FLiveLinkLog::WarningOnce(NAME_SameWorldTime, SubjectKey, TEXT("A new frame data for subjet '%s' has the same time as a previous frame."), *SubjectKey.SubjectName.ToString());
@@ -473,9 +468,15 @@ int32 FLiveLinkSubject::FindNewFrame_SceneTime(const FQualifiedFrameTime& Qualif
 	// If we do not have a TC set, keep buffering, the TC may be unresponsive for a moment. We do not want to loose data.
 	if (FApp::GetCurrentFrameTime().IsSet() && CachedSettings.BufferSettings.bValidTimecodeFrameEnabled)
 	{
+		FFrameTime TimecodeClockOffsetTime;
+		if (CachedSettings.BufferSettings.bUseTimecodeSmoothLatest)
+		{
+			TimecodeClockOffsetTime = CachedSettings.BufferSettings.TimecodeFrameRate.AsFrameTime(CachedSettings.BufferSettings.TimecodeClockOffset);
+		}
+
 		const FQualifiedFrameTime CurrentSyncTime = FApp::GetCurrentFrameTime().GetValue();
 		const FFrameTime CurrentFrameTimeInFrameSpace = CurrentSyncTime.ConvertTo(CachedSettings.BufferSettings.TimecodeFrameRate);
-		const FFrameTime CurrentOffsetFrameTime = CurrentFrameTimeInFrameSpace - FFrameTime::FromDecimal(CachedSettings.BufferSettings.TimecodeFrameOffset) - CachedSettings.BufferSettings.ValidTimecodeFrame;
+		const FFrameTime CurrentOffsetFrameTime = CurrentFrameTimeInFrameSpace - FFrameTime::FromDecimal(CachedSettings.BufferSettings.TimecodeFrameOffset) - CachedSettings.BufferSettings.ValidTimecodeFrame - TimecodeClockOffsetTime;
 		if (QualifiedFrameTime.Time.AsDecimal() < CurrentOffsetFrameTime.AsDecimal())
 		{
 			static const FName NAME_InvalidTC = "LiveLinkSubject_InvalidTC";
@@ -517,14 +518,15 @@ int32 FLiveLinkSubject::FindNewFrame_SceneTime(const FQualifiedFrameTime& Qualif
 		}
 		--MaxInclusive;
 
-		const double NewFrameOffsettedTime = WorldTime.GetOffsettedTime();
+		//When ordering sub frames, use source time to see which frame comes first
+		const double NewFrameSourceTime = WorldTime.GetSourceTime();
 		int32 FrameIndex = MaxInclusive;
 		for (; FrameIndex >= MinInclusive; --FrameIndex)
 		{
-			const double FrameOffsettedTime = FrameData[FrameIndex].GetBaseData()->WorldTime.GetOffsettedTime();
-			if (FrameOffsettedTime <= NewFrameOffsettedTime)
+			const double FrameSourceTime = FrameData[FrameIndex].GetBaseData()->WorldTime.GetSourceTime();
+			if (FrameSourceTime <= NewFrameSourceTime)
 			{
-				if (FMath::IsNearlyEqual(FrameOffsettedTime, NewFrameOffsettedTime))
+				if (FMath::IsNearlyEqual(FrameSourceTime, NewFrameSourceTime))
 				{
 					static const FName NAME_SameWorldSceneTime = "LiveLinkSubject_SameWorldSceneTime";
 					FLiveLinkLog::WarningOnce(NAME_SameWorldSceneTime, SubjectKey, TEXT("A new frame data for subjet '%s' has the same timecode and the same time as a previous frame."), *SubjectKey.SubjectName.ToString());
@@ -606,7 +608,9 @@ void FLiveLinkSubject::AdjustSubFrame_SceneTime(int32 InFrameIndex)
 		{
 			FORCEINLINE bool operator()(const FLiveLinkFrameDataStruct& A, const FLiveLinkFrameDataStruct& B) const
 			{
-				return A.GetBaseData()->WorldTime.GetOffsettedTime() < B.GetBaseData()->WorldTime.GetOffsettedTime();
+				//Use SourceTime to make final granular sorting when generating subframes. 
+				//@note Should be updated to have frame numbers associated to each packet to correctly apply subframes
+				return A.GetBaseData()->WorldTime.GetSourceTime() < B.GetBaseData()->WorldTime.GetSourceTime();
 			}
 		};
 		Sort(&FrameData[LowerInclusiveLimit], HigherInclusiveLimit - LowerInclusiveLimit + 1, TLess());
@@ -647,6 +651,9 @@ bool FLiveLinkSubject::GetFrameAtWorldTime(const double InSeconds, FLiveLinkSubj
 		{
 			OutFrame.StaticData.InitializeWith(StaticData.GetStruct(), StaticData.GetBaseData());
 		}
+
+		//Apply clock offset of our source so outside frame has the current offset built in
+		OutFrame.FrameData.GetBaseData()->WorldTime.SetClockOffset(CachedSettings.BufferSettings.EngineTimeClockOffset);
 	}
 	else if (IsBufferStatsEnabled())
 	{
@@ -664,8 +671,8 @@ bool FLiveLinkSubject::GetFrameAtWorldTime_Closest(const double InSeconds, FLive
 	bool bBuiltFrame = false;
 	for (int32 FrameIndex = FrameData.Num() - 1; FrameIndex >= 0; --FrameIndex)
 	{
-		const double Time = FrameData[FrameIndex].GetBaseData()->WorldTime.GetOffsettedTime();
-		if (Time <= InSeconds)
+		const double FrameTime = FrameData[FrameIndex].GetBaseData()->WorldTime.GetSourceTime();
+		if (FrameTime <= InSeconds)
 		{
 			if (FrameIndex == FrameData.Num() - 1)
 			{
@@ -674,14 +681,14 @@ bool FLiveLinkSubject::GetFrameAtWorldTime_Closest(const double InSeconds, FLive
 				bBuiltFrame = true;
 
 				//If we tried to read above our buffer, stamp an overflow
-				bUnderflowDetected = !FMath::IsNearlyEqual(Time, InSeconds);
+				bUnderflowDetected = !FMath::IsNearlyEqual(FrameTime, InSeconds);
 				break;
 			}
 			else
 			{
-				const double NextTime = FrameData[FrameIndex + 1].GetBaseData()->WorldTime.GetOffsettedTime();
-				const float BlendWeight = (InSeconds - NextTime) / (NextTime - Time);
-				int32 CopyIndex = (BlendWeight > 0.5f) ? FrameIndex : FrameIndex + 1;
+				const double NextTime = FrameData[FrameIndex + 1].GetBaseData()->WorldTime.GetSourceTime();
+				const float BlendWeight = (InSeconds - NextTime) / (NextTime - FrameTime);
+				const int32 CopyIndex = (BlendWeight > 0.5f) ? FrameIndex : FrameIndex + 1;
 				OutFrame.FrameData.InitializeWith(FrameData[CopyIndex].GetStruct(), FrameData[CopyIndex].GetBaseData());
 				bBuiltFrame = true;
 				break;
@@ -709,8 +716,8 @@ bool FLiveLinkSubject::GetFrameAtWorldTime_Closest(const double InSeconds, FLive
 		}
 
 		FTimedDataInputEvaluationData EvaluationData;
-		EvaluationData.DistanceToNewestSampleSeconds = FrameData[FrameData.Num()-1].GetBaseData()->WorldTime.GetOffsettedTime() - InSeconds;
-		EvaluationData.DistanceToOldestSampleSeconds = InSeconds - FrameData[0].GetBaseData()->WorldTime.GetOffsettedTime();
+		EvaluationData.DistanceToNewestSampleSeconds = FrameData[FrameData.Num()-1].GetBaseData()->WorldTime.GetSourceTime() - InSeconds;
+		EvaluationData.DistanceToOldestSampleSeconds = InSeconds - FrameData[0].GetBaseData()->WorldTime.GetSourceTime();
 		UpdateEvaluationData(EvaluationData);
 	}
 
@@ -930,12 +937,6 @@ void FLiveLinkSubject::UpdateEvaluationData(const FTimedDataInputEvaluationData&
 	EvaluationStatistics.LastEvaluationData = EvaluationData;
 }
 
-void FLiveLinkSubject::UpdateClockTracker(const FLiveLinkBaseFrameData& NewFrame)
-{
-	EngineTimeClockOffset.UpdateEstimation(NewFrame.WorldTime.GetSourceTime(), NewFrame.ArrivalTime.WorldTime);
-	TimecodeClockOffset.UpdateEstimation(NewFrame.MetaData.SceneTime.AsSeconds(), NewFrame.ArrivalTime.SceneTime.AsSeconds());
-}
-
 void FLiveLinkSubject::RemoveFrames(int32 InCount)
 {
 	for (int32 i = 0; i < InCount; ++i)
@@ -1110,11 +1111,12 @@ FTimedDataChannelSampleTime FLiveLinkSubject::GetOldestDataTime() const
 		FFrameTime TimecodeClockOffsetTime;
 		if (CachedSettings.BufferSettings.bUseTimecodeSmoothLatest)
 		{
-			TimecodeClockOffsetTime = FrameData[0].GetBaseData()->MetaData.SceneTime.Rate.AsFrameTime(TimecodeClockOffset.GetEstimatedOffset());
+			TimecodeClockOffsetTime = FrameData[0].GetBaseData()->MetaData.SceneTime.Rate.AsFrameTime(CachedSettings.BufferSettings.TimecodeClockOffset);
 		}
 
+		const double OffsettedOldestTime = FrameData[0].GetBaseData()->WorldTime.GetSourceTime() + CachedSettings.BufferSettings.EngineTimeClockOffset;
 		const FQualifiedFrameTime OldestSceneTime = FQualifiedFrameTime(FrameData[0].GetBaseData()->MetaData.SceneTime.Time + TimecodeClockOffsetTime, FrameData[0].GetBaseData()->MetaData.SceneTime.Rate);
-		return FTimedDataChannelSampleTime(FrameData[0].GetBaseData()->WorldTime.GetOffsettedTime(), OldestSceneTime);
+		return FTimedDataChannelSampleTime(OffsettedOldestTime, OldestSceneTime);
 	}
 	return FTimedDataChannelSampleTime();
 }
@@ -1127,31 +1129,37 @@ FTimedDataChannelSampleTime FLiveLinkSubject::GetNewestDataTime() const
 		FFrameTime TimecodeClockOffsetTime;
 		if (CachedSettings.BufferSettings.bUseTimecodeSmoothLatest)
 		{
-			TimecodeClockOffsetTime = FrameData.Last().GetBaseData()->MetaData.SceneTime.Rate.AsFrameTime(TimecodeClockOffset.GetEstimatedOffset());
+			TimecodeClockOffsetTime = FrameData.Last().GetBaseData()->MetaData.SceneTime.Rate.AsFrameTime(CachedSettings.BufferSettings.TimecodeClockOffset);
 		}
 
+		const double OffsettedNewestTime = FrameData.Last().GetBaseData()->WorldTime.GetSourceTime() + CachedSettings.BufferSettings.EngineTimeClockOffset;
 		const FQualifiedFrameTime NewestSceneTime = FQualifiedFrameTime(FrameData.Last().GetBaseData()->MetaData.SceneTime.Time + TimecodeClockOffsetTime, FrameData.Last().GetBaseData()->MetaData.SceneTime.Rate);
-		return FTimedDataChannelSampleTime(FrameData.Last().GetBaseData()->WorldTime.GetOffsettedTime(), NewestSceneTime);
+		return FTimedDataChannelSampleTime(OffsettedNewestTime, NewestSceneTime);
 	}
 	return FTimedDataChannelSampleTime();
 }
 
 TArray<FTimedDataChannelSampleTime> FLiveLinkSubject::GetDataTimes() const
 {
-	//For timed data, adjust frame timings to match with clock offset to have valid timing diagrams
-	FFrameTime TimecodeClockOffsetTime;
-	if (CachedSettings.BufferSettings.bUseTimecodeSmoothLatest)
-	{
-		TimecodeClockOffsetTime = FrameData[0].GetBaseData()->MetaData.SceneTime.Rate.AsFrameTime(TimecodeClockOffset.GetEstimatedOffset());;
-	}
-	
 	TArray<FTimedDataChannelSampleTime> Result;
-	Result.Reset(FrameData.Num());
-	for (const FLiveLinkFrameDataStruct& Data : FrameData)
+	if (FrameData.Num() > 0)
 	{
-		const FQualifiedFrameTime AdjustedSceneTime = FQualifiedFrameTime(Data.GetBaseData()->MetaData.SceneTime.Time + TimecodeClockOffsetTime, Data.GetBaseData()->MetaData.SceneTime.Rate);
-		Result.Emplace(Data.GetBaseData()->WorldTime.GetOffsettedTime(), AdjustedSceneTime);
+		//For timed data, adjust frame timings to match with clock offset to have valid timing diagrams
+		FFrameTime TimecodeClockOffsetTime;
+		if (CachedSettings.BufferSettings.bUseTimecodeSmoothLatest)
+		{
+			TimecodeClockOffsetTime = FrameData[0].GetBaseData()->MetaData.SceneTime.Rate.AsFrameTime(CachedSettings.BufferSettings.TimecodeClockOffset);
+		}
+
+		Result.Reset(FrameData.Num());
+		for (const FLiveLinkFrameDataStruct& Data : FrameData)
+		{
+			const double OffsettedEngineTime = Data.GetBaseData()->WorldTime.GetSourceTime() + CachedSettings.BufferSettings.EngineTimeClockOffset;
+			const FQualifiedFrameTime AdjustedSceneTime = FQualifiedFrameTime(Data.GetBaseData()->MetaData.SceneTime.Time + TimecodeClockOffsetTime, Data.GetBaseData()->MetaData.SceneTime.Rate);
+			Result.Emplace(OffsettedEngineTime, AdjustedSceneTime);
+		}
 	}
+	
 	return Result;
 }
 
diff --git a/Engine/Plugins/Animation/LiveLink/Source/LiveLink/Private/LiveLinkSubject.h b/Engine/Plugins/Animation/LiveLink/Source/LiveLink/Private/LiveLinkSubject.h
--- a/Engine/Plugins/Animation/LiveLink/Source/LiveLink/Private/LiveLinkSubject.h
+++ b/Engine/Plugins/Animation/LiveLink/Source/LiveLink/Private/LiveLinkSubject.h
@@ -4,7 +4,6 @@
 
 #include "ILiveLinkSubject.h"
 
-#include "ClockOffsetEstimatorRamp.h"
 #include "Containers/Queue.h"
 #include "ILiveLinkSource.h"
 #include "ITimedDataInput.h"
@@ -143,9 +142,6 @@ private:
 	void IncreaseBufferOverFlowStat();
 	void UpdateEvaluationData(const FTimedDataInputEvaluationData& EvaluationData);
 
-	//Tracks clock difference between subject time and engine time.
-	void UpdateClockTracker(const FLiveLinkBaseFrameData& NewFrame);
-
 	//Remove frames from our buffer - based on receiving order
 	void RemoveFrames(int32 Count);
 
@@ -212,10 +208,6 @@ private:
 	// Some stats compiled by the subject.
 	FSubjectEvaluationStatistics EvaluationStatistics;
 
-	// Keep track of offsets between subject time and engine time
-	FClockOffsetEstimatorRamp EngineTimeClockOffset;
-	FClockOffsetEstimatorRamp TimecodeClockOffset;
-
 	/** 
 	 * Evaluation can be done on any thread so we need to protect statistic logging 
 	 * Some stats requires more than atomic sized vars so a critical section is used to protect when necessary
diff --git a/Engine/Plugins/Animation/LiveLink/Source/LiveLink/Private/LiveLinkTimedDataInput.cpp b/Engine/Plugins/Animation/LiveLink/Source/LiveLink/Private/LiveLinkTimedDataInput.cpp
--- a/Engine/Plugins/Animation/LiveLink/Source/LiveLink/Private/LiveLinkTimedDataInput.cpp
+++ b/Engine/Plugins/Animation/LiveLink/Source/LiveLink/Private/LiveLinkTimedDataInput.cpp
@@ -33,6 +33,13 @@ namespace LiveLinkTimedDataInput
 		ECVF_Default
 	);
 
+	TAutoConsoleVariable<bool> CVarLiveLinkUpdateContinuousClockOffset(
+		TEXT("LiveLink.TimedDataInput.UpdateClockOffset"),
+		true,
+		TEXT("By default, clock offset is continuously updated for each source. You can pause it if desired with this cvar and offset will be fixed to its value."),
+		ECVF_Default
+	);
+
 
 	ETimedDataInputEvaluationType ToTimedDataInputEvaluationType(ELiveLinkSourceMode SourceMode)
 	{
@@ -71,6 +78,9 @@ FLiveLinkTimedDataInput::FLiveLinkTimedDataInput(FLiveLinkClient* InClient, FGui
 	, Source(InSource)
 {
 	ITimeManagementModule::Get().GetTimedDataInputCollection().Add(this);
+
+	EngineClockOffset.SetCorrectionStep(GetDefault<ULiveLinkSettings>()->ClockOffsetCorrectionStep);
+	TimecodeClockOffset.SetCorrectionStep(GetDefault<ULiveLinkSettings>()->ClockOffsetCorrectionStep);
 }
 
 FLiveLinkTimedDataInput::~FLiveLinkTimedDataInput()
@@ -184,4 +194,35 @@ const FSlateBrush* FLiveLinkTimedDataInput::GetDisplayIcon() const
 }
 #endif
 
+void FLiveLinkTimedDataInput::ProcessNewFrameTimingInfo(FLiveLinkBaseFrameData& NewFrameData)
+{
+	if (LiveLinkTimedDataInput::CVarLiveLinkUpdateContinuousClockOffset.GetValueOnGameThread())
+	{
+		//Update both clock offsets for each frame received for our subjects
+		//We mark the last source/frame time that we used to update our offset to only update once per source frame
+		if (!FMath::IsNearlyEqual(NewFrameData.WorldTime.GetSourceTime(), LastWorldSourceTime))
+		{
+			LastWorldSourceTime = NewFrameData.WorldTime.GetSourceTime();
+			EngineClockOffset.UpdateEstimation(NewFrameData.WorldTime.GetSourceTime(), NewFrameData.ArrivalTime.WorldTime);
+		}
+
+		const double NewFrameSceneTime = NewFrameData.MetaData.SceneTime.AsSeconds();
+		if(!FMath::IsNearlyEqual(NewFrameSceneTime, LastSceneTime))
+		{
+			LastSceneTime = NewFrameSceneTime;
+			TimecodeClockOffset.UpdateEstimation(NewFrameSceneTime, NewFrameData.ArrivalTime.SceneTime.AsSeconds());
+		}
+	}
+
+	if (ULiveLinkSourceSettings* Settings = LiveLinkClient->GetSourceSettings(Source))
+	{
+		Settings->BufferSettings.EngineTimeClockOffset = EngineClockOffset.GetEstimatedOffset();
+		Settings->BufferSettings.TimecodeClockOffset = TimecodeClockOffset.GetEstimatedOffset();
+	}
+
+	//Update frame world time offset based on our latest clock offset
+	NewFrameData.WorldTime.SetClockOffset(EngineClockOffset.GetEstimatedOffset());
+}
+
+
 #undef LOCTEXT_NAMESPACE
diff --git a/Engine/Plugins/Animation/LiveLink/Source/LiveLink/Private/LiveLinkTimedDataInput.h b/Engine/Plugins/Animation/LiveLink/Source/LiveLink/Private/LiveLinkTimedDataInput.h
--- a/Engine/Plugins/Animation/LiveLink/Source/LiveLink/Private/LiveLinkTimedDataInput.h
+++ b/Engine/Plugins/Animation/LiveLink/Source/LiveLink/Private/LiveLinkTimedDataInput.h
@@ -2,9 +2,12 @@
 
 #pragma once
 
-#include "CoreMinimal.h"
 #include "ITimedDataInput.h"
 
+#include "CoreMinimal.h"
+#include "ClockOffsetEstimatorRamp.h"
+
+
 class FLiveLinkClient;
 enum class ELiveLinkSourceMode : uint8;
 
@@ -34,9 +37,21 @@ public:
 #endif
 	//~ End ITimedDataInput API
 
+public:
+	//Tracks clock difference between each frame received and arrival time in engine referential.
+	void ProcessNewFrameTimingInfo(FLiveLinkBaseFrameData& NewFrameData);
+
 
 private:
 	FLiveLinkClient* LiveLinkClient;
 	TArray<ITimedDataInputChannel*> Channels;
 	FGuid Source;
+
+	//Continuous clock offset estimator for engine time and timecode
+	FClockOffsetEstimatorRamp EngineClockOffset;
+	FClockOffsetEstimatorRamp TimecodeClockOffset;
+
+	// We will receive each frame for each subject of this source. Stamp last source time/timecode to only update our offset estimation once per "source frame"
+	double LastWorldSourceTime = 0.0;
+	double LastSceneTime = 0.0;
 };
diff --git a/Engine/Plugins/Animation/LiveLink/Source/LiveLink/Public/LiveLinkSettings.h b/Engine/Plugins/Animation/LiveLink/Source/LiveLink/Public/LiveLinkSettings.h
--- a/Engine/Plugins/Animation/LiveLink/Source/LiveLink/Public/LiveLinkSettings.h
+++ b/Engine/Plugins/Animation/LiveLink/Source/LiveLink/Public/LiveLinkSettings.h
@@ -77,6 +77,10 @@ public:
 	UPROPERTY(config, EditAnywhere, Category="LiveLink", meta=(DisplayName="Preset Save Location"))
 	FDirectoryPath PresetSaveDir;
 
+	/** Continuous clock offset correction step */
+	UPROPERTY(config, EditAnywhere, AdvancedDisplay, Category = "LiveLink")
+	float ClockOffsetCorrectionStep;
+
 	/** The default evaluation mode a source connected via the message bus should start with. */
 	UPROPERTY(config, EditAnywhere, AdvancedDisplay, Category = "LiveLink")
 	ELiveLinkSourceMode DefaultMessageBusSourceMode;
diff --git a/Engine/Plugins/Animation/LiveLink/Source/LiveLinkEditor/Private/LiveLinkSourceSettingsDetailCustomization.cpp b/Engine/Plugins/Animation/LiveLink/Source/LiveLinkEditor/Private/LiveLinkSourceSettingsDetailCustomization.cpp
--- a/Engine/Plugins/Animation/LiveLink/Source/LiveLinkEditor/Private/LiveLinkSourceSettingsDetailCustomization.cpp
+++ b/Engine/Plugins/Animation/LiveLink/Source/LiveLinkEditor/Private/LiveLinkSourceSettingsDetailCustomization.cpp
@@ -48,6 +48,9 @@ void FLiveLinkSourceSettingsDetailCustomization::CustomizeDetails(IDetailLayoutB
 		CategoryBuilder.AddProperty(BufferSettingsPropertyHandle->GetChildHandle(GET_MEMBER_NAME_CHECKED(FLiveLinkSourceBufferManagementSettings, bKeepAtLeastOneFrame))
 			, EPropertyLocation::Advanced);
 
+		CategoryBuilder.AddProperty(BufferSettingsPropertyHandle->GetChildHandle(GET_MEMBER_NAME_CHECKED(FLiveLinkSourceBufferManagementSettings, TimecodeClockOffset))
+			, EPropertyLocation::Advanced);
+
 		IDetailCategoryBuilder& SubFrameCategoryBuilder = InDetailBuilder.EditCategory("Sub Frame");
 		SubFrameCategoryBuilder.AddProperty(BufferSettingsPropertyHandle->GetChildHandle(GET_MEMBER_NAME_CHECKED(FLiveLinkSourceBufferManagementSettings, bGenerateSubFrame)));
 		SubFrameCategoryBuilder.AddProperty(BufferSettingsPropertyHandle->GetChildHandle(GET_MEMBER_NAME_CHECKED(FLiveLinkSourceBufferManagementSettings, SourceTimecodeFrameRate)));
@@ -58,8 +61,11 @@ void FLiveLinkSourceSettingsDetailCustomization::CustomizeDetails(IDetailLayoutB
 			.DisplayName(LOCTEXT("ValidEngineTimeDisplayName", "Valid Buffer"));
 		CategoryBuilder.AddProperty(BufferSettingsPropertyHandle->GetChildHandle(GET_MEMBER_NAME_CHECKED(FLiveLinkSourceBufferManagementSettings, EngineTimeOffset)))
 			.DisplayName(LOCTEXT("EngineTimeOffsetDisplayName", "Offset"));
+		
 		CategoryBuilder.AddProperty(BufferSettingsPropertyHandle->GetChildHandle(GET_MEMBER_NAME_CHECKED(FLiveLinkSourceBufferManagementSettings, bKeepAtLeastOneFrame))
 			, EPropertyLocation::Advanced);
+		CategoryBuilder.AddProperty(BufferSettingsPropertyHandle->GetChildHandle(GET_MEMBER_NAME_CHECKED(FLiveLinkSourceBufferManagementSettings, EngineTimeClockOffset))
+			, EPropertyLocation::Advanced);
 	}
 	else if (SourceMode == ELiveLinkSourceMode::Latest)
 	{
diff --git a/Engine/Source/Runtime/LiveLinkInterface/Public/LiveLinkSourceSettings.h b/Engine/Source/Runtime/LiveLinkInterface/Public/LiveLinkSourceSettings.h
--- a/Engine/Source/Runtime/LiveLinkInterface/Public/LiveLinkSourceSettings.h
+++ b/Engine/Source/Runtime/LiveLinkInterface/Public/LiveLinkSourceSettings.h
@@ -50,6 +50,10 @@ struct FLiveLinkSourceBufferManagementSettings
 	/** When evaluating with time: how far back from current time should we read the buffer (in seconds) */
 	UPROPERTY(EditAnywhere, Category = "Settings", meta=(ForceUnits=s))
 	float EngineTimeOffset = 0.0f;
+	
+	/** Continuously updated clock offset estimator between source clock and engine clock (in seconds) */
+	UPROPERTY(VisibleAnywhere, Category = "Settings", AdvancedDisplay, meta = (ForceUnits = s))
+	double EngineTimeClockOffset = 0.0;
 
 	/** When evaluating with timecode: what is the expected frame rate of the timecode */
 	UPROPERTY(EditAnywhere, Category = "Settings")
@@ -87,6 +91,10 @@ struct FLiveLinkSourceBufferManagementSettings
 	UPROPERTY(EditAnywhere, Category = "Settings")
 	float TimecodeFrameOffset = 0.f;
 
+	/** Continuously updated clock offset estimator between source timecode clock and engine timecode provider clock (in seconds) */
+	UPROPERTY(VisibleAnywhere, Category = "Settings", AdvancedDisplay, meta = (ForceUnits = s))
+	double TimecodeClockOffset = 0.0;
+
 	/** When evaluating with latest: how far back from latest frame should we read the buffer */
 	UPROPERTY(EditAnywhere, Category = "Settings")
 	int32 LatestOffset = 0;
