- Fixing CIS

#rnx

diff --git a/Engine/Plugins/Animation/LiveLink/Source/LiveLink/Private/LiveLinkTimedDataInput.cpp b/Engine/Plugins/Animation/LiveLink/Source/LiveLink/Private/LiveLinkTimedDataInput.cpp
--- a/Engine/Plugins/Animation/LiveLink/Source/LiveLink/Private/LiveLinkTimedDataInput.cpp
+++ b/Engine/Plugins/Animation/LiveLink/Source/LiveLink/Private/LiveLinkTimedDataInput.cpp
@@ -7,8 +7,11 @@
 #include "ITimeManagementModule.h"
 #include "LiveLinkClient.h"
 #include "LiveLinkLog.h"
+#include "LiveLinkSettings.h"
 #include "LiveLinkSourceSettings.h"
+#include "LiveLinkTypes.h"
 #include "TimedDataInputCollection.h"
+#include "UObject/UObjectGlobals.h"
 
 #if WITH_EDITOR
 #include "Styling/SlateStyle.h"
diff --git a/Engine/Plugins/Animation/LiveLink/Source/LiveLink/Private/LiveLinkTimedDataInput.h b/Engine/Plugins/Animation/LiveLink/Source/LiveLink/Private/LiveLinkTimedDataInput.h
--- a/Engine/Plugins/Animation/LiveLink/Source/LiveLink/Private/LiveLinkTimedDataInput.h
+++ b/Engine/Plugins/Animation/LiveLink/Source/LiveLink/Private/LiveLinkTimedDataInput.h
@@ -10,6 +10,8 @@
 
 class FLiveLinkClient;
 enum class ELiveLinkSourceMode : uint8;
+struct FLiveLinkBaseFrameData;
+
 
 class FLiveLinkTimedDataInput : public ITimedDataInput
 {
