Live Link / Rich Curve: Set all times and values at once and not auto tangent each key.

diff --git a/Engine/Plugins/Animation/LiveLink/Source/LiveLinkMovieScene/Private/MovieScene/MovieSceneLiveLinkPropertyHandler.cpp b/Engine/Plugins/Animation/LiveLink/Source/LiveLinkMovieScene/Private/MovieScene/MovieSceneLiveLinkPropertyHandler.cpp
--- a/Engine/Plugins/Animation/LiveLink/Source/LiveLinkMovieScene/Private/MovieScene/MovieSceneLiveLinkPropertyHandler.cpp
+++ b/Engine/Plugins/Animation/LiveLink/Source/LiveLinkMovieScene/Private/MovieScene/MovieSceneLiveLinkPropertyHandler.cpp
@@ -114,10 +114,20 @@ void FMovieSceneLiveLinkPropertyHandler<float>::Finalize(bool bInReduceKeys, con
 	for (int32 i = 0; i < Keys.Num(); ++i)
 	{
 		const TArray<FLiveLinkPropertyKey<float>>& ElementKeys = Keys[i];
+
+		TArray<FFrameNumber> Times;
+		Times.SetNum(ElementKeys.Num());
+		TArray<FMovieSceneFloatValue> Values;
+		Values.SetNum(ElementKeys.Num());
+		int32 j = 0;
 		for (const FLiveLinkPropertyKey<float>& Key : ElementKeys)
 		{
-			PropertyStorage->FloatChannel[i].AddCubicKey(Key.Time, Key.Value, RCTM_Auto);
+			Times[j] = Key.Time;
+			Values[j] = FMovieSceneFloatValue(Key.Value);
+			j++;
 		}
+
+		PropertyStorage->FloatChannel[i].Set(Times, Values);
 	}
 
 	if (bInReduceKeys)
@@ -512,12 +522,26 @@ void FMovieSceneLiveLinkPropertyHandler<FVector>::Finalize(bool bInReduceKeys, c
 	for (int32 i = 0; i < Keys.Num(); ++i)
 	{
 		const TArray<FLiveLinkPropertyKey<FVector>>& ElementKeys = Keys[i];
+
+		TArray<FFrameNumber> Times;
+		Times.SetNum(ElementKeys.Num());
+		TArray<FMovieSceneFloatValue> ValuesX, ValuesY, ValuesZ;
+		ValuesX.SetNum(ElementKeys.Num());
+		ValuesY.SetNum(ElementKeys.Num());
+		ValuesZ.SetNum(ElementKeys.Num());
+		int32 j = 0;
 		for (const FLiveLinkPropertyKey<FVector>& Key : ElementKeys)
 		{
-			PropertyStorage->FloatChannel[(i * 3) + 0].AddCubicKey(Key.Time, Key.Value.X, RCTM_Break);
-			PropertyStorage->FloatChannel[(i * 3) + 1].AddCubicKey(Key.Time, Key.Value.Y, RCTM_Break);
-			PropertyStorage->FloatChannel[(i * 3) + 2].AddCubicKey(Key.Time, Key.Value.Z, RCTM_Break);
+			Times[j] = Key.Time;
+			ValuesX[j] = FMovieSceneFloatValue(Key.Value.X);
+			ValuesY[j] = FMovieSceneFloatValue(Key.Value.Y);
+			ValuesZ[j] = FMovieSceneFloatValue(Key.Value.Z);
+			j++;
 		}
+
+		PropertyStorage->FloatChannel[(i * 3) + 0].Set(Times, ValuesX);
+		PropertyStorage->FloatChannel[(i * 3) + 1].Set(Times, ValuesY);
+		PropertyStorage->FloatChannel[(i * 3) + 2].Set(Times, ValuesZ);
 	}
 
 	if (bInReduceKeys)
@@ -527,6 +551,13 @@ void FMovieSceneLiveLinkPropertyHandler<FVector>::Finalize(bool bInReduceKeys, c
 			MovieScene::Optimize(&Channel, InOptimizationParams);
 		}
 	}
+	else
+	{
+		for (FMovieSceneFloatChannel& Channel : PropertyStorage->FloatChannel)
+		{
+			Channel.AutoSetTangents();
+		}
+	}
 }
 
 template<>
diff --git a/Engine/Source/Runtime/Engine/Classes/Curves/KeyHandle.h b/Engine/Source/Runtime/Engine/Classes/Curves/KeyHandle.h
--- a/Engine/Source/Runtime/Engine/Classes/Curves/KeyHandle.h
+++ b/Engine/Source/Runtime/Engine/Classes/Curves/KeyHandle.h
@@ -80,6 +80,8 @@ public:
 	TArray<FKeyHandle>::TConstIterator CreateConstIterator() const { return KeyHandles.CreateConstIterator(); }
 	const TMap<FKeyHandle, int32>& GetMap() const { return KeyHandlesToIndices; }
 
+	void SetKeyHandles(int32 Num);
+
 	/** ICPPStructOps implementation */
 	bool Serialize(FArchive& Ar);
 	bool operator==(const FKeyHandleMap& Other) const { return KeyHandles == Other.KeyHandles; }
diff --git a/Engine/Source/Runtime/Engine/Private/Curves/KeyHandle.cpp b/Engine/Source/Runtime/Engine/Private/Curves/KeyHandle.cpp
--- a/Engine/Source/Runtime/Engine/Private/Curves/KeyHandle.cpp
+++ b/Engine/Source/Runtime/Engine/Private/Curves/KeyHandle.cpp
@@ -50,6 +50,17 @@ void FKeyHandleMap::Add( const FKeyHandle& InHandle, int32 InIndex )
 }
 
 
+void FKeyHandleMap::SetKeyHandles(int32 Num)
+{
+	KeyHandles.SetNum(Num);
+	KeyHandlesToIndices.Reserve(Num);
+	for (int32 Index = 0; Index < Num; ++Index)
+	{
+		KeyHandlesToIndices.Add(FKeyHandle(), Index);
+	}
+}
+
+
 void FKeyHandleMap::Empty()
 {
 	KeyHandlesToIndices.Empty();
diff --git a/Engine/Source/Runtime/Engine/Private/Curves/RichCurve.cpp b/Engine/Source/Runtime/Engine/Private/Curves/RichCurve.cpp
--- a/Engine/Source/Runtime/Engine/Private/Curves/RichCurve.cpp
+++ b/Engine/Source/Runtime/Engine/Private/Curves/RichCurve.cpp
@@ -230,10 +230,12 @@ void FRichCurve::SetKeys(const TArray<FRichCurveKey>& InKeys)
 {
 	Reset();
 
+	Keys.SetNum(InKeys.Num());
+	KeyHandlesToIndices.SetKeyHandles(InKeys.Num());
+
 	for (int32 Index = 0; Index < InKeys.Num(); ++Index)
 	{
-		Keys.Add(InKeys[Index]);
-		KeyHandlesToIndices.Add(FKeyHandle(), Index);
+		Keys[Index] = InKeys[Index];
 	}
 
 	AutoSetTangents();
