#LiveLink IWYU fix

diff --git a/Engine/Plugins/Animation/LiveLink/Source/LiveLinkMovieScene/Private/MovieScene/MovieSceneLiveLinkSectionTemplate.cpp b/Engine/Plugins/Animation/LiveLink/Source/LiveLinkMovieScene/Private/MovieScene/MovieSceneLiveLinkSectionTemplate.cpp
--- a/Engine/Plugins/Animation/LiveLink/Source/LiveLinkMovieScene/Private/MovieScene/MovieSceneLiveLinkSectionTemplate.cpp
+++ b/Engine/Plugins/Animation/LiveLink/Source/LiveLinkMovieScene/Private/MovieScene/MovieSceneLiveLinkSectionTemplate.cpp
@@ -13,6 +13,7 @@
 #include "MovieScene/MovieSceneLiveLinkSubSection.h"
 #include "MovieScene/MovieSceneLiveLinkTrack.h"
 #include "Roles/LiveLinkAnimationTypes.h"
+#include "Misc/App.h"
 
 static TAutoConsoleVariable<int32> CVarSequencerAlwaysSendInterpolatedLiveLink(
 	TEXT("Sequencer.AlwaysSendInterpolatedLiveLink"),
