LiveLink Panel UI Cleanup

diff --git a/Engine/Plugins/Animation/LiveLink/Source/LiveLinkEditor/LiveLinkEditor.Build.cs b/Engine/Plugins/Animation/LiveLink/Source/LiveLinkEditor/LiveLinkEditor.Build.cs
--- a/Engine/Plugins/Animation/LiveLink/Source/LiveLinkEditor/LiveLinkEditor.Build.cs
+++ b/Engine/Plugins/Animation/LiveLink/Source/LiveLinkEditor/LiveLinkEditor.Build.cs
@@ -26,6 +26,7 @@ namespace UnrealBuildTool.Rules
 					"CoreUObject",
 					"DetailCustomizations",
 					"EditorStyle",
+					"EditorWidgets",
 					"Engine",
 					"InputCore",
 					"KismetCompiler",
diff --git a/Engine/Plugins/Animation/LiveLink/Source/LiveLinkEditor/Private/LiveLinkClientPanel.cpp b/Engine/Plugins/Animation/LiveLink/Source/LiveLinkEditor/Private/LiveLinkClientPanel.cpp
--- a/Engine/Plugins/Animation/LiveLink/Source/LiveLinkEditor/Private/LiveLinkClientPanel.cpp
+++ b/Engine/Plugins/Animation/LiveLink/Source/LiveLinkEditor/Private/LiveLinkClientPanel.cpp
@@ -15,6 +15,7 @@
 #include "Widgets/Input/SCheckBox.h"
 #include "Widgets/Input/STextEntryPopup.h"
 #include "Widgets/Views/SListView.h"
+#include "SWarningOrErrorBox.h"
 
 #include "ILiveLinkSource.h"
 #include "LiveLinkClientCommands.h"
@@ -467,17 +468,17 @@ void SLiveLinkClientPanel::Construct(const FArguments& Args, FLiveLinkClient* In
 	PerformanceThrottlingProperty->GetDisplayNameText();
 	FFormatNamedArguments Arguments;
 	Arguments.Add(TEXT("PropertyName"), PerformanceThrottlingProperty->GetDisplayNameText());
-	FText PerformanceWarningText = FText::Format(LOCTEXT("LiveLinkPerformanceWarningMessage", "Warning: The editor setting '{PropertyName}' is currently enabled\nThis will stop editor windows from updating in realtime while the editor is not in focus"), Arguments);
+	FText PerformanceWarningText = FText::Format(LOCTEXT("LiveLinkPerformanceWarningMessage", "Warning: The editor setting '{PropertyName}' is currently enabled.  This will stop editor windows from updating in realtime while the editor is not in focus."), Arguments);
 	
 	ChildSlot
 	[
 		SNew(SBorder)
 		.BorderImage(FEditorStyle::GetBrush("MessageLog.ListBorder")) // set panel background color to same color as message log at the bottom
+		.Padding(FMargin(4.f, 4.f, 4.f, 4.f))
 		[
 			SNew(SVerticalBox)
 			+SVerticalBox::Slot()
 			.AutoHeight()
-			.Padding(FMargin(0.0f, 4.0f, 0.0f, 0.0f))
 			[
 				SNew(SLiveLinkClientPanelToolbar, Client)
 			]
@@ -588,33 +589,17 @@ void SLiveLinkClientPanel::Construct(const FArguments& Args, FLiveLinkClient* In
 			]
 			+SVerticalBox::Slot()
 			.AutoHeight()
+			.Padding(FMargin(0.0f, 8.0f, 0.0f, 4.0f))
 			[
-				SNew(SBorder)
-				.BorderImage(FEditorStyle::GetBrush("SettingsEditor.CheckoutWarningBorder"))
-				.BorderBackgroundColor(FColor(166, 137, 0))
+				SNew(SWarningOrErrorBox)
 				.Visibility(this, &SLiveLinkClientPanel::ShowEditorPerformanceThrottlingWarning)
+				.MessageStyle(EMessageStyle::Warning)
+				.Message(PerformanceWarningText)
 				[
-					SNew(SHorizontalBox)
-					+SHorizontalBox::Slot()
-					.FillWidth(1.0f)
-					.Padding(FMargin(WarningPadding, WarningPadding, WarningPadding, WarningPadding))
-					.VAlign(VAlign_Center)
-					[
-						SNew(STextBlock)
-						.Text(PerformanceWarningText)
-						.Font(FEditorStyle::GetFontStyle("PropertyWindow.NormalFont"))
-						.ShadowColorAndOpacity(FLinearColor::Black.CopyWithNewOpacity(0.3f))
-						.ShadowOffset(FVector2D::UnitVector)
-					]
-					+SHorizontalBox::Slot()
-					.Padding(FMargin(0.f, 0.f, WarningPadding, 0.f))
-					.AutoWidth()
-					.VAlign(VAlign_Center)
-					[
-						SNew(SButton)
-						.OnClicked(this, &SLiveLinkClientPanel::DisableEditorPerformanceThrottling)
-						.Text(LOCTEXT("LiveLinkPerformanceWarningDisable", "Disable"))
-					]
+					SNew(SButton)
+					.OnClicked(this, &SLiveLinkClientPanel::DisableEditorPerformanceThrottling)
+					.TextStyle(FAppStyle::Get(), "DialogButtonText")
+					.Text(LOCTEXT("LiveLinkPerformanceWarningDisable", "Disable"))
 				]
 			]
 		]
diff --git a/Engine/Plugins/Animation/LiveLink/Source/LiveLinkEditor/Private/LiveLinkClientPanelToolbar.cpp b/Engine/Plugins/Animation/LiveLink/Source/LiveLinkEditor/Private/LiveLinkClientPanelToolbar.cpp
--- a/Engine/Plugins/Animation/LiveLink/Source/LiveLinkEditor/Private/LiveLinkClientPanelToolbar.cpp
+++ b/Engine/Plugins/Animation/LiveLink/Source/LiveLinkEditor/Private/LiveLinkClientPanelToolbar.cpp
@@ -46,6 +46,7 @@
 #include "Widgets/Text/STextBlock.h"
 #include "Widgets/SBoxPanel.h"
 #include "Widgets/SNullWidget.h"
+#include "SEditorHeaderButton.h"
 
 
 #define LOCTEXT_NAMESPACE "LiveLinkClientPanel"
@@ -436,53 +437,11 @@ void SLiveLinkClientPanelToolbar::Construct(const FArguments& Args, FLiveLinkCli
 			.Padding(.0f)
 			.AutoWidth()
 			[
-				// The green button containing the "+ Add Source" items
-				SAssignNew(AddSourceButton, SComboButton)
-				.ToolTipText(LOCTEXT("AddSource_ToolTip", "Add a new LiveLink source"))
-				.ButtonStyle(FEditorStyle::Get(), "FlatButton.Success")
-				.ForegroundColor(FLinearColor::White)
-				.ContentPadding(FMargin(4, 0))
+				SNew(SEditorHeaderButton)
 				.OnGetMenuContent(this, &SLiveLinkClientPanelToolbar::OnGenerateSourceMenu)
-				.OnMenuOpenChanged(this, &SLiveLinkClientPanelToolbar::OnGeneratedSourceMenuOpenChanged)
-				.HasDownArrow(false)
-				.HAlign(HAlign_Center)
-				.VAlign(VAlign_Center)
-				.ButtonContent()
-				[
-					SNew(SHorizontalBox)
-					// The "+" sign.
-					+ SHorizontalBox::Slot()
-					.AutoWidth()
-					.VAlign(VAlign_Center)
-					.Padding(FMargin(0, 1))
-					[
-						SNew(STextBlock)
-						.TextStyle(FEditorStyle::Get(), "NormalText.Important")
-						.Font(FEditorStyle::Get().GetFontStyle("FontAwesome.10"))
-						.Text(FEditorFontGlyphs::Plus)
-					]
-					// The "Create Session" text.
-					+ SHorizontalBox::Slot()
-					.AutoWidth()
-					.VAlign(VAlign_Center)
-					.Padding(FMargin(4, 0, 0, 0))
-					[
-						SNew(STextBlock)
-						.TextStyle(FEditorStyle::Get(), "NormalText.Important")
-						.Text(LOCTEXT("AddSource", "Source"))
-					]
-					// The caret sign.
-					+ SHorizontalBox::Slot()
-					.VAlign(VAlign_Center)
-					.AutoWidth()
-					.Padding(4, 0, 0, 0)
-					[
-						SNew(STextBlock)
-						.TextStyle(FEditorStyle::Get(), "NormalText.Important")
-						.Font(FEditorStyle::Get().GetFontStyle("FontAwesome.10"))
-						.Text(FEditorFontGlyphs::Caret_Down)
-					]
-				]
+				.Icon(FAppStyle::Get().GetBrush("Icons.Plus"))
+				.Text(LOCTEXT("AddSource", "Source"))
+				.ToolTipText(LOCTEXT("AddSource_ToolTip", "Add a new LiveLink source"))
 			]
 
 			+ SHorizontalBox::Slot()
@@ -646,17 +605,6 @@ TSharedRef<SWidget> SLiveLinkClientPanelToolbar::OnGenerateSourceMenu()
 	return MenuBuilder.MakeWidget();
 }
 
-void SLiveLinkClientPanelToolbar::OnGeneratedSourceMenuOpenChanged(bool bOpen)
-{
-	if (!bOpen)
-	{
-		if (AddSourceButton.IsValid())
-		{
-			AddSourceButton->SetMenuContent(SNullWidget::NullWidget);
-		}
-	}
-}
-
 void SLiveLinkClientPanelToolbar::RetrieveFactorySourcePanel(FMenuBuilder& MenuBuilder, int32 FactoryIndex)
 {
 	if (Factories.IsValidIndex(FactoryIndex))
diff --git a/Engine/Plugins/Animation/LiveLink/Source/LiveLinkEditor/Private/LiveLinkClientPanelToolbar.h b/Engine/Plugins/Animation/LiveLink/Source/LiveLinkEditor/Private/LiveLinkClientPanelToolbar.h
--- a/Engine/Plugins/Animation/LiveLink/Source/LiveLinkEditor/Private/LiveLinkClientPanelToolbar.h
+++ b/Engine/Plugins/Animation/LiveLink/Source/LiveLinkEditor/Private/LiveLinkClientPanelToolbar.h
@@ -35,7 +35,6 @@ class SLiveLinkClientPanelToolbar : public SCompoundWidget, public FGCObject
 
 private:
 	TSharedRef<SWidget> OnGenerateSourceMenu();
-	void OnGeneratedSourceMenuOpenChanged(bool bOpen);
 	void RetrieveFactorySourcePanel(FMenuBuilder& MenuBuilder, int32 FactoryIndex);
 	void ExecuteCreateSource(int32 FactoryIndex);
 	void OnSourceCreated(TSharedPtr<ILiveLinkSource> NewSource, FString ConnectionString, TSubclassOf<ULiveLinkSourceFactory> Factory);
@@ -61,7 +60,6 @@ private:
 	TWeakPtr<IMenu> VirtualSubjectMenu;
 	TWeakPtr<STextEntryPopup> VirtualSubjectPopup;
 	TWeakPtr<SEditableTextBox> VirtualSubjectSourceName;
-	TSharedPtr<SComboButton> AddSourceButton;
 
 	TWeakObjectPtr<ULiveLinkPreset> LiveLinkPreset;
 	TArray<ULiveLinkSourceFactory*> Factories;
