LiveLink: Fixed axis remapping pre processor

Simplified UI to use standard Front, Right, and Up axes
Deprecated old fields and updated ReleaseObjectVersion.h

[Copy from CL 16211885 by [REDACTED_PERSON]]

diff --git a/Engine/Plugins/Animation/LiveLink/Source/LiveLink/Private/PreProcessor/LiveLinkAxisSwitchPreProcessor.cpp b/Engine/Plugins/Animation/LiveLink/Source/LiveLink/Private/PreProcessor/LiveLinkAxisSwitchPreProcessor.cpp
--- a/Engine/Plugins/Animation/LiveLink/Source/LiveLink/Private/PreProcessor/LiveLinkAxisSwitchPreProcessor.cpp
+++ b/Engine/Plugins/Animation/LiveLink/Source/LiveLink/Private/PreProcessor/LiveLinkAxisSwitchPreProcessor.cpp
@@ -5,17 +5,10 @@
 #include "Roles/LiveLinkAnimationTypes.h"
 #include "Roles/LiveLinkTransformRole.h"
 #include "Roles/LiveLinkTransformTypes.h"
+#include "UObject/ReleaseObjectVersion.h"
 
-namespace
+namespace LiveLinkAxisSwitchPreProcessor
 {
-	EAxis::Type LiveLinkAxisToMatrixAxis(ELiveLinkAxis InAxis)
-	{
-		const uint8 Value = static_cast<uint8>(InAxis) % 3;
-		const EAxis::Type Results[] = {EAxis::X, EAxis::Y, EAxis::Z};
-
-		return Results[Value];
-	}
-
 	float LiveLinkAxisToVectorMember(ELiveLinkAxis InAxis, const FVector& Origin)
 	{
 		const uint8 Value = static_cast<uint8>(InAxis) % 3;
@@ -27,26 +20,83 @@ namespace
 		return static_cast<uint8>(InAxis) < 3 ? 1.f : -1.f;
 	}
 
-	void SwitchTransform(FTransform& Transform
-		, ELiveLinkAxis OrientationAxisX, ELiveLinkAxis OrientationAxisY, ELiveLinkAxis OrientationAxisZ
-		, ELiveLinkAxis TranslationAxisX, ELiveLinkAxis TranslationAxisY, ELiveLinkAxis TranslationAxisZ)
+	ELiveLinkAxis AxisAbs(ELiveLinkAxis InAxis)
 	{
-		const FMatrix InMatrix = Transform.ToMatrixWithScale();
-
-		FVector DestAxisX = InMatrix.GetScaledAxis(LiveLinkAxisToMatrixAxis(OrientationAxisX)) * AxisSign(OrientationAxisX);
-		FVector DestAxisY = InMatrix.GetScaledAxis(LiveLinkAxisToMatrixAxis(OrientationAxisY)) * AxisSign(OrientationAxisY);
-		FVector DestAxisZ = InMatrix.GetScaledAxis(LiveLinkAxisToMatrixAxis(OrientationAxisZ)) * AxisSign(OrientationAxisZ);
-
-		FVector Origin = InMatrix.GetOrigin();
-		FVector NewOrigin;
-		NewOrigin.X = LiveLinkAxisToVectorMember(TranslationAxisX, Origin) * AxisSign(TranslationAxisX);
-		NewOrigin.Y = LiveLinkAxisToVectorMember(TranslationAxisY, Origin) * AxisSign(TranslationAxisY);
-		NewOrigin.Z = LiveLinkAxisToVectorMember(TranslationAxisZ, Origin) * AxisSign(TranslationAxisZ);
-
-		FMatrix Result(InMatrix);
-		Result.SetAxes(&DestAxisX, &DestAxisY, &DestAxisZ, &NewOrigin);
+		const uint8 Value = static_cast<uint8>(InAxis) % 3;
+		return static_cast<ELiveLinkAxis>(Value);
+	}
 
-		Transform.SetFromMatrix(Result);
+	void SwitchTransform(FTransform& Transform, ELiveLinkAxis FrontAxis, ELiveLinkAxis RightAxis, ELiveLinkAxis UpAxis)
+	{
+		FVector Location = Transform.GetLocation();
+		FVector NewLocation;
+		NewLocation[0] = LiveLinkAxisToVectorMember(FrontAxis, Location) * AxisSign(FrontAxis);
+		NewLocation[1] = LiveLinkAxisToVectorMember(RightAxis, Location) * AxisSign(RightAxis);
+		NewLocation[2] = LiveLinkAxisToVectorMember(UpAxis, Location) * AxisSign(UpAxis);
+
+		FVector Scale = Transform.GetScale3D();
+		FVector NewScale;
+		NewScale[0] = LiveLinkAxisToVectorMember(FrontAxis, Scale);
+		NewScale[1] = LiveLinkAxisToVectorMember(RightAxis, Scale);
+		NewScale[2] = LiveLinkAxisToVectorMember(UpAxis, Scale);
+
+		FQuat Orientation = Transform.GetRotation();
+		FQuat NewOrientation;
+		switch (FrontAxis)
+		{
+			case ELiveLinkAxis::X:		NewOrientation.X = Orientation.X;	break;
+			case ELiveLinkAxis::XNeg:	NewOrientation.X = -Orientation.X;	break;
+			case ELiveLinkAxis::Y:		NewOrientation.X = Orientation.Y;	break;
+			case ELiveLinkAxis::YNeg:	NewOrientation.X = -Orientation.Y;	break;
+			case ELiveLinkAxis::Z:		NewOrientation.X = Orientation.Z;	break;
+			case ELiveLinkAxis::ZNeg:	NewOrientation.X = -Orientation.Z;	break;
+		}
+		switch (RightAxis)
+		{
+			case ELiveLinkAxis::X:		NewOrientation.Y = Orientation.X;	break;
+			case ELiveLinkAxis::XNeg:	NewOrientation.Y = -Orientation.X;	break;
+			case ELiveLinkAxis::Y:		NewOrientation.Y = Orientation.Y;	break;
+			case ELiveLinkAxis::YNeg:	NewOrientation.Y = -Orientation.Y;	break;
+			case ELiveLinkAxis::Z:		NewOrientation.Y = Orientation.Z;	break;
+			case ELiveLinkAxis::ZNeg:	NewOrientation.Y = -Orientation.Z;	break;
+		}
+		switch (UpAxis)
+		{
+			case ELiveLinkAxis::X:		NewOrientation.Z = Orientation.X;	break;
+			case ELiveLinkAxis::XNeg:	NewOrientation.Z = -Orientation.X;	break;
+			case ELiveLinkAxis::Y:		NewOrientation.Z = Orientation.Y;	break;
+			case ELiveLinkAxis::YNeg:	NewOrientation.Z = -Orientation.Y;	break;
+			case ELiveLinkAxis::Z:		NewOrientation.Z = Orientation.Z;	break;
+			case ELiveLinkAxis::ZNeg:	NewOrientation.Z = -Orientation.Z;	break;
+		}
+
+		// Calculate the Levi-Civita symbol for W scaling according to the following table
+		// +1 xyz yzx zxy
+		// -1 xzy yxz zyx
+		// 0 if any axes are duplicated (x == y, y == z, or z == x)
+		float LCSymbol = 0.0f;
+		ELiveLinkAxis FrontAxisAbs = AxisAbs(FrontAxis);
+		ELiveLinkAxis RightAxisAbs = AxisAbs(RightAxis);
+		ELiveLinkAxis UpAxisAbs = AxisAbs(UpAxis);
+		if (((FrontAxisAbs == ELiveLinkAxis::X) && (RightAxisAbs == ELiveLinkAxis::Y) && (UpAxisAbs == ELiveLinkAxis::Z)) ||
+			((FrontAxisAbs == ELiveLinkAxis::Y) && (RightAxisAbs == ELiveLinkAxis::Z) && (UpAxisAbs == ELiveLinkAxis::X)) ||
+			((FrontAxisAbs == ELiveLinkAxis::Z) && (RightAxisAbs == ELiveLinkAxis::X) && (UpAxisAbs == ELiveLinkAxis::Y)))
+		{
+			LCSymbol = 1.0f;
+		}
+		else if (((FrontAxisAbs == ELiveLinkAxis::X) && (RightAxisAbs == ELiveLinkAxis::Z) && (UpAxisAbs == ELiveLinkAxis::Y)) ||
+			((FrontAxisAbs == ELiveLinkAxis::Y) && (RightAxisAbs == ELiveLinkAxis::X) && (UpAxisAbs == ELiveLinkAxis::Z)) ||
+			((FrontAxisAbs == ELiveLinkAxis::Z) && (RightAxisAbs == ELiveLinkAxis::Y) && (UpAxisAbs == ELiveLinkAxis::X)))
+		{
+			LCSymbol = -1.0f;
+		}
+
+		// Make sure the handedness is corrected to match the remapping
+		const float FlipSign = AxisSign(FrontAxis) * AxisSign(RightAxis) * AxisSign(UpAxis);
+
+		NewOrientation.W = Orientation.W * LCSymbol * FlipSign;
+
+		Transform.SetComponents(NewOrientation.GetNormalized(), NewLocation, NewScale);
 	}
 }
 
@@ -61,7 +111,7 @@ TSubclassOf<ULiveLinkRole> ULiveLinkTransformAxisSwitchPreProcessor::FLiveLinkTr
 bool ULiveLinkTransformAxisSwitchPreProcessor::FLiveLinkTransformAxisSwitchPreProcessorWorker::PreProcessFrame(FLiveLinkFrameDataStruct& InOutFrame) const
 {
 	FLiveLinkTransformFrameData& TransformData = *InOutFrame.Cast<FLiveLinkTransformFrameData>();
-	SwitchTransform(TransformData.Transform, OrientationAxisX, OrientationAxisY, OrientationAxisZ, TranslationAxisX, TranslationAxisY, TranslationAxisZ);
+	LiveLinkAxisSwitchPreProcessor::SwitchTransform(TransformData.Transform, FrontAxis, RightAxis, UpAxis);
 	return true;
 }
 
@@ -78,12 +128,9 @@ ULiveLinkFramePreProcessor::FWorkerSharedPtr ULiveLinkTransformAxisSwitchPreProc
 	if (!Instance.IsValid())
 	{
 		Instance = MakeShared<FLiveLinkTransformAxisSwitchPreProcessorWorker, ESPMode::ThreadSafe>();
-		Instance->OrientationAxisX = OrientationAxisX;
-		Instance->OrientationAxisY = OrientationAxisY;
-		Instance->OrientationAxisZ = OrientationAxisZ;
-		Instance->TranslationAxisX = TranslationAxisX;
-		Instance->TranslationAxisY = TranslationAxisY;
-		Instance->TranslationAxisZ = TranslationAxisZ;
+		Instance->FrontAxis = FrontAxis;
+		Instance->RightAxis = RightAxis;
+		Instance->UpAxis = UpAxis;
 	}
 
 	return Instance;
@@ -92,12 +139,9 @@ ULiveLinkFramePreProcessor::FWorkerSharedPtr ULiveLinkTransformAxisSwitchPreProc
 #if WITH_EDITOR
 void ULiveLinkTransformAxisSwitchPreProcessor::PostEditChangeChainProperty(struct FPropertyChangedChainEvent& PropertyChangedEvent)
 {
-	if (PropertyChangedEvent.Property->GetFName() == GET_MEMBER_NAME_CHECKED(ThisClass, OrientationAxisX) ||
-		PropertyChangedEvent.Property->GetFName() == GET_MEMBER_NAME_CHECKED(ThisClass, OrientationAxisY) ||
-		PropertyChangedEvent.Property->GetFName() == GET_MEMBER_NAME_CHECKED(ThisClass, OrientationAxisZ) ||
-		PropertyChangedEvent.Property->GetFName() == GET_MEMBER_NAME_CHECKED(ThisClass, TranslationAxisX) ||
-		PropertyChangedEvent.Property->GetFName() == GET_MEMBER_NAME_CHECKED(ThisClass, TranslationAxisY) ||
-		PropertyChangedEvent.Property->GetFName() == GET_MEMBER_NAME_CHECKED(ThisClass, TranslationAxisZ))
+	if (PropertyChangedEvent.Property->GetFName() == GET_MEMBER_NAME_CHECKED(ThisClass, FrontAxis) ||
+		PropertyChangedEvent.Property->GetFName() == GET_MEMBER_NAME_CHECKED(ThisClass, RightAxis) ||
+		PropertyChangedEvent.Property->GetFName() == GET_MEMBER_NAME_CHECKED(ThisClass, UpAxis))
 	{
 		Instance.Reset();
 	}
@@ -121,7 +165,7 @@ bool ULiveLinkAnimationAxisSwitchPreProcessor::FLiveLinkAnimationAxisSwitchPrePr
 	
 	for (FTransform& Transform : AnimationData.Transforms)
 	{
-		SwitchTransform(Transform, OrientationAxisX, OrientationAxisY, OrientationAxisZ, TranslationAxisX, TranslationAxisY, TranslationAxisZ);
+		LiveLinkAxisSwitchPreProcessor::SwitchTransform(Transform, FrontAxis, RightAxis, UpAxis);
 	}
 	
 	return true;
@@ -140,13 +184,31 @@ ULiveLinkFramePreProcessor::FWorkerSharedPtr ULiveLinkAnimationAxisSwitchPreProc
 	if (!Instance.IsValid())
 	{
 		Instance = MakeShared<FLiveLinkAnimationAxisSwitchPreProcessorWorker, ESPMode::ThreadSafe>();
-		Instance->OrientationAxisX = OrientationAxisX;
-		Instance->OrientationAxisY = OrientationAxisY;
-		Instance->OrientationAxisZ = OrientationAxisZ;
-		Instance->TranslationAxisX = TranslationAxisX;
-		Instance->TranslationAxisY = TranslationAxisY;
-		Instance->TranslationAxisZ = TranslationAxisZ;
+		Instance->FrontAxis = FrontAxis;
+		Instance->RightAxis = RightAxis;
+		Instance->UpAxis = UpAxis;
 	}
 
 	return Instance;
 }
+
+void ULiveLinkTransformAxisSwitchPreProcessor::Serialize(FArchive& Ar)
+{
+	Super::Serialize(Ar);
+
+#if WITH_EDITOR
+	Ar.UsingCustomVersion(FReleaseObjectVersion::GUID);
+
+	if (Ar.IsLoading())
+	{
+		if (Ar.CustomVer(FReleaseObjectVersion::GUID) < FReleaseObjectVersion::AddedFrontRightUpAxesToLiveLinkPreProcessor)
+		{
+			PRAGMA_DISABLE_DEPRECATION_WARNINGS
+			FrontAxis = OrientationAxisX_DEPRECATED;
+			RightAxis = OrientationAxisY_DEPRECATED;
+			UpAxis = OrientationAxisZ_DEPRECATED;
+			PRAGMA_ENABLE_DEPRECATION_WARNINGS
+		}
+	}
+#endif
+}
diff --git a/Engine/Plugins/Animation/LiveLink/Source/LiveLink/Public/PreProcessor/LiveLinkAxisSwitchPreProcessor.h b/Engine/Plugins/Animation/LiveLink/Source/LiveLink/Public/PreProcessor/LiveLinkAxisSwitchPreProcessor.h
--- a/Engine/Plugins/Animation/LiveLink/Source/LiveLink/Public/PreProcessor/LiveLinkAxisSwitchPreProcessor.h
+++ b/Engine/Plugins/Animation/LiveLink/Source/LiveLink/Public/PreProcessor/LiveLinkAxisSwitchPreProcessor.h
@@ -6,6 +6,7 @@
 #include "LiveLinkAxisSwitchPreProcessor.generated.h"
 
 
+// DO NOT modify the order or values of these enums - the code relies on the order and values and will fail if it's changed
 UENUM()
 enum class ELiveLinkAxis : uint8
 {
@@ -31,30 +32,36 @@ public:
 	class FLiveLinkTransformAxisSwitchPreProcessorWorker : public ILiveLinkFramePreProcessorWorker
 	{
 	public:
-		ELiveLinkAxis OrientationAxisX = ELiveLinkAxis::X;
-		ELiveLinkAxis OrientationAxisY = ELiveLinkAxis::Y;
-		ELiveLinkAxis OrientationAxisZ = ELiveLinkAxis::Z;
-		ELiveLinkAxis TranslationAxisX = ELiveLinkAxis::X;
-		ELiveLinkAxis TranslationAxisY = ELiveLinkAxis::Y;
-		ELiveLinkAxis TranslationAxisZ = ELiveLinkAxis::Z;
+		ELiveLinkAxis FrontAxis = ELiveLinkAxis::X;
+		ELiveLinkAxis RightAxis = ELiveLinkAxis::Y;
+		ELiveLinkAxis UpAxis = ELiveLinkAxis::Z;
 
 		virtual TSubclassOf<ULiveLinkRole> GetRole() const override;
 		virtual bool PreProcessFrame(FLiveLinkFrameDataStruct& InOutFrame) const override;
 	};
 
 protected:
+#if WITH_EDITORONLY_DATA
+	UPROPERTY(meta = (DeprecatedProperty, DeprecationMessage = "Use FrontAxis, RightAxis, UpAxis instead"))
+	ELiveLinkAxis OrientationAxisX_DEPRECATED = ELiveLinkAxis::X;
+	UPROPERTY(meta = (DeprecatedProperty, DeprecationMessage = "Use FrontAxis, RightAxis, UpAxis instead"))
+	ELiveLinkAxis OrientationAxisY_DEPRECATED = ELiveLinkAxis::Y;
+	UPROPERTY(meta = (DeprecatedProperty, DeprecationMessage = "Use FrontAxis, RightAxis, UpAxis instead"))
+	ELiveLinkAxis OrientationAxisZ_DEPRECATED = ELiveLinkAxis::Z;
+	UPROPERTY(meta = (DeprecatedProperty, DeprecationMessage = "Use FrontAxis, RightAxis, UpAxis instead"))
+	ELiveLinkAxis TranslationAxisX_DEPRECATED = ELiveLinkAxis::X;
+	UPROPERTY(meta = (DeprecatedProperty, DeprecationMessage = "Use FrontAxis, RightAxis, UpAxis instead"))
+	ELiveLinkAxis TranslationAxisY_DEPRECATED = ELiveLinkAxis::Y;
+	UPROPERTY(meta = (DeprecatedProperty, DeprecationMessage = "Use FrontAxis, RightAxis, UpAxis instead"))
+	ELiveLinkAxis TranslationAxisZ_DEPRECATED = ELiveLinkAxis::Z;
+#endif	// WITH_EDITORONLY_DATA
+
 	UPROPERTY(EditAnywhere, Category = "LiveLink")
-	ELiveLinkAxis OrientationAxisX = ELiveLinkAxis::X;
-	UPROPERTY(EditAnywhere, Category = "LiveLink")
-	ELiveLinkAxis OrientationAxisY = ELiveLinkAxis::Y;
-	UPROPERTY(EditAnywhere, Category = "LiveLink")
-	ELiveLinkAxis OrientationAxisZ = ELiveLinkAxis::Z;
-	UPROPERTY(EditAnywhere, Category = "LiveLink")
-	ELiveLinkAxis TranslationAxisX = ELiveLinkAxis::X;
+	ELiveLinkAxis FrontAxis = ELiveLinkAxis::X;
 	UPROPERTY(EditAnywhere, Category = "LiveLink")
-	ELiveLinkAxis TranslationAxisY = ELiveLinkAxis::Y;
+	ELiveLinkAxis RightAxis = ELiveLinkAxis::Y;
 	UPROPERTY(EditAnywhere, Category = "LiveLink")
-	ELiveLinkAxis TranslationAxisZ = ELiveLinkAxis::Z;
+	ELiveLinkAxis UpAxis = ELiveLinkAxis::Z;
 
 public:
 	virtual TSubclassOf<ULiveLinkRole> GetRole() const override;
@@ -62,6 +69,7 @@ public:
 
 public:
 	//~ UObject interface
+	virtual void Serialize(FArchive& Ar) override;
 #if WITH_EDITOR
 	virtual void PostEditChangeChainProperty(struct FPropertyChangedChainEvent& PropertyChangedEvent) override;
 #endif //WITH_EDITOR
