Callsite upgrades needed to permit a refresh of Wrapped Object Pointer usage in Engine + ShooterGame

diff --git a/Engine/Plugins/Animation/ControlRig/Source/ControlRig/Private/ControlRig.cpp b/Engine/Plugins/Animation/ControlRig/Source/ControlRig/Private/ControlRig.cpp
--- a/Engine/Plugins/Animation/ControlRig/Source/ControlRig/Private/ControlRig.cpp
+++ b/Engine/Plugins/Animation/ControlRig/Source/ControlRig/Private/ControlRig.cpp
@@ -2036,9 +2036,9 @@ void UControlRig::SetInteractionRig(UControlRig* InInteractionRig)
 		InteractionRig->OnInitialized_AnyThread().AddUObject(this, &UControlRig::HandleInteractionRigInitialized);
 		InteractionRig->OnExecuted_AnyThread().AddUObject(this, &UControlRig::HandleInteractionRigExecuted);
 		InteractionRig->ControlSelected().AddUObject(this, &UControlRig::HandleInteractionRigControlSelected, false);
-		OnInitialized_AnyThread().AddUObject(InteractionRig, &UControlRig::HandleInteractionRigInitialized);
-		OnExecuted_AnyThread().AddUObject(InteractionRig, &UControlRig::HandleInteractionRigExecuted);
-		ControlSelected().AddUObject(InteractionRig, &UControlRig::HandleInteractionRigControlSelected, true);
+		OnInitialized_AnyThread().AddUObject(ToRawPtr(InteractionRig), &UControlRig::HandleInteractionRigInitialized);
+		OnExecuted_AnyThread().AddUObject(ToRawPtr(InteractionRig), &UControlRig::HandleInteractionRigExecuted);
+		ControlSelected().AddUObject(ToRawPtr(InteractionRig), &UControlRig::HandleInteractionRigControlSelected, true);
 
 		FControlRigBracketScope BracketScope(InterRigSyncBracket);
 		InteractionRig->HandleInteractionRigExecuted(this, EControlRigState::Update, FRigUnit_BeginExecution::EventName);
@@ -2152,7 +2152,7 @@ void UControlRig::RemoveUserDataOfClass(TSubclassOf<UAssetUserData> InUserDataCl
 
 const TArray<UAssetUserData*>* UControlRig::GetAssetUserDataArray() const
 {
-	return &AssetUserData;
+	return &ToRawPtrTArrayUnsafe(AssetUserData);
 }
 
 void UControlRig::CopyPoseFromOtherRig(UControlRig* Subject)
diff --git a/Engine/Plugins/Animation/ControlRig/Source/ControlRigDeveloper/Private/ControlRigBlueprint.cpp b/Engine/Plugins/Animation/ControlRig/Source/ControlRigDeveloper/Private/ControlRigBlueprint.cpp
--- a/Engine/Plugins/Animation/ControlRig/Source/ControlRigDeveloper/Private/ControlRigBlueprint.cpp
+++ b/Engine/Plugins/Animation/ControlRig/Source/ControlRigDeveloper/Private/ControlRigBlueprint.cpp
@@ -1070,7 +1070,7 @@ URigVMController* UControlRigBlueprint::GetController(URigVMGraph* InGraph) cons
 		InGraph = Model;
 	}
 
-	URigVMController* const* ControllerPtr = Controllers.Find(InGraph);
+	UE_TRANSITIONAL_OBJECT_PTR(URigVMController) const* ControllerPtr = Controllers.Find(InGraph);
 	if (ControllerPtr)
 	{
 		return *ControllerPtr;
diff --git a/Engine/Plugins/Animation/ControlRig/Source/ControlRigDeveloper/Private/Graph/ControlRigGraphNode.cpp b/Engine/Plugins/Animation/ControlRig/Source/ControlRigDeveloper/Private/Graph/ControlRigGraphNode.cpp
--- a/Engine/Plugins/Animation/ControlRig/Source/ControlRigDeveloper/Private/Graph/ControlRigGraphNode.cpp
+++ b/Engine/Plugins/Animation/ControlRig/Source/ControlRigDeveloper/Private/Graph/ControlRigGraphNode.cpp
@@ -888,7 +888,7 @@ FName UControlRigGraphNode::GetModelNodeName() const
 
 URigVMPin* UControlRigGraphNode::GetModelPinFromPinPath(const FString& InPinPath) const
 {
-	if (URigVMPin*const* CachedModelPinPtr = CachedModelPins.Find(InPinPath))
+	if (UE_TRANSITIONAL_OBJECT_PTR(URigVMPin)const* CachedModelPinPtr = CachedModelPins.Find(InPinPath))
 	{
 		URigVMPin* CachedModelPin = *CachedModelPinPtr;
 		if (!CachedModelPin->HasAnyFlags(RF_Transient) && CachedModelPin->GetNode())
diff --git a/Engine/Plugins/Animation/ControlRig/Source/ControlRigEditor/Private/EditMode/ControlRigControlsProxy.cpp b/Engine/Plugins/Animation/ControlRig/Source/ControlRigEditor/Private/EditMode/ControlRigControlsProxy.cpp
--- a/Engine/Plugins/Animation/ControlRig/Source/ControlRigEditor/Private/EditMode/ControlRigControlsProxy.cpp
+++ b/Engine/Plugins/Animation/ControlRig/Source/ControlRigEditor/Private/EditMode/ControlRigControlsProxy.cpp
@@ -655,7 +655,7 @@ void UControlRigBoolControlProxy::SetKey(const IPropertyHandle& KeyedPropertyHan
 
 UControlRigControlsProxy* UControlRigDetailPanelControlProxies::FindProxy(const FName& Name) const
 {
-	UControlRigControlsProxy* const* Proxy = AllProxies.Find(Name);
+	UE_TRANSITIONAL_OBJECT_PTR(UControlRigControlsProxy) const* Proxy = AllProxies.Find(Name);
 	if (Proxy &&  Proxy[0])
 	{
 		return Proxy[0];
diff --git a/Engine/Plugins/Animation/LiveLink/Source/LiveLinkComponents/Private/LiveLinkComponentController.cpp b/Engine/Plugins/Animation/LiveLink/Source/LiveLinkComponents/Private/LiveLinkComponentController.cpp
--- a/Engine/Plugins/Animation/LiveLink/Source/LiveLinkComponents/Private/LiveLinkComponentController.cpp
+++ b/Engine/Plugins/Animation/LiveLink/Source/LiveLinkComponents/Private/LiveLinkComponentController.cpp
@@ -98,7 +98,7 @@ void ULiveLinkComponentController::SetControllerClassForRole(TSubclassOf<ULiveLi
 {
 	if (ControllerMap.Contains(RoleClass))
 	{
-		ULiveLinkControllerBase*& CurrentController = ControllerMap.FindOrAdd(RoleClass);
+		UE_TRANSITIONAL_OBJECT_PTR(ULiveLinkControllerBase)& CurrentController = ControllerMap.FindOrAdd(RoleClass);
 		if (CurrentController == nullptr || CurrentController->GetClass() != DesiredControllerClass)
 		{
 			//Controller is about to change, cleanup current one before 
@@ -187,7 +187,7 @@ void ULiveLinkComponentController::TickComponent(float DeltaTime, ELevelTick Tic
 	const bool bHasValidData = bEvaluateLiveLink ? LiveLinkClient.EvaluateFrame_AnyThread(SubjectRepresentation.Subject, SubjectRepresentation.Role, SubjectData) : false;
 
 	//Go through each controllers and initialize them if we're dirty and tick them if there's valid data to process
-	for (TTuple<TSubclassOf<ULiveLinkRole>, ULiveLinkControllerBase*>& ControllerEntry : ControllerMap)
+	for (auto& ControllerEntry : ControllerMap)
 	{
 		ULiveLinkControllerBase* Controller = ControllerEntry.Value;
 		if (Controller)
@@ -291,7 +291,7 @@ bool ULiveLinkComponentController::IsControllerMapOutdated() const
 	//Check if all map matches class hierarchy
 	for (const TSubclassOf<ULiveLinkRole>& RoleClass : SelectedRoleHierarchy)
 	{
-		const ULiveLinkControllerBase* const* FoundController = ControllerMap.Find(RoleClass);
+		UE_TRANSITIONAL_OBJECT_PTR(ULiveLinkControllerBase) const* FoundController = ControllerMap.Find(RoleClass);
 
 		//If ControllerMap doesn't have an entry for one of the role class hierarchy, we need to update
 		if (FoundController == nullptr)
@@ -347,7 +347,7 @@ TSubclassOf<ULiveLinkControllerBase> ULiveLinkComponentController::GetController
 void ULiveLinkComponentController::CleanupControllersInMap()
 {
 	//Cleanup the currently active controllers in the map
-	for (TPair<TSubclassOf<ULiveLinkRole>, ULiveLinkControllerBase*>& ControllerPair : ControllerMap)
+	for (auto& ControllerPair : ControllerMap)
 	{
 		if (ControllerPair.Value)
 		{
diff --git a/Engine/Plugins/Animation/LiveLink/Source/LiveLinkEditor/Private/LiveLinkComponentDetailCustomization.cpp b/Engine/Plugins/Animation/LiveLink/Source/LiveLinkEditor/Private/LiveLinkComponentDetailCustomization.cpp
--- a/Engine/Plugins/Animation/LiveLink/Source/LiveLinkEditor/Private/LiveLinkComponentDetailCustomization.cpp
+++ b/Engine/Plugins/Animation/LiveLink/Source/LiveLinkEditor/Private/LiveLinkComponentDetailCustomization.cpp
@@ -245,7 +245,7 @@ bool FLiveLinkComponentDetailCustomization::IsControllerItemSelected(FName Item,
 	if (ULiveLinkComponentController* EditedObjectPtr = EditedObject.Get())
 	{
 
-		ULiveLinkControllerBase** CurrentClass = EditedObjectPtr->ControllerMap.Find(RoleClass);
+		UE_TRANSITIONAL_OBJECT_PTR(ULiveLinkControllerBase)* CurrentClass = EditedObjectPtr->ControllerMap.Find(RoleClass);
 		if (CurrentClass == nullptr || *CurrentClass == nullptr)
 		{
 			return Item.IsNone();
@@ -266,7 +266,7 @@ FSlateColor FLiveLinkComponentDetailCustomization::HandleControllerStatusColorAn
 
 	if (ULiveLinkComponentController* EditorObjectPtr = EditedObject.Get())
 	{
-		ULiveLinkControllerBase** AssociatedControllerPtr = EditorObjectPtr->ControllerMap.Find(RoleClassEntry);
+		UE_TRANSITIONAL_OBJECT_PTR(ULiveLinkControllerBase)* AssociatedControllerPtr = EditorObjectPtr->ControllerMap.Find(RoleClassEntry);
 		if (AssociatedControllerPtr && *AssociatedControllerPtr)
 		{
 			if (UActorComponent* SelectedComponent = EditorObjectPtr->ComponentToControl.GetComponent(EditorObjectPtr->GetOwner()))
@@ -301,7 +301,7 @@ FText FLiveLinkComponentDetailCustomization::HandleControllerStatusText(TSubclas
 
 	if (ULiveLinkComponentController* EditorObjectPtr = EditedObject.Get())
 	{
-		ULiveLinkControllerBase** AssociatedControllerPtr = EditorObjectPtr->ControllerMap.Find(RoleClassEntry);
+		UE_TRANSITIONAL_OBJECT_PTR(ULiveLinkControllerBase)* AssociatedControllerPtr = EditorObjectPtr->ControllerMap.Find(RoleClassEntry);
 		if (AssociatedControllerPtr && *AssociatedControllerPtr)
 		{
 			if (UActorComponent* SelectedComponent = EditorObjectPtr->ComponentToControl.GetComponent(EditorObjectPtr->GetOwner()))
@@ -336,7 +336,7 @@ FText FLiveLinkComponentDetailCustomization::HandleControllerStatusToolTipText(T
 
 	if (ULiveLinkComponentController* EditorObjectPtr = EditedObject.Get())
 	{
-		ULiveLinkControllerBase** AssociatedControllerPtr = EditorObjectPtr->ControllerMap.Find(RoleClassEntry);
+		UE_TRANSITIONAL_OBJECT_PTR(ULiveLinkControllerBase)* AssociatedControllerPtr = EditorObjectPtr->ControllerMap.Find(RoleClassEntry);
 		if (AssociatedControllerPtr && *AssociatedControllerPtr)
 		{
 			if (UActorComponent* SelectedComponent = EditorObjectPtr->ComponentToControl.GetComponent(EditorObjectPtr->GetOwner()))
diff --git a/Engine/Plugins/Experimental/ChaosCloth/Source/Chaos/Private/ChaosCloth/ChaosClothingSimulationInteractor.cpp b/Engine/Plugins/Experimental/ChaosCloth/Source/Chaos/Private/ChaosCloth/ChaosClothingSimulationInteractor.cpp
--- a/Engine/Plugins/Experimental/ChaosCloth/Source/Chaos/Private/ChaosCloth/ChaosClothingSimulationInteractor.cpp
+++ b/Engine/Plugins/Experimental/ChaosCloth/Source/Chaos/Private/ChaosCloth/ChaosClothingSimulationInteractor.cpp
@@ -187,7 +187,7 @@ void UChaosClothingSimulationInteractor::ClothConfigUpdated()
 void UChaosClothingSimulationInteractor::SetAnimDriveSpringStiffness(float Stiffness)
 {
 	// Set the anim drive stiffness through the ChaosClothInteractor to allow the value to be overridden by the cloth interactor if needed
-	for (const TPair<FName, UClothingInteractor*>& ClothingInteractor : UClothingSimulationInteractor::ClothingInteractors)
+	for (const auto& ClothingInteractor : UClothingSimulationInteractor::ClothingInteractors)
 	{
 		if (UChaosClothingInteractor* const ChaosClothingInteractor = Cast<UChaosClothingInteractor>(ClothingInteractor.Value))
 		{
diff --git a/Engine/Plugins/VirtualProduction/Takes/Source/TakesCore/Public/TakesUtils.h b/Engine/Plugins/VirtualProduction/Takes/Source/TakesCore/Public/TakesUtils.h
--- a/Engine/Plugins/VirtualProduction/Takes/Source/TakesCore/Public/TakesUtils.h
+++ b/Engine/Plugins/VirtualProduction/Takes/Source/TakesCore/Public/TakesUtils.h
@@ -89,6 +89,18 @@ namespace TakesUtils
 		return true;
 	}
 
+	template<typename AssetType>
+	static bool CreateNewAssetPackage(FString& InPackageName, TObjectPtr<AssetType>& OutAsset, FText* OutError, AssetType* OptionalBase = nullptr)
+	{
+		AssetType* RawOutAsset = nullptr;
+		if (CreateNewAssetPackage(InPackageName, RawOutAsset, OutError, OptionalBase))
+		{
+			OutAsset = RawOutAsset;
+			return true;
+		}
+		return false;
+	}
+
 	/**
 	 * Utility function that creates an asset with the specified asset path and name.
 	 * If the asset cannot be created (as one already exists), we try to postfix the asset
diff --git a/Engine/Source/Developer/RigVMDeveloper/Private/RigVMModel/RigVMController.cpp b/Engine/Source/Developer/RigVMDeveloper/Private/RigVMModel/RigVMController.cpp
--- a/Engine/Source/Developer/RigVMDeveloper/Private/RigVMModel/RigVMController.cpp
+++ b/Engine/Source/Developer/RigVMDeveloper/Private/RigVMModel/RigVMController.cpp
@@ -4162,7 +4162,7 @@ bool URigVMController::RemoveNode(URigVMNode* InNode, bool bSetupUndoRedo, bool
 			}
 			FunctionLibrary->FunctionReferences.Remove(LibraryNode);
 
-			for(const TPair<FString, URigVMLibraryNode*>& Pair : FunctionLibrary->LocalizedFunctions)
+			for(const auto& Pair : FunctionLibrary->LocalizedFunctions)
 			{
 				if(Pair.Value == LibraryNode)
 				{
diff --git a/Engine/Source/Runtime/ClothingSystemRuntimeInterface/Private/ClothingSimulationInteractor.cpp b/Engine/Source/Runtime/ClothingSystemRuntimeInterface/Private/ClothingSimulationInteractor.cpp
--- a/Engine/Source/Runtime/ClothingSystemRuntimeInterface/Private/ClothingSimulationInteractor.cpp
+++ b/Engine/Source/Runtime/ClothingSystemRuntimeInterface/Private/ClothingSimulationInteractor.cpp
@@ -42,7 +42,7 @@ void UClothingSimulationInteractor::Sync(IClothingSimulation* Simulation, ICloth
 	LastNumSubsteps = Simulation->GetNumSubsteps();
 	LastSimulationTime = Simulation->GetSimulationTime();
 
-	for (const TPair<FName, UClothingInteractor*>& ClothingInteractor : UClothingSimulationInteractor::ClothingInteractors)
+	for (const auto& ClothingInteractor : UClothingSimulationInteractor::ClothingInteractors)
 	{
 		if (ClothingInteractor.Value)
 		{
