Fix remove all subjects missing from the virtual subject context menu

diff --git a/Engine/Plugins/Animation/LiveLink/Source/LiveLink/Private/LiveLinkClient.cpp b/Engine/Plugins/Animation/LiveLink/Source/LiveLink/Private/LiveLinkClient.cpp
--- a/Engine/Plugins/Animation/LiveLink/Source/LiveLink/Private/LiveLinkClient.cpp
+++ b/Engine/Plugins/Animation/LiveLink/Source/LiveLink/Private/LiveLinkClient.cpp
@@ -1336,7 +1336,7 @@ FSimpleMulticastDelegate& FLiveLinkClient::OnLiveLinkTicked()
 	return OnLiveLinkTickedDelegate;
 }
 
-TArray<FGuid> FLiveLinkClient::GetDisplayableSources() const
+TArray<FGuid> FLiveLinkClient::GetDisplayableSources(bool bIncludeVirtualSources) const
 {
 	TArray<FGuid> Results;
 
@@ -1345,7 +1345,7 @@ TArray<FGuid> FLiveLinkClient::GetDisplayableSources() const
 
 	for (const FLiveLinkCollectionSourceItem& Data : PresetSources)
 	{
-		if (Data.Source->CanBeDisplayedInUI())
+		if (Data.Source->CanBeDisplayedInUI() || (bIncludeVirtualSources && Data.IsVirtualSource()))
 		{
 			Results.Add(Data.Guid);
 		}
diff --git a/Engine/Plugins/Animation/LiveLink/Source/LiveLink/Public/LiveLinkClient.h b/Engine/Plugins/Animation/LiveLink/Source/LiveLink/Public/LiveLinkClient.h
--- a/Engine/Plugins/Animation/LiveLink/Source/LiveLink/Public/LiveLinkClient.h
+++ b/Engine/Plugins/Animation/LiveLink/Source/LiveLink/Public/LiveLinkClient.h
@@ -156,7 +156,12 @@ public:
 	/** Callback when property changed for one of the source settings */
 	void OnPropertyChanged(FGuid EntryGuid, const FPropertyChangedEvent& PropertyChangedEvent);
 
-	TArray<FGuid> GetDisplayableSources() const;
+	/**
+	 * Get all sources that can be displayed in the UI's source list.
+	 * @param bIncludeVirtualSources Whether or not to include virtual sources in the returned list, since virtual sources are not displayed in the source list.
+	 * @return the list of displayable sources.
+	 */
+	TArray<FGuid> GetDisplayableSources(bool bIncludeVirtualSources = false) const;
 
 	FLiveLinkSubjectTimeSyncData GetTimeSyncData(FLiveLinkSubjectName SubjectName);
 
diff --git a/Engine/Plugins/Animation/LiveLink/Source/LiveLinkEditor/Private/LiveLinkClientPanel.cpp b/Engine/Plugins/Animation/LiveLink/Source/LiveLinkEditor/Private/LiveLinkClientPanel.cpp
--- a/Engine/Plugins/Animation/LiveLink/Source/LiveLinkEditor/Private/LiveLinkClientPanel.cpp
+++ b/Engine/Plugins/Animation/LiveLink/Source/LiveLinkEditor/Private/LiveLinkClientPanel.cpp
@@ -880,7 +880,8 @@ bool SLiveLinkClientPanel::CanRemoveSource()
 
 bool SLiveLinkClientPanel::HasSource() const
 {
-	return Client->GetDisplayableSources().Num() > 0;
+	constexpr bool bIncludeVirtualSources = true;
+	return Client->GetDisplayableSources(bIncludeVirtualSources).Num() > 0;
 }
 
 void SLiveLinkClientPanel::HandleRemoveAllSources()
