Fix ApplyPreset assuming RemoveAllSources is done in one tick

diff --git a/Engine/Plugins/Animation/LiveLink/Source/LiveLink/LiveLink.Build.cs b/Engine/Plugins/Animation/LiveLink/Source/LiveLink/LiveLink.Build.cs
--- a/Engine/Plugins/Animation/LiveLink/Source/LiveLink/LiveLink.Build.cs
+++ b/Engine/Plugins/Animation/LiveLink/Source/LiveLink/LiveLink.Build.cs
@@ -34,6 +34,7 @@ namespace UnrealBuildTool.Rules
 				new string[]
 				{
 					"MessageLog",
+					"UnrealEd",
 				});
 			}
 		}
diff --git a/Engine/Plugins/Animation/LiveLink/Source/LiveLink/Private/LiveLinkClient.cpp b/Engine/Plugins/Animation/LiveLink/Source/LiveLink/Private/LiveLinkClient.cpp
--- a/Engine/Plugins/Animation/LiveLink/Source/LiveLink/Private/LiveLinkClient.cpp
+++ b/Engine/Plugins/Animation/LiveLink/Source/LiveLink/Private/LiveLinkClient.cpp
@@ -457,14 +457,14 @@ bool FLiveLinkClient::HasSourceBeenAdded(TSharedPtr<ILiveLinkSource> InSource) c
 	return false;
 }
 
-TArray<FGuid> FLiveLinkClient::GetSources() const
+TArray<FGuid> FLiveLinkClient::GetSources(bool bEvenIfPendingKill) const
 {
 	check(Collection);
 
 	TArray<FGuid> Result;
 	for (const FLiveLinkCollectionSourceItem& SourceItem : Collection->GetSources())
 	{
-		if (!SourceItem.bPendingKill && !SourceItem.IsVirtualSource())
+		if ((!SourceItem.bPendingKill || bEvenIfPendingKill) && !SourceItem.IsVirtualSource())
 		{
 			Result.Add(SourceItem.Guid);
 		}
@@ -472,14 +472,14 @@ TArray<FGuid> FLiveLinkClient::GetSources() const
 	return Result;
 }
 
-TArray<FGuid> FLiveLinkClient::GetVirtualSources() const
+TArray<FGuid> FLiveLinkClient::GetVirtualSources(bool bEvenIfPendingKill) const
 {
 	check(Collection);
 
 	TArray<FGuid> Result;
 	for (const FLiveLinkCollectionSourceItem& SourceItem : Collection->GetSources())
 	{
-		if (!SourceItem.bPendingKill && SourceItem.IsVirtualSource())
+		if ((!SourceItem.bPendingKill || bEvenIfPendingKill) && SourceItem.IsVirtualSource())
 		{
 			Result.Add(SourceItem.Guid);
 		}
diff --git a/Engine/Plugins/Animation/LiveLink/Source/LiveLink/Private/LiveLinkModule.cpp b/Engine/Plugins/Animation/LiveLink/Source/LiveLink/Private/LiveLinkModule.cpp
--- a/Engine/Plugins/Animation/LiveLink/Source/LiveLink/Private/LiveLinkModule.cpp
+++ b/Engine/Plugins/Animation/LiveLink/Source/LiveLink/Private/LiveLinkModule.cpp
@@ -64,7 +64,7 @@ void FLiveLinkModule::OnEngineLoopInitComplete()
 {
 	if (ULiveLinkPreset* Preset = GetDefault<ULiveLinkSettings>()->DefaultLiveLinkPreset.LoadSynchronous())
 	{
-		Preset->ApplyToClient();
+		Preset->ApplyToClientLatent();
 	}
 }
 
diff --git a/Engine/Plugins/Animation/LiveLink/Source/LiveLink/Private/LiveLinkPreset.cpp b/Engine/Plugins/Animation/LiveLink/Source/LiveLink/Private/LiveLinkPreset.cpp
--- a/Engine/Plugins/Animation/LiveLink/Source/LiveLink/Private/LiveLinkPreset.cpp
+++ b/Engine/Plugins/Animation/LiveLink/Source/LiveLink/Private/LiveLinkPreset.cpp
@@ -1,11 +1,20 @@
 // Copyright Epic Games, Inc. All Rights Reserved.
 
 #include "LiveLinkPreset.h"
+
+#include "Engine/Engine.h"
+#include "Engine/LatentActionManager.h"
 #include "Features/IModularFeatures.h"
 #include "HAL/IConsoleManager.h"
+#include "LatentActions.h"
 #include "LiveLinkClient.h"
 #include "LiveLinkLog.h"
-#include "UObject/ConstructorHelpers.h"
+#include "Misc/App.h"
+#include "TimerManager.h"
+
+#if WITH_EDITOR
+#include "Editor.h"
+#endif
 
 namespace
 {
@@ -17,13 +26,13 @@ namespace
 				for (const FString& Element : Args)
 				{
 					const TCHAR* PresetStr = TEXT("Preset=");
-					int32 FoundElement = Element.Find(PresetStr);
+					const int32 FoundElement = Element.Find(PresetStr);
 					if (FoundElement != INDEX_NONE)
 					{
 						UObject* Object = StaticLoadObject(ULiveLinkPreset::StaticClass(), nullptr, *Element + FoundElement + FCString::Strlen(PresetStr));
 						if (Object)
 						{
-							CastChecked<ULiveLinkPreset>(Object)->ApplyToClient();
+							CastChecked<ULiveLinkPreset>(Object)->ApplyToClientLatent();
 						}
 					}
 				}
@@ -38,7 +47,7 @@ namespace
 				for (const FString& Element : Args)
 				{
 					const TCHAR* PresetStr = TEXT("Preset=");
-					int32 FoundElement = Element.Find(PresetStr);
+					const int32 FoundElement = Element.Find(PresetStr);
 					if (FoundElement != INDEX_NONE)
 					{
 						UObject* Object = StaticLoadObject(ULiveLinkPreset::StaticClass(), nullptr, *Element + FoundElement + FCString::Strlen(PresetStr));
@@ -52,6 +61,149 @@ namespace
 		);
 }
 
+namespace LiveLinkPresetUtils
+{
+	UWorld* GetWorld()
+	{
+		UWorld* World = nullptr;
+
+#if WITH_EDITOR
+		if (GEditor && FApp::CanEverRender())
+		{
+			World = GEditor->GetEditorWorldContext(false).World();
+		}
+#endif
+
+		if (World)
+		{
+			return World;
+		}
+
+		if (GEngine)
+		{
+			for (const FWorldContext& WorldContext : GEngine->GetWorldContexts())
+			{
+				if (WorldContext.WorldType == EWorldType::Game)
+				{
+					return World;
+				}
+			}
+		}
+		
+		return nullptr;
+	}
+}
+
+struct FApplyToClientPollingOperation
+{
+	/** Holds a weak pointer to the preset that will be applied. */
+	TWeakObjectPtr<ULiveLinkPreset> WeakPreset;
+	/** Keeps track of remaining time before aborting  */
+	double RemainingTime = 1.0;
+	/** Keeps track of the last time the Update function was called. */
+	double LastTimeSinceUpdate = 0.0;
+	
+	/** Result of a call to Update() */
+	enum class EApplyToClientUpdateResult : uint8
+	{
+		Success, // Preset was successfully applied.
+		Failure, // Preset could not be applied.
+		Pending  // Operation is still ongoing.
+	};
+
+	FApplyToClientPollingOperation(ULiveLinkPreset* Preset)
+		: WeakPreset(Preset)
+	{
+	}
+	
+	EApplyToClientUpdateResult Update()
+	{
+		bool bResult = true;
+        
+        ULiveLinkPreset* Preset = WeakPreset.Get();
+        if (!Preset || !IModularFeatures::Get().IsModularFeatureAvailable(ILiveLinkClient::ModularFeatureName))
+        {
+        	FLiveLinkLog::Error(TEXT("Could not apply preset"));
+        	return EApplyToClientUpdateResult::Failure;
+        }
+        
+        FLiveLinkClient& LiveLinkClient = IModularFeatures::Get().GetModularFeature<FLiveLinkClient>(ILiveLinkClient::ModularFeatureName);
+
+        if (LastTimeSinceUpdate == 0.0)
+        {
+        	// Start the process of removing sources.
+        	LiveLinkClient.RemoveAllSources();
+        	LastTimeSinceUpdate = FApp::GetCurrentTime();
+        }
+
+        LiveLinkClient.Tick();
+
+        constexpr bool bEvenIfPendingKill = true;
+        if (LiveLinkClient.GetSources(bEvenIfPendingKill).Num() == 0)
+        {
+        	for (const FLiveLinkSourcePreset& SourcePreset : Preset->GetSourcePresets())
+        	{
+        		bResult &= LiveLinkClient.CreateSource(SourcePreset);
+        	}
+
+        	for (const FLiveLinkSubjectPreset& SubjectPreset : Preset->GetSubjectPresets())
+        	{
+        		bResult &= LiveLinkClient.CreateSubject(SubjectPreset);
+        	}
+        
+        	if (bResult)
+        	{
+        		FLiveLinkLog::Info(TEXT("Applied '%s'"), *Preset->GetFullName());
+        	}
+        	else
+        	{
+        		FLiveLinkLog::Error(TEXT("Could not apply '%s'"), *Preset->GetFullName());
+        	}
+
+        	return bResult ? EApplyToClientUpdateResult::Success : EApplyToClientUpdateResult::Failure;
+        }
+
+        RemainingTime -= FApp::GetCurrentTime() - LastTimeSinceUpdate;
+        if (RemainingTime <= 0.0)
+        {
+        	return EApplyToClientUpdateResult::Failure;
+        }
+        
+        LastTimeSinceUpdate = FApp::GetCurrentTime();
+		return EApplyToClientUpdateResult::Pending;
+	}
+};
+
+class FApplyToClientLatentAction : public FPendingLatentAction
+{
+public:
+	FName ExecutionFunction;
+	int32 OutputLink;
+	FWeakObjectPtr CallbackTarget;
+	FApplyToClientPollingOperation ApplyOperation;
+	
+	FApplyToClientLatentAction(const FLatentActionInfo& InLatentInfo, ULiveLinkPreset* InOwnerPreset)
+        : ExecutionFunction(InLatentInfo.ExecutionFunction)
+        , OutputLink(InLatentInfo.Linkage)
+        , CallbackTarget(InLatentInfo.CallbackTarget)
+		, ApplyOperation(InOwnerPreset)
+	{
+	}
+
+	virtual void UpdateOperation(FLatentResponse& Response) override
+	{
+		if (ApplyOperation.Update() != FApplyToClientPollingOperation::EApplyToClientUpdateResult::Pending)
+		{
+			Response.FinishAndTriggerIf(true, ExecutionFunction, OutputLink, CallbackTarget);
+		}
+	}
+};
+
+ULiveLinkPreset::~ULiveLinkPreset()
+{
+	ClearApplyToClientTimer();
+}
+
 bool ULiveLinkPreset::ApplyToClient() const
 {
 	bool bResult = false;
@@ -86,6 +238,57 @@ bool ULiveLinkPreset::ApplyToClient() const
 	return bResult;
 }
 
+void ULiveLinkPreset::ApplyToClientLatent(UObject* WorldContextObject, FLatentActionInfo LatentInfo)
+{
+	if (UWorld* World = GEngine->GetWorldFromContextObject(WorldContextObject, EGetWorldErrorMode::ReturnNull))
+	{
+		FLatentActionManager& LatentActionManager = World->GetLatentActionManager();
+		if (!LatentActionManager.FindExistingAction<FApplyToClientLatentAction>(LatentInfo.CallbackTarget, LatentInfo.UUID))
+		{
+			LatentActionManager.AddNewAction(LatentInfo.CallbackTarget, LatentInfo.UUID, new FApplyToClientLatentAction(LatentInfo, this));
+		}
+	}
+}
+
+void ULiveLinkPreset::ApplyToClientLatent(TFunction<void(bool)> CompletionCallback)
+{
+	if (ApplyToClientPollingOperation.IsValid())
+	{
+		FLiveLinkLog::Error(TEXT("Could not apply '%s', operation is already in progress."), *GetFullName());
+		if (CompletionCallback != nullptr)
+		{
+			CompletionCallback(false);
+		}
+		return;
+	}
+
+	ApplyToClientPollingOperation = MakePimpl<FApplyToClientPollingOperation>(this);
+
+	if (UWorld* World = LiveLinkPresetUtils::GetWorld())
+	{
+		FTimerDelegate TimerDelegate;
+		TimerDelegate.BindLambda([WeakThis = TWeakObjectPtr<ULiveLinkPreset>(this), WrappedCallback = MoveTemp(CompletionCallback)]()
+		{
+			if (ensure(WeakThis->ApplyToClientPollingOperation))
+			{
+				FApplyToClientPollingOperation::EApplyToClientUpdateResult Result = WeakThis->ApplyToClientPollingOperation->Update();
+				if (Result != FApplyToClientPollingOperation::EApplyToClientUpdateResult::Pending)
+				{
+					if (WrappedCallback != nullptr)
+					{
+						WrappedCallback(Result == FApplyToClientPollingOperation::EApplyToClientUpdateResult::Success ? true : false);
+					}
+
+					WeakThis->ApplyToClientPollingOperation.Reset();
+					WeakThis->ClearApplyToClientTimer();
+				}
+			}
+		});
+		
+		World->GetTimerManager().SetTimer(ApplyToClientTimerHandle, MoveTemp(TimerDelegate), 0.01f, true);
+	}
+}
+
 bool ULiveLinkPreset::AddToClient(const bool bRecreatePresets) const
 {
 	bool bResult = false;
@@ -96,7 +299,7 @@ bool ULiveLinkPreset::AddToClient(const bool bRecreatePresets) const
 		TArray<FGuid> AllSources;
 		AllSources.Append(LiveLinkClient.GetSources());
 		AllSources.Append(LiveLinkClient.GetVirtualSources());
-		TArray<FLiveLinkSubjectKey> AllSubjects = LiveLinkClient.GetSubjects(true, true);
+		const TArray<FLiveLinkSubjectKey> AllSubjects = LiveLinkClient.GetSubjects(true, true);
 
 		TSet<FGuid> FoundSources;
 		TSet<FLiveLinkSubjectKey> FoundSubjects;
@@ -166,12 +369,12 @@ void ULiveLinkPreset::BuildFromClient()
 	{
 		FLiveLinkClient& LiveLinkClient = IModularFeatures::Get().GetModularFeature<FLiveLinkClient>(ILiveLinkClient::ModularFeatureName);
 
-		for (FGuid SourceGuid : LiveLinkClient.GetSources())
+		for (const FGuid& SourceGuid : LiveLinkClient.GetSources())
 		{
 			Sources.Add(LiveLinkClient.GetSourcePreset(SourceGuid, this));
 		}
 
-		for (FGuid SourceGuid : LiveLinkClient.GetVirtualSources())
+		for (const FGuid& SourceGuid : LiveLinkClient.GetVirtualSources())
 		{
 			FLiveLinkSourcePreset NewPreset = LiveLinkClient.GetSourcePreset(SourceGuid, this);
 			if (NewPreset.Guid.IsValid())
@@ -186,3 +389,11 @@ void ULiveLinkPreset::BuildFromClient()
 		}
 	}
 }
+
+void ULiveLinkPreset::ClearApplyToClientTimer()
+{
+	if (UWorld* World = LiveLinkPresetUtils::GetWorld())
+	{
+		World->GetTimerManager().ClearTimer(ApplyToClientTimerHandle);
+	}
+}
diff --git a/Engine/Plugins/Animation/LiveLink/Source/LiveLink/Public/LiveLinkClient.h b/Engine/Plugins/Animation/LiveLink/Source/LiveLink/Public/LiveLinkClient.h
--- a/Engine/Plugins/Animation/LiveLink/Source/LiveLink/Public/LiveLinkClient.h
+++ b/Engine/Plugins/Animation/LiveLink/Source/LiveLink/Public/LiveLinkClient.h
@@ -80,8 +80,8 @@ public:
 	virtual void RemoveSource(TSharedPtr<ILiveLinkSource> Source) override;
 	virtual void RemoveSource(FGuid InEntryGuid) override;
 	virtual bool HasSourceBeenAdded(TSharedPtr<ILiveLinkSource> Source) const override;
-	virtual TArray<FGuid> GetSources() const override;
-	virtual TArray<FGuid> GetVirtualSources() const override;
+	virtual TArray<FGuid> GetSources(bool bEvenIfPendingKill = false) const override;
+	virtual TArray<FGuid> GetVirtualSources(bool bEvenIfPendingKill = false) const override;
 	virtual FLiveLinkSourcePreset GetSourcePreset(FGuid SourceGuid, UObject* DuplicatedObjectOuter) const override;
 	virtual FText GetSourceType(FGuid EntryGuid) const override;
 
diff --git a/Engine/Plugins/Animation/LiveLink/Source/LiveLink/Public/LiveLinkPreset.h b/Engine/Plugins/Animation/LiveLink/Source/LiveLink/Public/LiveLinkPreset.h
--- a/Engine/Plugins/Animation/LiveLink/Source/LiveLink/Public/LiveLinkPreset.h
+++ b/Engine/Plugins/Animation/LiveLink/Source/LiveLink/Public/LiveLinkPreset.h
@@ -3,17 +3,22 @@
 #pragma once
 
 #include "CoreMinimal.h"
+#include "Engine/EngineTypes.h"
+#include "Engine/LatentActionManager.h"
+#include "LiveLinkPresetTypes.h"
+#include "Templates/PimplPtr.h"
 #include "UObject/Object.h"
 #include "UObject/ObjectMacros.h"
-#include "LiveLinkPresetTypes.h"
-#include "LiveLinkPreset.generated.h"
 
+#include "LiveLinkPreset.generated.h"
 
 UCLASS(BlueprintType)
 class LIVELINK_API ULiveLinkPreset : public UObject
 {
 	GENERATED_BODY()
 
+	virtual ~ULiveLinkPreset();
+
 private:
 	UPROPERTY(VisibleAnywhere, Category = "LiveLinkSourcePresets")
 	TArray<FLiveLinkSourcePreset> Sources;
@@ -32,9 +37,18 @@ public:
 	 * Remove all previous sources and subjects and add the sources and subjects from this preset.
 	 * @return True is all sources and subjects from this preset could be created and added.
 	 */
+	UE_DEPRECATED(5.0, "This function is deprecated, please use ApplyToClientLatent")
 	UFUNCTION(BlueprintCallable, BlueprintPure = false, Category="LiveLink")
 	bool ApplyToClient() const;
 
+	/**
+	 * Remove all previous sources and subjects and add the sources and subjects from this preset.
+	 * @return True is all sources and subjects from this preset could be created and added.
+	 */
+	UFUNCTION(BlueprintCallable, BlueprintPure = false, Category="LiveLink", meta = (Latent, LatentInfo = "LatentInfo", HidePin = "WorldContextObject", DefaultToSelf = "WorldContextObject"))
+	void ApplyToClientLatent(UObject* WorldContextObject, FLatentActionInfo LatentInfo);
+	void ApplyToClientLatent(TFunction<void(bool)> CompletionCallback = nullptr);
+
 	/**
 	 * Add the sources and subjects from this preset, but leave any existing sources and subjects connected.
 	 *
@@ -48,4 +62,15 @@ public:
 	/** Reset this preset and build the list of sources and subjects from the client. */
 	UFUNCTION(BlueprintCallable, Category="LiveLink")
 	void BuildFromClient();
+
+private:
+	/** Clear the timer registered with the current world. */
+	void ClearApplyToClientTimer();
+	
+private:
+	/** Holds an handle to the timer used to apply a preset asynchronously with ApplyToClientLatent. */
+	FTimerHandle ApplyToClientTimerHandle;
+
+	/** Holds the current ApplyToClient async operation. */
+	TPimplPtr<struct FApplyToClientPollingOperation> ApplyToClientPollingOperation;
 };
diff --git a/Engine/Plugins/Animation/LiveLink/Source/LiveLinkEditor/Private/LiveLinkClientPanelToolbar.cpp b/Engine/Plugins/Animation/LiveLink/Source/LiveLinkEditor/Private/LiveLinkClientPanelToolbar.cpp
--- a/Engine/Plugins/Animation/LiveLink/Source/LiveLinkEditor/Private/LiveLinkClientPanelToolbar.cpp
+++ b/Engine/Plugins/Animation/LiveLink/Source/LiveLinkEditor/Private/LiveLinkClientPanelToolbar.cpp
@@ -878,7 +878,7 @@ void SLiveLinkClientPanelToolbar::OnImportPreset(const FAssetData& InPreset)
 	if (ImportedPreset)
 	{
 		FScopedTransaction Transaction(LOCTEXT("ImportPreset_Transaction", "Import LiveLink Preset"));
-		ImportedPreset->ApplyToClient();
+		ImportedPreset->ApplyToClientLatent();
 	}
 	LiveLinkPreset = ImportedPreset;
 }
@@ -895,7 +895,7 @@ FReply SLiveLinkClientPanelToolbar::OnRevertChanges()
 	ULiveLinkPreset* CurrentPreset = LiveLinkPreset.Get();
 	if (CurrentPreset)
 	{
-		CurrentPreset->ApplyToClient();
+		CurrentPreset->ApplyToClientLatent();
 	}
 
 	return FReply::Handled();
@@ -906,4 +906,5 @@ bool SLiveLinkClientPanelToolbar::HasLoadedLiveLinkPreset() const
 	return LiveLinkPreset.IsValid();
 }
 
+
 #undef LOCTEXT_NAMESPACE
\ No newline at end of file
diff --git a/Engine/Source/Runtime/LiveLinkInterface/Public/ILiveLinkClient.h b/Engine/Source/Runtime/LiveLinkInterface/Public/ILiveLinkClient.h
--- a/Engine/Source/Runtime/LiveLinkInterface/Public/ILiveLinkClient.h
+++ b/Engine/Source/Runtime/LiveLinkInterface/Public/ILiveLinkClient.h
@@ -103,7 +103,7 @@ PRAGMA_ENABLE_DEPRECATION_WARNINGS
 
 
 /**
- * Interface for streaming and consuming data from external sources into UE.
+ * Interface for streaming and consuming data from external sources into UE4.
  * A LiveLinkSource, may stream multiple LiveLinkSubject.
  * Pushing and evaluating data can be executed on any thread. The other functions must be executed on the Game Thread.
  * Subject may shared name between sources, but only 1 of those subjects may be enabled.
@@ -135,11 +135,19 @@ public:
 	/** Is the source been added */
 	virtual bool HasSourceBeenAdded(TSharedPtr<ILiveLinkSource> Source) const = 0;
 
-	/** Get a list of all the sources */
-	virtual TArray<FGuid> GetSources() const = 0;
+	/**
+	 * Get a list of all the sources
+	 * @param bEvenIfPendingKill Whether or not to include sources that are pending kill.
+	 * @return the list of sources.
+	 */
+	virtual TArray<FGuid> GetSources(bool bEvenIfPendingKill = false) const = 0;
 
-	/** Get a list of all the VirtualSubjects sources */
-	virtual TArray<FGuid> GetVirtualSources() const = 0;
+	/**
+	 * Get a list of all the VirtualSubjects sources
+	 * @param bEvenIfPendingKill Whether or not to include sources that are pending kill.
+	 * @return the list of sources.
+	 */
+	virtual TArray<FGuid> GetVirtualSources(bool bEvenIfPendingKill = false) const = 0;
 
 	/** Get the source preset from the live link client. The settings will be duplicated into DuplicatedObjectOuter. */
 	virtual FLiveLinkSourcePreset GetSourcePreset(FGuid SourceGuid, UObject* DuplicatedObjectOuter) const = 0;
