LiveLink, Compositing, Aja, Blackmagic: Merged from RES.

diff --git a/Engine/Plugins/Animation/LiveLink/Source/LiveLink/Private/SLiveLinkMessageBusSourceFactory.cpp b/Engine/Plugins/Animation/LiveLink/Source/LiveLink/Private/SLiveLinkMessageBusSourceFactory.cpp
--- a/Engine/Plugins/Animation/LiveLink/Source/LiveLink/Private/SLiveLinkMessageBusSourceFactory.cpp
+++ b/Engine/Plugins/Animation/LiveLink/Source/LiveLink/Private/SLiveLinkMessageBusSourceFactory.cpp
@@ -4,15 +4,10 @@
 
 #include "ILiveLinkClient.h"
 #include "ILiveLinkModule.h"
-#include "LiveLinkMessages.h"
-#include "LiveLinkMessageBusFinder.h"
 #include "LiveLinkMessageBusDiscoveryManager.h"
+#include "LiveLinkMessageBusFinder.h"
 #include "Widgets/Layout/SBox.h"
 
-#include "Features/IModularFeatures.h"
-#include "MessageEndpointBuilder.h"
-#include "Misc/App.h"
-
 #define LOCTEXT_NAMESPACE "LiveLinkMessageBusSourceEditor"
 
 namespace ProviderPollUI
@@ -21,6 +16,14 @@ namespace ProviderPollUI
 	const FName MachineColumnName(TEXT("Machine"));
 };
 
+namespace
+{
+	bool operator==(const FProviderPollResult& LHS, const FProviderPollResult& RHS)
+	{
+		return LHS.Name == RHS.Name && LHS.MachineName == RHS.MachineName;
+	}
+}
+
 class SProviderPollRow : public SMultiColumnTableRow<FProviderPollResultPtr>
 {
 public:
@@ -74,9 +77,7 @@ SLiveLinkMessageBusSourceFactory::~SLiveLinkMessageBusSourceFactory()
 
 void SLiveLinkMessageBusSourceFactory::Construct(const FArguments& Args)
 {
-	//LastTickTime = 0.0;
 	OnSourceSelected = Args._OnSourceSelected;
-	LastUIUpdateSeconds = 0;
 
 	ILiveLinkModule::Get().GetMessageBusDiscoveryManager().AddDiscoveryMessageRequest();
 
@@ -90,8 +91,8 @@ void SLiveLinkMessageBusSourceFactory::Construct(const FArguments& Args)
 			.HeightOverride(200)
 			.WidthOverride(200)
 			[
-				SAssignNew(ListView, SListView<FProviderPollResultPtr>)
-				.ListItemsSource(&PollData)
+				SAssignNew(ListView, SListView<TSharedPtr<FLiveLinkSource>>)
+				.ListItemsSource(&Sources)
 				.SelectionMode(ESelectionMode::SingleToggle)
 				.OnGenerateRow(this, &SLiveLinkMessageBusSourceFactory::MakeSourceListViewWidget)
 				.OnSelectionChanged(this, &SLiveLinkMessageBusSourceFactory::OnSourceListSelectionChanged)
@@ -116,11 +117,30 @@ void SLiveLinkMessageBusSourceFactory::Tick(const FGeometry& AllottedGeometry, c
 	{
 		PollData.Reset();
 		PollData.Append(ILiveLinkModule::Get().GetMessageBusDiscoveryManager().GetDiscoveryResults());
-		PollData.StableSort([](const FProviderPollResultPtr& LHS, const FProviderPollResultPtr& RHS)
+		
+		Sources.RemoveAllSwap([this](TSharedPtr<FLiveLinkSource> Source)
+		{
+			return FApp::GetCurrentTime() - Source->LastTimeSincePong > SecondsBeforeSourcesDisappear;
+		});
+		
+		for (const FProviderPollResultPtr& PollResult : PollData)
+		{
+			TSharedPtr<FLiveLinkSource>* FoundSource = Sources.FindByPredicate([&PollResult](const TSharedPtr<FLiveLinkSource> Source){ return Source && PollResult && *Source->PollResult == *PollResult;});
+			if (FoundSource && *FoundSource)
+			{
+				(*FoundSource)->LastTimeSincePong = FApp::GetCurrentTime();
+			}
+			else
+			{
+				Sources.Add(MakeShared<FLiveLinkSource>(PollResult));
+			}
+		}
+
+		Sources.StableSort([](const TSharedPtr<FLiveLinkSource>& LHS, const TSharedPtr<FLiveLinkSource>& RHS)
 		{
-			if (LHS.IsValid() && RHS.IsValid())
+			if (LHS.IsValid() && RHS.IsValid() && LHS->PollResult.IsValid() && RHS->PollResult.IsValid())
 			{
-				return LHS->Name.Compare(RHS->Name) <= 0;
+				return LHS->PollResult->Name.Compare(RHS->PollResult->Name) <= 0;
 			}
 			return true;
 		});
@@ -130,14 +150,14 @@ void SLiveLinkMessageBusSourceFactory::Tick(const FGeometry& AllottedGeometry, c
 	}
 }
 
-TSharedRef<ITableRow> SLiveLinkMessageBusSourceFactory::MakeSourceListViewWidget(FProviderPollResultPtr PollResult, const TSharedRef<STableViewBase>& OwnerTable) const
+TSharedRef<ITableRow> SLiveLinkMessageBusSourceFactory::MakeSourceListViewWidget(TSharedPtr<FLiveLinkSource> Source, const TSharedRef<STableViewBase>& OwnerTable) const
 {
-	return SNew(SProviderPollRow, OwnerTable).PollResultPtr(PollResult);
+	return SNew(SProviderPollRow, OwnerTable).PollResultPtr(Source->PollResult);
 }
 
-void SLiveLinkMessageBusSourceFactory::OnSourceListSelectionChanged(FProviderPollResultPtr PollResult, ESelectInfo::Type SelectionType)
+void SLiveLinkMessageBusSourceFactory::OnSourceListSelectionChanged(TSharedPtr<FLiveLinkSource> Source, ESelectInfo::Type SelectionType)
 {
-	SelectedResult = PollResult;
+	SelectedResult = Source->PollResult;
 	OnSourceSelected.ExecuteIfBound(SelectedResult);
 }
 
diff --git a/Engine/Plugins/Animation/LiveLink/Source/LiveLink/Private/SLiveLinkMessageBusSourceFactory.h b/Engine/Plugins/Animation/LiveLink/Source/LiveLink/Private/SLiveLinkMessageBusSourceFactory.h
--- a/Engine/Plugins/Animation/LiveLink/Source/LiveLink/Private/SLiveLinkMessageBusSourceFactory.h
+++ b/Engine/Plugins/Animation/LiveLink/Source/LiveLink/Private/SLiveLinkMessageBusSourceFactory.h
@@ -4,10 +4,9 @@
 
 #include "Widgets/SCompoundWidget.h"
 #include "IMessageContext.h"
-#include "MessageEndpoint.h"
-#include "Misc/Guid.h"
-#include "Widgets/Views/SListView.h"
 #include "LiveLinkMessageBusFinder.h"
+#include "Misc/App.h"
+#include "Widgets/Views/SListView.h"
 
 struct FLiveLinkPongMessage;
 struct FMessageAddress;
@@ -33,17 +32,45 @@ class SLiveLinkMessageBusSourceFactory : public SCompoundWidget
 	virtual void Tick(const FGeometry& AllottedGeometry, const double InCurrentTime, const float InDeltaTime);
 
 private:
-	TSharedRef<ITableRow> MakeSourceListViewWidget(FProviderPollResultPtr PollResult, const TSharedRef<STableViewBase>& OwnerTable) const;
+	/** Represents a source with its latest time since it was polled successfully. */
+	struct FLiveLinkSource
+	{
+		FProviderPollResultPtr PollResult;
+		double LastTimeSincePong = 0.0;
+		
+		FLiveLinkSource(FProviderPollResultPtr InPollResult)
+            : PollResult(MoveTemp(InPollResult))
+			, LastTimeSincePong(FApp::GetCurrentTime()) 
+		{
+		}
+	};
 
-	void OnSourceListSelectionChanged(FProviderPollResultPtr PollResult, ESelectInfo::Type SelectionType);
+private:
+	/** Create a widget entry in the source list based on a live link source. */
+	TSharedRef<ITableRow> MakeSourceListViewWidget(TSharedPtr<FLiveLinkSource> Source, const TSharedRef<STableViewBase>& OwnerTable) const;
 
-	TSharedPtr<SListView<FProviderPollResultPtr>> ListView;
+	/** Handles selection changing in the list. */
+	void OnSourceListSelectionChanged(TSharedPtr<FLiveLinkSource> Source, ESelectInfo::Type SelectionType);
 
+private:
+	/** The widgets displayed in the list. */
+	TSharedPtr<SListView<TSharedPtr<FLiveLinkSource>>> ListView;
+
+	/** Latest sources pulled from the discovery manager. */
 	TArray<FProviderPollResultPtr> PollData;
 
+	/** The sources displayed by the list. */
+	TArray<TSharedPtr<FLiveLinkSource>> Sources;
+
+	/** The source that was selected. */
 	FProviderPollResultPtr SelectedResult;
 
+	/** The last source that was selected. */
 	FOnLiveLinkMessageBusSourceSelected OnSourceSelected;
 
-	double LastUIUpdateSeconds;
+	/** The last time at which the UI was refreshed. */
+	double LastUIUpdateSeconds = 0;
+
+	/** Time before a source should disappear from the UI after its last pong. */
+	double SecondsBeforeSourcesDisappear = 2.0;
 };
\ No newline at end of file
diff --git a/Engine/Plugins/Compositing/OpenColorIO/Source/OpenColorIO/Private/OpenColorIOColorTransform.cpp b/Engine/Plugins/Compositing/OpenColorIO/Source/OpenColorIO/Private/OpenColorIOColorTransform.cpp
--- a/Engine/Plugins/Compositing/OpenColorIO/Source/OpenColorIO/Private/OpenColorIOColorTransform.cpp
+++ b/Engine/Plugins/Compositing/OpenColorIO/Source/OpenColorIO/Private/OpenColorIOColorTransform.cpp
@@ -234,7 +234,7 @@ void UOpenColorIOColorTransform::CacheResourceTextures()
 #if !PLATFORM_EXCEPTIONS_DISABLED
 			catch (OCIO_NAMESPACE::Exception& exception)
 			{
-				UE_LOG(LogOpenColorIO, Error, TEXT("Failed to cache 3dLUT for color transform %s. Error message: %s."), *GetTransformFriendlyName(), StringCast<TCHAR>(exception.what()).Get());
+				UE_LOG(LogOpenColorIO, Log, TEXT("Failed to cache 3dLUT for color transform %s. Error message: %s."), *GetTransformFriendlyName(), StringCast<TCHAR>(exception.what()).Get());
 			}
 #endif
 		}
@@ -413,7 +413,7 @@ bool UOpenColorIOColorTransform::UpdateShaderInfo(FString& OutShaderCodeHash, FS
 #if !PLATFORM_EXCEPTIONS_DISABLED
 		catch (OCIO_NAMESPACE::Exception& exception)
 		{
-			UE_LOG(LogOpenColorIO, Error, TEXT("Failed to fetch shader info for color transform %s. Error message: %s."), *GetTransformFriendlyName(), StringCast<TCHAR>(exception.what()).Get());
+			UE_LOG(LogOpenColorIO, Log, TEXT("Failed to fetch shader info for color transform %s. Error message: %s."), *GetTransformFriendlyName(), StringCast<TCHAR>(exception.what()).Get());
 		}
 #endif
 	}
diff --git a/Engine/Plugins/Compositing/OpenColorIO/Source/OpenColorIO/Private/OpenColorIOConfiguration.cpp b/Engine/Plugins/Compositing/OpenColorIO/Source/OpenColorIO/Private/OpenColorIOConfiguration.cpp
--- a/Engine/Plugins/Compositing/OpenColorIO/Source/OpenColorIO/Private/OpenColorIOConfiguration.cpp
+++ b/Engine/Plugins/Compositing/OpenColorIO/Source/OpenColorIO/Private/OpenColorIOConfiguration.cpp
@@ -378,14 +378,20 @@ void UOpenColorIOConfiguration::LoadConfigurationFile()
 			LoadedConfig.reset();
 
 			FString FullPath;
-			if (!FPaths::IsRelative(ConfigurationFile.FilePath))
+			FString ConfigurationFilePath = ConfigurationFile.FilePath;
+			if (ConfigurationFilePath.Contains(TEXT("{Engine}")))
 			{
-				FullPath = ConfigurationFile.FilePath;
+				ConfigurationFilePath = FPaths::ConvertRelativePathToFull(ConfigurationFilePath.Replace(TEXT("{Engine}"), *FPaths::EngineDir()));
+			}    
+
+			if (!FPaths::IsRelative(ConfigurationFilePath))
+			{
+				FullPath = ConfigurationFilePath;
 			}
 			else
 			{
 				const FString AbsoluteGameDir = FPaths::ConvertRelativePathToFull(FPaths::ProjectDir());
-				FullPath = FPaths::ConvertRelativePathToFull(FPaths::Combine(AbsoluteGameDir, ConfigurationFile.FilePath));
+				FullPath = FPaths::ConvertRelativePathToFull(FPaths::Combine(AbsoluteGameDir, ConfigurationFilePath));
 			}
 
 			OCIO_NAMESPACE::ConstConfigRcPtr NewConfig = OCIO_NAMESPACE::Config::CreateFromFile(StringCast<ANSICHAR>(*FullPath).Get());
diff --git a/Engine/Plugins/Compositing/OpenColorIO/Source/OpenColorIOEditor/Private/OpenColorIOColorSpaceCustomization.cpp b/Engine/Plugins/Compositing/OpenColorIO/Source/OpenColorIOEditor/Private/OpenColorIOColorSpaceCustomization.cpp
--- a/Engine/Plugins/Compositing/OpenColorIO/Source/OpenColorIOEditor/Private/OpenColorIOColorSpaceCustomization.cpp
+++ b/Engine/Plugins/Compositing/OpenColorIO/Source/OpenColorIOEditor/Private/OpenColorIOColorSpaceCustomization.cpp
@@ -110,14 +110,20 @@ bool FOpenColorIOColorSpaceCustomization::LoadConfigurationFile(const FFilePath&
 #endif
 	{
 		FString FullPath;
-		if (!FPaths::IsRelative(InFilePath.FilePath))
+		FString ConfigurationFilePath = InFilePath.FilePath;
+		if (ConfigurationFilePath.Contains(TEXT("{Engine}")))
 		{
-			FullPath = InFilePath.FilePath;
+			ConfigurationFilePath = FPaths::ConvertRelativePathToFull(ConfigurationFilePath.Replace(TEXT("{Engine}"), *FPaths::EngineDir()));
+		}    
+
+		if (!FPaths::IsRelative(ConfigurationFilePath))
+		{
+			FullPath = ConfigurationFilePath;
 		}
 		else
 		{
 			const FString AbsoluteGameDir = FPaths::ConvertRelativePathToFull(FPaths::ProjectDir());
-			FullPath = FPaths::ConvertRelativePathToFull(FPaths::Combine(AbsoluteGameDir, InFilePath.FilePath));
+			FullPath = FPaths::ConvertRelativePathToFull(FPaths::Combine(AbsoluteGameDir, ConfigurationFilePath));
 		}
 
 		CachedConfigFile = OCIO_NAMESPACE::Config::CreateFromFile(StringCast<ANSICHAR>(*FullPath).Get());
