Fix CIS warning in LiveLinkPreset

diff --git a/Engine/Plugins/Animation/LiveLink/Source/LiveLink/Private/LiveLinkPreset.cpp b/Engine/Plugins/Animation/LiveLink/Source/LiveLink/Private/LiveLinkPreset.cpp
--- a/Engine/Plugins/Animation/LiveLink/Source/LiveLink/Private/LiveLinkPreset.cpp
+++ b/Engine/Plugins/Animation/LiveLink/Source/LiveLink/Private/LiveLinkPreset.cpp
@@ -85,7 +85,7 @@ namespace LiveLinkPresetUtils
 			{
 				if (WorldContext.WorldType == EWorldType::Game)
 				{
-					return World;
+					return WorldContext.World();
 				}
 			}
 		}
