- Merging fix for default livelink preset not being applied on launch

diff --git a/Engine/Plugins/Animation/LiveLink/Source/LiveLink/Private/LiveLinkModule.cpp b/Engine/Plugins/Animation/LiveLink/Source/LiveLink/Private/LiveLinkModule.cpp
--- a/Engine/Plugins/Animation/LiveLink/Source/LiveLink/Private/LiveLinkModule.cpp
+++ b/Engine/Plugins/Animation/LiveLink/Source/LiveLink/Private/LiveLinkModule.cpp
@@ -6,6 +6,7 @@
 #include "LiveLinkLogInstance.h"
 #include "LiveLinkPreset.h"
 #include "LiveLinkSettings.h"
+#include "Misc/CommandLine.h"
 #include "Misc/CoreDelegates.h"
 
 #define LOCTEXT_NAMESPACE "LiveLinkModule"
@@ -62,9 +63,36 @@ void FLiveLinkModule::CreateStyle()
 
 void FLiveLinkModule::OnEngineLoopInitComplete()
 {
-	if (ULiveLinkPreset* Preset = GetDefault<ULiveLinkSettings>()->DefaultLiveLinkPreset.LoadSynchronous())
+	ULiveLinkPreset* StartupPreset = nullptr;
+	const FString& CommandLine = FCommandLine::Get();
+	const TCHAR* PresetStr = TEXT("LiveLink.Preset.Apply Preset=");
+	const int32 FoundElement = CommandLine.Find(PresetStr);
+	if (FoundElement != INDEX_NONE)
 	{
-		Preset->ApplyToClientLatent();
+		int32 LiveLinkArgumentEnd = CommandLine.Find(",", ESearchCase::IgnoreCase, ESearchDir::FromStart, FoundElement);
+		if (LiveLinkArgumentEnd == INDEX_NONE)
+		{
+			LiveLinkArgumentEnd = CommandLine.Find("\"", ESearchCase::IgnoreCase, ESearchDir::FromStart, FoundElement);
+		}
+		if (LiveLinkArgumentEnd != INDEX_NONE)
+		{
+			const int32 StartIndex = FoundElement + FCString::Strlen(PresetStr);
+			if (CommandLine.IsValidIndex(StartIndex) && CommandLine.IsValidIndex(LiveLinkArgumentEnd))
+			{
+				const FString LiveLinkPresetName = CommandLine.Mid(StartIndex, LiveLinkArgumentEnd - StartIndex);
+				StartupPreset = Cast<ULiveLinkPreset>(StaticLoadObject(ULiveLinkPreset::StaticClass(), nullptr, *LiveLinkPresetName));
+			}
+		}
+	}
+
+	if (StartupPreset == nullptr)
+	{
+		StartupPreset = GetDefault<ULiveLinkSettings>()->DefaultLiveLinkPreset.LoadSynchronous();
+	}
+
+	if (StartupPreset)
+	{
+		StartupPreset->ApplyToClientLatent();
 	}
 }
 
diff --git a/Engine/Plugins/Animation/LiveLink/Source/LiveLink/Private/LiveLinkPreset.cpp b/Engine/Plugins/Animation/LiveLink/Source/LiveLink/Private/LiveLinkPreset.cpp
--- a/Engine/Plugins/Animation/LiveLink/Source/LiveLink/Private/LiveLinkPreset.cpp
+++ b/Engine/Plugins/Animation/LiveLink/Source/LiveLink/Private/LiveLinkPreset.cpp
@@ -2,6 +2,7 @@
 
 #include "LiveLinkPreset.h"
 
+#include "Misc/CoreDelegates.h"
 #include "Engine/Engine.h"
 #include "Engine/LatentActionManager.h"
 #include "Features/IModularFeatures.h"
@@ -10,11 +11,7 @@
 #include "LiveLinkClient.h"
 #include "LiveLinkLog.h"
 #include "Misc/App.h"
-#include "TimerManager.h"
 
-#if WITH_EDITOR
-#include "Editor.h"
-#endif
 
 namespace
 {
@@ -23,21 +20,25 @@ namespace
 		TEXT("Apply a LiveLinkPreset. Use: LiveLink.Preset.Apply Preset=/Game/Folder/MyLiveLinkPreset.MyLiveLinkPreset"),
 		FConsoleCommandWithArgsDelegate::CreateLambda([](const TArray<FString>& Args)
 			{
-				for (const FString& Element : Args)
+				//LiveLinkModule now looks for this commandline argument when starting up. No need to execute it twice at launch
+				if (GFrameCounter > 1)
 				{
-					const TCHAR* PresetStr = TEXT("Preset=");
-					const int32 FoundElement = Element.Find(PresetStr);
-					if (FoundElement != INDEX_NONE)
+					for (const FString& Element : Args)
 					{
-						UObject* Object = StaticLoadObject(ULiveLinkPreset::StaticClass(), nullptr, *Element + FoundElement + FCString::Strlen(PresetStr));
-						if (Object)
+						const TCHAR* PresetStr = TEXT("Preset=");
+						const int32 FoundElement = Element.Find(PresetStr);
+						if (FoundElement != INDEX_NONE)
 						{
-							CastChecked<ULiveLinkPreset>(Object)->ApplyToClientLatent();
+							UObject* Object = StaticLoadObject(ULiveLinkPreset::StaticClass(), nullptr, *Element + FoundElement + FCString::Strlen(PresetStr));
+							if (Object)
+							{
+								CastChecked<ULiveLinkPreset>(Object)->ApplyToClientLatent();
+							}
 						}
 					}
 				}
 			})
-		);
+	);
 
 	static FAutoConsoleCommand LiveLinkPresetAddCmd(
 		TEXT("LiveLink.Preset.Add"),
@@ -61,39 +62,6 @@ namespace
 		);
 }
 
-namespace LiveLinkPresetUtils
-{
-	UWorld* GetWorld()
-	{
-		UWorld* World = nullptr;
-
-#if WITH_EDITOR
-		if (GEditor && FApp::CanEverRender())
-		{
-			World = GEditor->GetEditorWorldContext(false).World();
-		}
-#endif
-
-		if (World)
-		{
-			return World;
-		}
-
-		if (GEngine)
-		{
-			for (const FWorldContext& WorldContext : GEngine->GetWorldContexts())
-			{
-				if (WorldContext.WorldType == EWorldType::Game)
-				{
-					return WorldContext.World();
-				}
-			}
-		}
-		
-		return nullptr;
-	}
-}
-
 struct FApplyToClientPollingOperation
 {
 	/** Holds a weak pointer to the preset that will be applied. */
@@ -174,6 +142,8 @@ struct FApplyToClientPollingOperation
 	}
 };
 
+TPimplPtr<FApplyToClientPollingOperation> ULiveLinkPreset::ApplyToClientPollingOperation;
+
 class FApplyToClientLatentAction : public FPendingLatentAction
 {
 public:
@@ -264,29 +234,23 @@ void ULiveLinkPreset::ApplyToClientLatent(TFunction<void(bool)> CompletionCallba
 
 	ApplyToClientPollingOperation = MakePimpl<FApplyToClientPollingOperation>(this);
 
-	if (UWorld* World = LiveLinkPresetUtils::GetWorld())
+	ApplyToClientEndFrameHandle = FCoreDelegates::OnEndFrame.AddLambda([WeakThis = TWeakObjectPtr<ULiveLinkPreset>(this), WrappedCallback = MoveTemp(CompletionCallback)]()
 	{
-		FTimerDelegate TimerDelegate;
-		TimerDelegate.BindLambda([WeakThis = TWeakObjectPtr<ULiveLinkPreset>(this), WrappedCallback = MoveTemp(CompletionCallback)]()
+		if (ensure(WeakThis->ApplyToClientPollingOperation))
 		{
-			if (ensure(WeakThis->ApplyToClientPollingOperation))
+			FApplyToClientPollingOperation::EApplyToClientUpdateResult Result = WeakThis->ApplyToClientPollingOperation->Update();
+			if (Result != FApplyToClientPollingOperation::EApplyToClientUpdateResult::Pending)
 			{
-				FApplyToClientPollingOperation::EApplyToClientUpdateResult Result = WeakThis->ApplyToClientPollingOperation->Update();
-				if (Result != FApplyToClientPollingOperation::EApplyToClientUpdateResult::Pending)
+				if (WrappedCallback != nullptr)
 				{
-					if (WrappedCallback != nullptr)
-					{
-						WrappedCallback(Result == FApplyToClientPollingOperation::EApplyToClientUpdateResult::Success ? true : false);
-					}
-
-					WeakThis->ApplyToClientPollingOperation.Reset();
-					WeakThis->ClearApplyToClientTimer();
+					WrappedCallback(Result == FApplyToClientPollingOperation::EApplyToClientUpdateResult::Success ? true : false);
 				}
+
+				WeakThis->ApplyToClientPollingOperation.Reset();
+				WeakThis->ClearApplyToClientTimer();
 			}
-		});
-		
-		World->GetTimerManager().SetTimer(ApplyToClientTimerHandle, MoveTemp(TimerDelegate), 0.01f, true);
-	}
+		}
+	});
 }
 
 bool ULiveLinkPreset::AddToClient(const bool bRecreatePresets) const
@@ -392,8 +356,8 @@ void ULiveLinkPreset::BuildFromClient()
 
 void ULiveLinkPreset::ClearApplyToClientTimer()
 {
-	if (UWorld* World = LiveLinkPresetUtils::GetWorld())
+	if (ApplyToClientEndFrameHandle.IsValid())
 	{
-		World->GetTimerManager().ClearTimer(ApplyToClientTimerHandle);
+		FCoreDelegates::OnEndFrame.Remove(ApplyToClientEndFrameHandle);
 	}
 }
diff --git a/Engine/Plugins/Animation/LiveLink/Source/LiveLink/Public/LiveLinkPreset.h b/Engine/Plugins/Animation/LiveLink/Source/LiveLink/Public/LiveLinkPreset.h
--- a/Engine/Plugins/Animation/LiveLink/Source/LiveLink/Public/LiveLinkPreset.h
+++ b/Engine/Plugins/Animation/LiveLink/Source/LiveLink/Public/LiveLinkPreset.h
@@ -68,9 +68,9 @@ private:
 	void ClearApplyToClientTimer();
 	
 private:
-	/** Holds an handle to the timer used to apply a preset asynchronously with ApplyToClientLatent. */
-	FTimerHandle ApplyToClientTimerHandle;
+	/** Holds a handle to the OnEndFrame delegate used to apply a preset asynchronously with ApplyToClientLatent. */
+	FDelegateHandle ApplyToClientEndFrameHandle;
 
-	/** Holds the current ApplyToClient async operation. */
-	TPimplPtr<struct FApplyToClientPollingOperation> ApplyToClientPollingOperation;
+	/** Holds the current ApplyToClient async operation. Only one operation for all presets can be done at a time. */
+	static TPimplPtr<struct FApplyToClientPollingOperation> ApplyToClientPollingOperation;
 };
