LiveLink: Added position and orientation offset after axis remapping to pre processor

diff --git a/Engine/Plugins/Animation/LiveLink/Source/LiveLink/Private/PreProcessor/LiveLinkAxisSwitchPreProcessor.cpp b/Engine/Plugins/Animation/LiveLink/Source/LiveLink/Private/PreProcessor/LiveLinkAxisSwitchPreProcessor.cpp
--- a/Engine/Plugins/Animation/LiveLink/Source/LiveLink/Private/PreProcessor/LiveLinkAxisSwitchPreProcessor.cpp
+++ b/Engine/Plugins/Animation/LiveLink/Source/LiveLink/Private/PreProcessor/LiveLinkAxisSwitchPreProcessor.cpp
@@ -26,7 +26,7 @@ namespace LiveLinkAxisSwitchPreProcessor
 		return static_cast<ELiveLinkAxis>(Value);
 	}
 
-	void SwitchTransform(FTransform& Transform, ELiveLinkAxis FrontAxis, ELiveLinkAxis RightAxis, ELiveLinkAxis UpAxis)
+	void SwitchTransform(FTransform& Transform, ELiveLinkAxis FrontAxis, ELiveLinkAxis RightAxis, ELiveLinkAxis UpAxis, bool bUseOffsetPosition, FVector OffsetPosition, bool bUseOffsetOrientation, FRotator OffsetOrientation)
 	{
 		FVector Location = Transform.GetLocation();
 		FVector NewLocation;
@@ -96,6 +96,16 @@ namespace LiveLinkAxisSwitchPreProcessor
 
 		NewOrientation.W = Orientation.W * LCSymbol * FlipSign;
 
+		// Apply any offsets after the remapping
+		if (bUseOffsetPosition)
+		{
+			NewLocation += OffsetPosition;
+		}
+		if (bUseOffsetOrientation)
+		{
+			NewOrientation *= OffsetOrientation.Quaternion();
+		}
+
 		Transform.SetComponents(NewOrientation.GetNormalized(), NewLocation, NewScale);
 	}
 }
@@ -111,7 +121,7 @@ TSubclassOf<ULiveLinkRole> ULiveLinkTransformAxisSwitchPreProcessor::FLiveLinkTr
 bool ULiveLinkTransformAxisSwitchPreProcessor::FLiveLinkTransformAxisSwitchPreProcessorWorker::PreProcessFrame(FLiveLinkFrameDataStruct& InOutFrame) const
 {
 	FLiveLinkTransformFrameData& TransformData = *InOutFrame.Cast<FLiveLinkTransformFrameData>();
-	LiveLinkAxisSwitchPreProcessor::SwitchTransform(TransformData.Transform, FrontAxis, RightAxis, UpAxis);
+	LiveLinkAxisSwitchPreProcessor::SwitchTransform(TransformData.Transform, FrontAxis, RightAxis, UpAxis, bUseOffsetPosition, OffsetPosition, bUseOffsetOrientation, OffsetOrientation);
 	return true;
 }
 
@@ -131,6 +141,10 @@ ULiveLinkFramePreProcessor::FWorkerSharedPtr ULiveLinkTransformAxisSwitchPreProc
 		Instance->FrontAxis = FrontAxis;
 		Instance->RightAxis = RightAxis;
 		Instance->UpAxis = UpAxis;
+		Instance->bUseOffsetPosition = bUseOffsetPosition;
+		Instance->bUseOffsetOrientation = bUseOffsetOrientation;
+		Instance->OffsetPosition = OffsetPosition;
+		Instance->OffsetOrientation = OffsetOrientation;
 	}
 
 	return Instance;
@@ -139,9 +153,20 @@ ULiveLinkFramePreProcessor::FWorkerSharedPtr ULiveLinkTransformAxisSwitchPreProc
 #if WITH_EDITOR
 void ULiveLinkTransformAxisSwitchPreProcessor::PostEditChangeChainProperty(struct FPropertyChangedChainEvent& PropertyChangedEvent)
 {
-	if (PropertyChangedEvent.Property->GetFName() == GET_MEMBER_NAME_CHECKED(ThisClass, FrontAxis) ||
-		PropertyChangedEvent.Property->GetFName() == GET_MEMBER_NAME_CHECKED(ThisClass, RightAxis) ||
-		PropertyChangedEvent.Property->GetFName() == GET_MEMBER_NAME_CHECKED(ThisClass, UpAxis))
+	static const FName NAME_FrontAxis = GET_MEMBER_NAME_CHECKED(ThisClass, FrontAxis);
+	static const FName NAME_RightAxis = GET_MEMBER_NAME_CHECKED(ThisClass, RightAxis);
+	static const FName NAME_UpAxis = GET_MEMBER_NAME_CHECKED(ThisClass, UpAxis);
+	static const FName NAME_OffsetPosition = GET_MEMBER_NAME_CHECKED(ThisClass, OffsetPosition);
+	static const FName NAME_OffsetOrientation = GET_MEMBER_NAME_CHECKED(ThisClass, OffsetOrientation);
+
+	const FName PropertyName = PropertyChangedEvent.Property->GetFName();
+	const FName StructName = PropertyChangedEvent.PropertyChain.GetActiveMemberNode()->GetValue()->GetFName();
+
+	if ((PropertyName == NAME_FrontAxis) || (PropertyName == NAME_RightAxis) || (PropertyName == NAME_UpAxis))
+	{
+		Instance.Reset();
+	}
+	else if ((StructName == NAME_OffsetPosition) || (StructName == NAME_OffsetOrientation))
 	{
 		Instance.Reset();
 	}
@@ -165,7 +190,7 @@ bool ULiveLinkAnimationAxisSwitchPreProcessor::FLiveLinkAnimationAxisSwitchPrePr
 	
 	for (FTransform& Transform : AnimationData.Transforms)
 	{
-		LiveLinkAxisSwitchPreProcessor::SwitchTransform(Transform, FrontAxis, RightAxis, UpAxis);
+		LiveLinkAxisSwitchPreProcessor::SwitchTransform(Transform, FrontAxis, RightAxis, UpAxis, bUseOffsetPosition, OffsetPosition, bUseOffsetOrientation, OffsetOrientation);
 	}
 	
 	return true;
@@ -187,6 +212,10 @@ ULiveLinkFramePreProcessor::FWorkerSharedPtr ULiveLinkAnimationAxisSwitchPreProc
 		Instance->FrontAxis = FrontAxis;
 		Instance->RightAxis = RightAxis;
 		Instance->UpAxis = UpAxis;
+		Instance->bUseOffsetPosition = bUseOffsetPosition;
+		Instance->bUseOffsetOrientation = bUseOffsetOrientation;
+		Instance->OffsetPosition = OffsetPosition;
+		Instance->OffsetOrientation = OffsetOrientation;
 	}
 
 	return Instance;
diff --git a/Engine/Plugins/Animation/LiveLink/Source/LiveLink/Public/PreProcessor/LiveLinkAxisSwitchPreProcessor.h b/Engine/Plugins/Animation/LiveLink/Source/LiveLink/Public/PreProcessor/LiveLinkAxisSwitchPreProcessor.h
--- a/Engine/Plugins/Animation/LiveLink/Source/LiveLink/Public/PreProcessor/LiveLinkAxisSwitchPreProcessor.h
+++ b/Engine/Plugins/Animation/LiveLink/Source/LiveLink/Public/PreProcessor/LiveLinkAxisSwitchPreProcessor.h
@@ -36,6 +36,12 @@ public:
 		ELiveLinkAxis RightAxis = ELiveLinkAxis::Y;
 		ELiveLinkAxis UpAxis = ELiveLinkAxis::Z;
 
+		bool bUseOffsetPosition = false;
+		bool bUseOffsetOrientation = false;
+
+		FVector OffsetPosition = FVector::ZeroVector;
+		FRotator OffsetOrientation = FRotator::ZeroRotator;
+
 		virtual TSubclassOf<ULiveLinkRole> GetRole() const override;
 		virtual bool PreProcessFrame(FLiveLinkFrameDataStruct& InOutFrame) const override;
 	};
@@ -63,6 +69,16 @@ protected:
 	UPROPERTY(EditAnywhere, Category = "LiveLink")
 	ELiveLinkAxis UpAxis = ELiveLinkAxis::Z;
 
+	UPROPERTY(EditAnywhere, Category = "Enables", meta = (PinHiddenByDefault, InlineEditConditionToggle))
+	bool bUseOffsetPosition = false;
+	UPROPERTY(EditAnywhere, Category = "Enables", meta = (PinHiddenByDefault, InlineEditConditionToggle))
+	bool bUseOffsetOrientation = false;
+
+	UPROPERTY(EditAnywhere, Category = "LiveLink", meta = (EditCondition = "bUseOffsetPosition"))
+	FVector OffsetPosition = FVector::ZeroVector;
+	UPROPERTY(EditAnywhere, Category = "LiveLink", meta = (EditCondition = "bUseOffsetOrientation"))
+	FRotator OffsetOrientation = FRotator::ZeroRotator;
+
 public:
 	virtual TSubclassOf<ULiveLinkRole> GetRole() const override;
 	virtual ULiveLinkFramePreProcessor::FWorkerSharedPtr FetchWorker() override;
