- Fix crash when force deleting LiveLink remap asset

diff --git a/Engine/Plugins/Animation/LiveLink/Source/LiveLink/Private/LiveLinkRemapAsset.cpp b/Engine/Plugins/Animation/LiveLink/Source/LiveLink/Private/LiveLinkRemapAsset.cpp
--- a/Engine/Plugins/Animation/LiveLink/Source/LiveLink/Private/LiveLinkRemapAsset.cpp
+++ b/Engine/Plugins/Animation/LiveLink/Source/LiveLink/Private/LiveLinkRemapAsset.cpp
@@ -26,9 +26,10 @@ void ULiveLinkRemapAsset::BeginDestroy()
 #if WITH_EDITOR
 	if (OnBlueprintCompiledDelegate.IsValid())
 	{
-		UBlueprint* Blueprint = Cast<UBlueprint>(GetClass()->ClassGeneratedBy);
-		check(Blueprint);
-		Blueprint->OnCompiled().Remove(OnBlueprintCompiledDelegate);
+		if (UBlueprint* Blueprint = Cast<UBlueprint>(GetClass()->ClassGeneratedBy))
+		{
+			Blueprint->OnCompiled().Remove(OnBlueprintCompiledDelegate);
+		}
 		OnBlueprintCompiledDelegate.Reset();
 	}
 #endif
