- Fixing livelink rebroadcaster still being present when no more subjects are

diff --git a/Engine/Plugins/Animation/LiveLink/Source/LiveLink/Private/LiveLinkClient.cpp b/Engine/Plugins/Animation/LiveLink/Source/LiveLink/Private/LiveLinkClient.cpp
--- a/Engine/Plugins/Animation/LiveLink/Source/LiveLink/Private/LiveLinkClient.cpp
+++ b/Engine/Plugins/Animation/LiveLink/Source/LiveLink/Private/LiveLinkClient.cpp
@@ -71,6 +71,7 @@ FLiveLinkClient::FLiveLinkClient()
 #endif
 
 	Collection = MakeUnique<FLiveLinkSourceCollection>();
+	OnLiveLinkSubjectRemoved().AddRaw(this, &FLiveLinkClient::OnSubjectRemovedCallback);
 	FCoreDelegates::OnPreExit.AddRaw(this, &FLiveLinkClient::Shutdown);
 
 	//IMediaModule& MediaModule = FModuleManager::LoadModuleChecked<IMediaModule>("Media");
@@ -215,6 +216,8 @@ void FLiveLinkClient::Shutdown()
 
 	if (Collection)
 	{
+		OnLiveLinkSubjectRemoved().RemoveAll(this);
+
 		double Timeout = 2.0;
 		GConfig->GetDouble(TEXT("LiveLink"), TEXT("ClientShutdownTimeout"), Timeout, GGameIni);
 
@@ -257,6 +260,7 @@ void FLiveLinkClient::HandleSubjectRebroadcast(ILiveLinkSubject* InSubject, cons
 					StaticDataCopy.InitializeWith(InSubject->GetStaticData());
 					RebroadcastLiveLinkProvider->UpdateSubjectStaticData(InSubject->GetSubjectKey().SubjectName, InSubject->GetRole(), MoveTemp(StaticDataCopy));
 					InSubject->SetStaticDataAsRebroadcasted(true);
+					RebroadcastedSubjects.Add(InSubject->GetSubjectKey());
 				}
 				
 				// Make a copy of the data for use by the rebroadcaster
@@ -273,10 +277,29 @@ void FLiveLinkClient::HandleSubjectRebroadcast(ILiveLinkSubject* InSubject, cons
 	}
 	else if (InSubject->HasStaticDataBeenRebroadcasted())
 	{
-		if (RebroadcastLiveLinkProvider.IsValid())
+		RemoveRebroadcastedSubject(InSubject->GetSubjectKey());
+		InSubject->SetStaticDataAsRebroadcasted(false);
+	}
+}
+
+void FLiveLinkClient::OnSubjectRemovedCallback(FLiveLinkSubjectKey InSubjectKey)
+{
+	RemoveRebroadcastedSubject(InSubjectKey);
+}
+
+void FLiveLinkClient::RemoveRebroadcastedSubject(FLiveLinkSubjectKey InSubjectKey)
+{
+	if (RebroadcastLiveLinkProvider.IsValid())
+	{
+		if (RebroadcastedSubjects.Contains(InSubjectKey))
 		{
-			RebroadcastLiveLinkProvider->RemoveSubject(InSubject->GetSubjectKey().SubjectName);
-			InSubject->SetStaticDataAsRebroadcasted(false);
+			RebroadcastLiveLinkProvider->RemoveSubject(InSubjectKey.SubjectName);
+			RebroadcastedSubjects.Remove(InSubjectKey);
+
+			if (RebroadcastedSubjects.Num() <= 0)
+			{
+				RebroadcastLiveLinkProvider.Reset();
+			}
 		}
 	}
 }
diff --git a/Engine/Plugins/Animation/LiveLink/Source/LiveLink/Public/LiveLinkClient.h b/Engine/Plugins/Animation/LiveLink/Source/LiveLink/Public/LiveLinkClient.h
--- a/Engine/Plugins/Animation/LiveLink/Source/LiveLink/Public/LiveLinkClient.h
+++ b/Engine/Plugins/Animation/LiveLink/Source/LiveLink/Public/LiveLinkClient.h
@@ -242,6 +242,12 @@ private:
 	/** Process virtual subject for rebroadcast purpose */
 	void HandleSubjectRebroadcast(ILiveLinkSubject* InSubject, const FLiveLinkFrameDataStruct& InFrameData);
 
+	/** Called when a subject is removed. Used to remove rebroadcasted subjects */
+	void OnSubjectRemovedCallback(FLiveLinkSubjectKey InSubjectKey);
+
+	/** Removes a subject from the rebroadcast provider and resets it if there are no more subjects */
+	void RemoveRebroadcastedSubject(FLiveLinkSubjectKey InSubjectKey);
+
 private:
 	/** The current collection used. */
 	TUniquePtr<FLiveLinkSourceCollection> Collection;
@@ -285,6 +291,7 @@ private:
 	/** LiveLink Provider for rebroadcasting */
 	TSharedPtr<ILiveLinkProvider> RebroadcastLiveLinkProvider;
 	FString RebroadcastLiveLinkProviderName;
+	TSet<FLiveLinkSubjectKey> RebroadcastedSubjects;
 
 #if WITH_EDITOR
 	/** Delegate when a subject is evaluated. */
