Add BP getter/setter for the live link component controller's subject representation, and BP event when the controller map has changed.

diff --git a/Engine/Plugins/Animation/LiveLink/Source/LiveLinkComponents/Private/LiveLinkComponentController.cpp b/Engine/Plugins/Animation/LiveLink/Source/LiveLinkComponents/Private/LiveLinkComponentController.cpp
--- a/Engine/Plugins/Animation/LiveLink/Source/LiveLinkComponents/Private/LiveLinkComponentController.cpp
+++ b/Engine/Plugins/Animation/LiveLink/Source/LiveLinkComponents/Private/LiveLinkComponentController.cpp
@@ -93,6 +93,22 @@ void ULiveLinkComponentController::OnSubjectRoleChanged()
 		}
 #endif
 	}
+
+	if (OnControllerMapUpdatedDelegate.IsBound())
+	{
+		FEditorScriptExecutionGuard ScriptGuard;
+		OnControllerMapUpdatedDelegate.Broadcast();
+	}
+}
+
+void ULiveLinkComponentController::SetSubjectRepresentation(FLiveLinkSubjectRepresentation InSubjectRepresentation)
+{
+	SubjectRepresentation = InSubjectRepresentation;
+
+	if (IsControllerMapOutdated())
+	{
+		OnSubjectRoleChanged();
+	}
 }
 
 void ULiveLinkComponentController::SetControllerClassForRole(TSubclassOf<ULiveLinkRole> RoleClass, TSubclassOf<ULiveLinkControllerBase> DesiredControllerClass)
@@ -286,6 +302,8 @@ void ULiveLinkComponentController::ConvertOldControllerSystem()
 	Controller_DEPRECATED = nullptr;
 }
 
+#endif //WITH_EDITOR
+
 bool ULiveLinkComponentController::IsControllerMapOutdated() const
 {
 	TArray<TSubclassOf<ULiveLinkRole>> SelectedRoleHierarchy = GetSelectedRoleHierarchyClasses(SubjectRepresentation.Role);
@@ -311,8 +329,6 @@ bool ULiveLinkComponentController::IsControllerMapOutdated() const
 	return false;
 }
 
-#endif //WITH_EDITOR
-
 TArray<TSubclassOf<ULiveLinkRole>> ULiveLinkComponentController::GetSelectedRoleHierarchyClasses(const TSubclassOf<ULiveLinkRole> InCurrentRoleClass) const
 {
 	TArray<TSubclassOf<ULiveLinkRole>> ClassHierarchy;
diff --git a/Engine/Plugins/Animation/LiveLink/Source/LiveLinkComponents/Public/LiveLinkComponentController.h b/Engine/Plugins/Animation/LiveLink/Source/LiveLinkComponents/Public/LiveLinkComponentController.h
--- a/Engine/Plugins/Animation/LiveLink/Source/LiveLinkComponents/Public/LiveLinkComponentController.h
+++ b/Engine/Plugins/Animation/LiveLink/Source/LiveLinkComponents/Public/LiveLinkComponentController.h
@@ -7,6 +7,7 @@
 #include "LiveLinkComponentController.generated.h"
 
 DECLARE_DYNAMIC_MULTICAST_DELEGATE_OneParam(FLiveLinkTickDelegate, float, DeltaTime);
+DECLARE_DYNAMIC_MULTICAST_DELEGATE(FOnControllerMapUpdatedDelegate);
 
 class ULiveLinkControllerBase;
 
@@ -20,7 +21,7 @@ public:
 	~ULiveLinkComponentController();
 
 public:
-	UPROPERTY(BlueprintReadOnly, EditAnywhere, Category="LiveLink")
+	UPROPERTY(EditAnywhere, BlueprintSetter = SetSubjectRepresentation, BlueprintGetter = GetSubjectRepresentation, Category="LiveLink")
 	FLiveLinkSubjectRepresentation SubjectRepresentation;
 
 #if WITH_EDITORONLY_DATA
@@ -39,6 +40,10 @@ public:
 	UPROPERTY(BlueprintAssignable, Category = "LiveLink")
 	FLiveLinkTickDelegate OnLiveLinkUpdated;
 
+	// This Event is triggered any time the controller map is updated
+	UPROPERTY(BlueprintAssignable, Category = "LiveLink")
+	FOnControllerMapUpdatedDelegate OnControllerMapUpdatedDelegate;
+
 	UPROPERTY(EditInstanceOnly, Category = "LiveLink", meta = (UseComponentPicker, AllowedClasses = "ActorComponent", DisallowedClasses = "LiveLinkComponentController"))
 	FComponentReference ComponentToControl;
 
@@ -67,8 +72,13 @@ public:
 	void SetControllerClassForRole(TSubclassOf<ULiveLinkRole> RoleClass, TSubclassOf<ULiveLinkControllerBase> DesiredControllerClass);
 
 	/** Return Representation of Subject that is used in the controller */
+	UFUNCTION(BlueprintGetter)
 	FLiveLinkSubjectRepresentation GetSubjectRepresentation() const { return SubjectRepresentation; }
 
+	/** Set Representation of Subject that is used in the controller and update the controller map */
+	UFUNCTION(BlueprintSetter)
+	void SetSubjectRepresentation(FLiveLinkSubjectRepresentation InSubjectRepresentation);
+
 	/** Returns true if ControllerMap needs to be updated for the current Role. Useful for customization or C++ modification to the Role */
 	bool IsControllerMapOutdated() const;
 	
