Add incompatibility messaging for Livelink providers if using incorrect version

diff --git a/Engine/Plugins/Animation/LiveLink/Source/LiveLink/Private/LiveLinkMessageBusDiscoveryManager.cpp b/Engine/Plugins/Animation/LiveLink/Source/LiveLink/Private/LiveLinkMessageBusDiscoveryManager.cpp
--- a/Engine/Plugins/Animation/LiveLink/Source/LiveLink/Private/LiveLinkMessageBusDiscoveryManager.cpp
+++ b/Engine/Plugins/Animation/LiveLink/Source/LiveLink/Private/LiveLinkMessageBusDiscoveryManager.cpp
@@ -110,7 +110,14 @@ void FLiveLinkMessageBusDiscoveryManager::HandlePongMessage(const FLiveLinkPongM
 
 	if (Message.PollRequest == LastPingRequest)
 	{
+		bool bIsValidProvider = true;
+
+		if (Message.LiveLinkVersion < ILiveLinkClient::LIVELINK_VERSION) // UE4 Remote Clients always send 1 as the LiveLink version
+		{
+			bIsValidProvider = false;
+		}
+
 		const double MachineTimeOffset = LiveLinkMessageBusHelper::CalculateProviderMachineOffset(Message.CreationPlatformTime, Context);
-		LastProviderPoolResults.Emplace(MakeShared<FProviderPollResult, ESPMode::ThreadSafe>(Context->GetSender(), Message.ProviderName, Message.MachineName, MachineTimeOffset));
+		LastProviderPoolResults.Emplace(MakeShared<FProviderPollResult, ESPMode::ThreadSafe>(Context->GetSender(), Message.ProviderName, Message.MachineName, MachineTimeOffset, bIsValidProvider));
 	}
 }
\ No newline at end of file
diff --git a/Engine/Plugins/Animation/LiveLink/Source/LiveLink/Private/LiveLinkMessageBusFinder.cpp b/Engine/Plugins/Animation/LiveLink/Source/LiveLink/Private/LiveLinkMessageBusFinder.cpp
--- a/Engine/Plugins/Animation/LiveLink/Source/LiveLink/Private/LiveLinkMessageBusFinder.cpp
+++ b/Engine/Plugins/Animation/LiveLink/Source/LiveLink/Private/LiveLinkMessageBusFinder.cpp
@@ -79,9 +79,16 @@ void ULiveLinkMessageBusFinder::HandlePongMessage(const FLiveLinkPongMessage& Me
 	if (Message.PollRequest == CurrentPollRequest)
 	{
 		FScopeLock ScopedLock(&PollDataCriticalSection);
-		
+
+		bool bIsValidProvider = true;
+
+		if (Message.LiveLinkVersion < ILiveLinkClient::LIVELINK_VERSION) // UE4 Remote Clients always send 1 as the LiveLink version
+		{
+			bIsValidProvider = false;
+		}
+
 		const double MachineTimeOffset = LiveLinkMessageBusHelper::CalculateProviderMachineOffset(Message.CreationPlatformTime, Context);
-		PollData.Add(FProviderPollResult(Context->GetSender(), Message.ProviderName, Message.MachineName, MachineTimeOffset));
+		PollData.Add(FProviderPollResult(Context->GetSender(), Message.ProviderName, Message.MachineName, MachineTimeOffset, bIsValidProvider));
 	}
 };
 
diff --git a/Engine/Plugins/Animation/LiveLink/Source/LiveLink/Private/SLiveLinkMessageBusSourceFactory.cpp b/Engine/Plugins/Animation/LiveLink/Source/LiveLink/Private/SLiveLinkMessageBusSourceFactory.cpp
--- a/Engine/Plugins/Animation/LiveLink/Source/LiveLink/Private/SLiveLinkMessageBusSourceFactory.cpp
+++ b/Engine/Plugins/Animation/LiveLink/Source/LiveLink/Private/SLiveLinkMessageBusSourceFactory.cpp
@@ -43,6 +43,13 @@ public:
 			.Padding(1.0f),
 			OwnerTableView
 		);
+
+		if (!PollResultPtr->bIsValidProvider)
+		{
+			SetEnabled(false);
+			SetToolTipText(LOCTEXT("InvalidProvider", 
+				"This provider is invalid because it was built with an older version of Unreal Engine. You must upgrade your provider to UE5 to continue using it."));
+		}
 	}
 
 	/** Overridden from SMultiColumnTableRow.  Generates a widget for this column of the list view. */
@@ -51,7 +58,7 @@ public:
 		if (ColumnName == ProviderPollUI::TypeColumnName)
 		{
 			return	SNew(STextBlock)
-					.Text(FText::FromString(PollResultPtr->Name));
+				.Text(FText::FromString(PollResultPtr->Name));
 		}
 		else if (ColumnName == ProviderPollUI::MachineColumnName)
 		{
diff --git a/Engine/Plugins/Animation/LiveLink/Source/LiveLink/Public/LiveLinkMessageBusFinder.h b/Engine/Plugins/Animation/LiveLink/Source/LiveLink/Public/LiveLinkMessageBusFinder.h
--- a/Engine/Plugins/Animation/LiveLink/Source/LiveLink/Public/LiveLinkMessageBusFinder.h
+++ b/Engine/Plugins/Animation/LiveLink/Source/LiveLink/Public/LiveLinkMessageBusFinder.h
@@ -27,6 +27,7 @@ public:
 	
 	FProviderPollResult() = default;
 	
+	UE_DEPRECATED(5.0, "This version of the FProviderPollResult constructor is deprecated. Please use the new constructor instead.")
 	FProviderPollResult(const FMessageAddress& InAddress, const FString& InName, const FString& InMachineName, double InMachineOffset)
 		: Address(InAddress)
 		, Name(InName)
@@ -34,19 +35,31 @@ public:
 		, MachineTimeOffset(InMachineOffset)
 	{}
 
+	FProviderPollResult(const FMessageAddress& InAddress, const FString& InName, const FString& InMachineName, double InMachineOffset, bool bInIsValidProvider)
+		: Address(InAddress)
+		, Name(InName)
+		, MachineName(InMachineName)
+		, MachineTimeOffset(InMachineOffset)
+		, bIsValidProvider(bInIsValidProvider)
+	{}
+
 	FMessageAddress Address;
 
 	// The name of the provider
 	UPROPERTY(BlueprintReadOnly, EditAnywhere, Category="LiveLink", meta=(DisplayName = "Provider Name"))
-	FString			Name;
+	FString Name;
 
 	// The name of the machine the provider is running on
 	UPROPERTY(BlueprintReadOnly, EditAnywhere, Category="LiveLink")
-	FString			MachineName;
+	FString MachineName;
 
 	// Offset between sender's engine time and receiver's engine time
 	UPROPERTY()
-	double			MachineTimeOffset = 0.0;
+	double MachineTimeOffset = 0.0;
+
+	// Whether the provider is valid (compatible with the current version of LiveLink)
+	UPROPERTY(BlueprintReadOnly, VisibleAnywhere, Category = "LiveLink")
+	bool bIsValidProvider = true;
 };
 
 typedef TSharedPtr<FProviderPollResult, ESPMode::ThreadSafe> FProviderPollResultPtr;
diff --git a/Engine/Source/Runtime/LiveLinkMessageBusFramework/Private/LiveLinkProvider.cpp b/Engine/Source/Runtime/LiveLinkMessageBusFramework/Private/LiveLinkProvider.cpp
--- a/Engine/Source/Runtime/LiveLinkMessageBusFramework/Private/LiveLinkProvider.cpp
+++ b/Engine/Source/Runtime/LiveLinkMessageBusFramework/Private/LiveLinkProvider.cpp
@@ -469,7 +469,7 @@ void FLiveLinkProvider::HandlePingMessage(const FLiveLinkPingMessage& Message, c
 		return;
 	}
 
-	MessageEndpoint->Send(FMessageEndpoint::MakeMessage<FLiveLinkPongMessage>(ProviderName, MachineName, Message.PollRequest), Context->GetSender());
+	MessageEndpoint->Send(FMessageEndpoint::MakeMessage<FLiveLinkPongMessage>(ProviderName, MachineName, Message.PollRequest, LIVELINK_SupportedVersion), Context->GetSender());
 }
 
 void FLiveLinkProvider::HandleConnectMessage(const FLiveLinkConnectMessage& Message, const TSharedRef<IMessageContext, ESPMode::ThreadSafe>& Context)
diff --git a/Engine/Source/Runtime/LiveLinkMessageBusFramework/Public/LiveLinkMessages.h b/Engine/Source/Runtime/LiveLinkMessageBusFramework/Public/LiveLinkMessages.h
--- a/Engine/Source/Runtime/LiveLinkMessageBusFramework/Public/LiveLinkMessages.h
+++ b/Engine/Source/Runtime/LiveLinkMessageBusFramework/Public/LiveLinkMessages.h
@@ -52,7 +52,10 @@ struct FLiveLinkPongMessage
 	// default constructor for the receiver
 	FLiveLinkPongMessage() = default;
 
+	UE_DEPRECATED(5.0, "This version of the FLiveLinkPongMessage constructor is deprecated. Please use the new constructor instead to ensure the LiveLinkVersion is set properly.")
 	FLiveLinkPongMessage(const FString& InProviderName, const FString& InMachineName, const FGuid& InPollRequest) : ProviderName(InProviderName), MachineName(InMachineName), PollRequest(InPollRequest), CreationPlatformTime(FPlatformTime::Seconds()) {}
+
+	FLiveLinkPongMessage(const FString& InProviderName, const FString& InMachineName, const FGuid& InPollRequest, int32 InLiveLinkVersion) : ProviderName(InProviderName), MachineName(InMachineName), PollRequest(InPollRequest), LiveLinkVersion(InLiveLinkVersion), CreationPlatformTime(FPlatformTime::Seconds()) {}
 };
 
 USTRUCT()
