- Adding machine name sorting to have stable livelink messagebus source list
- Fixed sorting in LiveLink factories

diff --git a/Engine/Plugins/Animation/LiveLink/Source/LiveLink/Private/SLiveLinkMessageBusSourceFactory.cpp b/Engine/Plugins/Animation/LiveLink/Source/LiveLink/Private/SLiveLinkMessageBusSourceFactory.cpp
--- a/Engine/Plugins/Animation/LiveLink/Source/LiveLink/Private/SLiveLinkMessageBusSourceFactory.cpp
+++ b/Engine/Plugins/Animation/LiveLink/Source/LiveLink/Private/SLiveLinkMessageBusSourceFactory.cpp
@@ -147,7 +147,16 @@ void SLiveLinkMessageBusSourceFactory::Tick(const FGeometry& AllottedGeometry, c
 		{
 			if (LHS.IsValid() && RHS.IsValid() && LHS->PollResult.IsValid() && RHS->PollResult.IsValid())
 			{
-				return LHS->PollResult->Name.Compare(RHS->PollResult->Name) <= 0;
+				//List by SourceName first and then by machine name
+				const int32 Result = LHS->PollResult->Name.Compare(RHS->PollResult->Name);
+				if (Result == 0)
+				{
+					return LHS->PollResult->MachineName.Compare(RHS->PollResult->MachineName) <= 0;
+				}
+				else
+				{
+					return Result <= 0;
+				}
 			}
 			return true;
 		});
diff --git a/Engine/Plugins/Animation/LiveLink/Source/LiveLinkEditor/Private/LiveLinkClientPanelToolbar.cpp b/Engine/Plugins/Animation/LiveLink/Source/LiveLinkEditor/Private/LiveLinkClientPanelToolbar.cpp
--- a/Engine/Plugins/Animation/LiveLink/Source/LiveLinkEditor/Private/LiveLinkClientPanelToolbar.cpp
+++ b/Engine/Plugins/Animation/LiveLink/Source/LiveLinkEditor/Private/LiveLinkClientPanelToolbar.cpp
@@ -2,6 +2,7 @@
 
 #include "LiveLinkClientPanelToolbar.h"
 
+#include "Algo/StableSort.h"
 #include "AssetData.h"
 #include "AssetRegistryModule.h"
 #include "AssetToolsModule.h"
@@ -423,6 +424,11 @@ void SLiveLinkClientPanelToolbar::Construct(const FArguments& Args, FLiveLinkCli
 		}
 	}
 
+	Algo::StableSort(Factories, [](ULiveLinkSourceFactory* LHS, ULiveLinkSourceFactory* RHS)
+	{
+		return LHS->GetSourceDisplayName().CompareTo(RHS->GetSourceDisplayName()) <= 0;
+	});
+
 	const int32 ButtonBoxSize = 28;
 	FMargin ButtonPadding(4.f);
 
