Add another offset to the livelink source settings to achieve smooth evaluation with EngineTime
#UE-139207

diff --git a/Engine/Plugins/Animation/LiveLink/Source/LiveLink/Private/InterpolationProcessor/LiveLinkBasicFrameInterpolateProcessor.cpp b/Engine/Plugins/Animation/LiveLink/Source/LiveLink/Private/InterpolationProcessor/LiveLinkBasicFrameInterpolateProcessor.cpp
--- a/Engine/Plugins/Animation/LiveLink/Source/LiveLink/Private/InterpolationProcessor/LiveLinkBasicFrameInterpolateProcessor.cpp
+++ b/Engine/Plugins/Animation/LiveLink/Source/LiveLink/Private/InterpolationProcessor/LiveLinkBasicFrameInterpolateProcessor.cpp
@@ -1,6 +1,8 @@
 // Copyright Epic Games, Inc. All Rights Reserved.
 
 #include "InterpolationProcessor/LiveLinkBasicFrameInterpolateProcessor.h"
+
+#include "LiveLinkClient.h"
 #include "Roles/LiveLinkBasicRole.h"
 #include "Roles/LiveLinkBasicTypes.h"
 
@@ -184,6 +186,9 @@ namespace LiveLinkInterpolation
 		int32 FrameDataIndexB = INDEX_NONE;
 		if (ULiveLinkBasicFrameInterpolationProcessor::FLiveLinkBasicFrameInterpolationProcessorWorker::FindInterpolateIndex(InTime, InSourceFrames, FrameDataIndexA, FrameDataIndexB, OutInterpolationInfo))
 		{
+			OutInterpolationInfo.FrameIndexA = FrameDataIndexA;
+			OutInterpolationInfo.FrameIndexB = FrameDataIndexB;
+
 			if (FrameDataIndexA == FrameDataIndexB)
 			{
 				// Copy over the frame directly
diff --git a/Engine/Plugins/Animation/LiveLink/Source/LiveLink/Private/LiveLinkSubject.cpp b/Engine/Plugins/Animation/LiveLink/Source/LiveLink/Private/LiveLinkSubject.cpp
--- a/Engine/Plugins/Animation/LiveLink/Source/LiveLink/Private/LiveLinkSubject.cpp
+++ b/Engine/Plugins/Animation/LiveLink/Source/LiveLink/Private/LiveLinkSubject.cpp
@@ -166,8 +166,16 @@ void FLiveLinkSubject::Update()
 
 		case ELiveLinkSourceMode::EngineTime:
 		{
-			//Compute the evaluation time using all offsets (user offset + clock offset)
-			const double ReadTime = FApp::GetCurrentTime() - CachedSettings.BufferSettings.EngineTimeOffset - CachedSettings.BufferSettings.EngineTimeClockOffset;
+			//Compute the evaluation time using all offsets (user offset + clock offset + smooth offset)
+			const double ReadTime = FApp::GetCurrentTime() - CachedSettings.BufferSettings.EngineTimeOffset - CachedSettings.BufferSettings.EngineTimeClockOffset - CachedSettings.BufferSettings.SmoothEngineTimeOffset;
+
+			UE_LOG(LogLiveLink, Verbose, TEXT("Subject '%s' Read Time: %lf    User Offset: %.4lf    Clock Offset: %.4lf    Smooth Offset: %.4lf")
+				, *SubjectKey.SubjectName.ToString()
+				, ReadTime
+				, CachedSettings.BufferSettings.EngineTimeOffset
+				, CachedSettings.BufferSettings.EngineTimeClockOffset
+				, CachedSettings.BufferSettings.SmoothEngineTimeOffset);
+
 			bSnapshotIsValid = GetFrameAtWorldTime(ReadTime, FrameSnapshot);
 			break;
 		}
@@ -211,7 +219,8 @@ TArray<FLiveLinkTime> FLiveLinkSubject::GetFrameTimes() const
 	for (const FLiveLinkFrameDataStruct& Data : FrameData)
 	{
 		//When giving out frame times, include the latest Source clock offset
-		Result.Emplace(Data.GetBaseData()->WorldTime.GetSourceTime() + CachedSettings.BufferSettings.EngineTimeClockOffset, Data.GetBaseData()->MetaData.SceneTime);
+		const double OffsettedWorldTime = Data.GetBaseData()->WorldTime.GetSourceTime() + CachedSettings.BufferSettings.EngineTimeClockOffset + CachedSettings.BufferSettings.SmoothEngineTimeOffset;
+		Result.Emplace(OffsettedWorldTime, Data.GetBaseData()->MetaData.SceneTime);
 	}
 	return Result;
 }
@@ -742,6 +751,11 @@ bool FLiveLinkSubject::GetFrameAtWorldTime_Interpolated(const double InSeconds,
 	FLiveLinkInterpolationInfo InterpolationInfo;
 	FrameInterpolationProcessor->Interpolate(InSeconds, StaticData, FrameData, OutFrame, InterpolationInfo);
 
+	UE_LOG(LogLiveLink, Verbose, TEXT("Subject '%s' interpolating between frames %d and %d")
+		, *SubjectKey.SubjectName.ToString()
+		, InterpolationInfo.FrameIndexA
+		, InterpolationInfo.FrameIndexB);
+
 	VerifyInterpolationInfo(InterpolationInfo);
 
 	return true;
@@ -1132,7 +1146,7 @@ FTimedDataChannelSampleTime FLiveLinkSubject::GetOldestDataTime() const
 			TimecodeClockOffsetTime = FrameData[0].GetBaseData()->MetaData.SceneTime.Rate.AsFrameTime(CachedSettings.BufferSettings.TimecodeClockOffset);
 		}
 
-		const double OffsettedOldestTime = FrameData[0].GetBaseData()->WorldTime.GetSourceTime() + CachedSettings.BufferSettings.EngineTimeClockOffset;
+		const double OffsettedOldestTime = FrameData[0].GetBaseData()->WorldTime.GetSourceTime() + CachedSettings.BufferSettings.EngineTimeClockOffset + CachedSettings.BufferSettings.SmoothEngineTimeOffset;
 		const FQualifiedFrameTime OldestSceneTime = FQualifiedFrameTime(FrameData[0].GetBaseData()->MetaData.SceneTime.Time + TimecodeClockOffsetTime, FrameData[0].GetBaseData()->MetaData.SceneTime.Rate);
 		return FTimedDataChannelSampleTime(OffsettedOldestTime, OldestSceneTime);
 	}
@@ -1150,7 +1164,7 @@ FTimedDataChannelSampleTime FLiveLinkSubject::GetNewestDataTime() const
 			TimecodeClockOffsetTime = FrameData.Last().GetBaseData()->MetaData.SceneTime.Rate.AsFrameTime(CachedSettings.BufferSettings.TimecodeClockOffset);
 		}
 
-		const double OffsettedNewestTime = FrameData.Last().GetBaseData()->WorldTime.GetSourceTime() + CachedSettings.BufferSettings.EngineTimeClockOffset;
+		const double OffsettedNewestTime = FrameData.Last().GetBaseData()->WorldTime.GetSourceTime() + CachedSettings.BufferSettings.EngineTimeClockOffset + CachedSettings.BufferSettings.SmoothEngineTimeOffset;
 		const FQualifiedFrameTime NewestSceneTime = FQualifiedFrameTime(FrameData.Last().GetBaseData()->MetaData.SceneTime.Time + TimecodeClockOffsetTime, FrameData.Last().GetBaseData()->MetaData.SceneTime.Rate);
 		return FTimedDataChannelSampleTime(OffsettedNewestTime, NewestSceneTime);
 	}
@@ -1172,7 +1186,7 @@ TArray<FTimedDataChannelSampleTime> FLiveLinkSubject::GetDataTimes() const
 		Result.Reset(FrameData.Num());
 		for (const FLiveLinkFrameDataStruct& Data : FrameData)
 		{
-			const double OffsettedEngineTime = Data.GetBaseData()->WorldTime.GetSourceTime() + CachedSettings.BufferSettings.EngineTimeClockOffset;
+			const double OffsettedEngineTime = Data.GetBaseData()->WorldTime.GetSourceTime() + CachedSettings.BufferSettings.EngineTimeClockOffset + CachedSettings.BufferSettings.SmoothEngineTimeOffset;
 			const FQualifiedFrameTime AdjustedSceneTime = FQualifiedFrameTime(Data.GetBaseData()->MetaData.SceneTime.Time + TimecodeClockOffsetTime, Data.GetBaseData()->MetaData.SceneTime.Rate);
 			Result.Emplace(OffsettedEngineTime, AdjustedSceneTime);
 		}
diff --git a/Engine/Plugins/Animation/LiveLink/Source/LiveLink/Private/LiveLinkTimedDataInput.cpp b/Engine/Plugins/Animation/LiveLink/Source/LiveLink/Private/LiveLinkTimedDataInput.cpp
--- a/Engine/Plugins/Animation/LiveLink/Source/LiveLink/Private/LiveLinkTimedDataInput.cpp
+++ b/Engine/Plugins/Animation/LiveLink/Source/LiveLink/Private/LiveLinkTimedDataInput.cpp
@@ -43,6 +43,12 @@ namespace LiveLinkTimedDataInput
 		ECVF_Default
 	);
 
+	TAutoConsoleVariable<float> CVarLiveLinkNumFramesForSmoothOffset(
+		TEXT("LiveLink.TimedDataInput.NumFramesForSmoothOffset"),
+		1.5,
+		TEXT("The fractional number of source frames used an offset to achieve smooth evaluation time."),
+		ECVF_Default
+	);
 
 	ETimedDataInputEvaluationType ToTimedDataInputEvaluationType(ELiveLinkSourceMode SourceMode)
 	{
@@ -84,6 +90,8 @@ FLiveLinkTimedDataInput::FLiveLinkTimedDataInput(FLiveLinkClient* InClient, FGui
 
 	EngineClockOffset.SetCorrectionStep(GetDefault<ULiveLinkSettings>()->ClockOffsetCorrectionStep);
 	TimecodeClockOffset.SetCorrectionStep(GetDefault<ULiveLinkSettings>()->ClockOffsetCorrectionStep);
+
+	FrameTimes.Reserve(FrameTimeBufferSize);
 }
 
 FLiveLinkTimedDataInput::~FLiveLinkTimedDataInput()
@@ -199,20 +207,29 @@ const FSlateBrush* FLiveLinkTimedDataInput::GetDisplayIcon() const
 
 void FLiveLinkTimedDataInput::ProcessNewFrameTimingInfo(FLiveLinkBaseFrameData& NewFrameData)
 {
-	if (LiveLinkTimedDataInput::CVarLiveLinkUpdateContinuousClockOffset.GetValueOnGameThread())
+	bool bUpdateContinuousClockOffset = LiveLinkTimedDataInput::CVarLiveLinkUpdateContinuousClockOffset.GetValueOnGameThread();
+
+	//Update both clock offsets for each frame received for our subjects
+	//We mark the last source/frame time that we used to update our offset to only update once per source frame
+	if (!FMath::IsNearlyEqual(NewFrameData.WorldTime.GetSourceTime(), LastWorldSourceTime))
 	{
-		//Update both clock offsets for each frame received for our subjects
-		//We mark the last source/frame time that we used to update our offset to only update once per source frame
-		if (!FMath::IsNearlyEqual(NewFrameData.WorldTime.GetSourceTime(), LastWorldSourceTime))
+		LastWorldSourceTime = NewFrameData.WorldTime.GetSourceTime();
+
+		if (bUpdateContinuousClockOffset)
 		{
-			LastWorldSourceTime = NewFrameData.WorldTime.GetSourceTime();
 			EngineClockOffset.UpdateEstimation(NewFrameData.WorldTime.GetSourceTime(), NewFrameData.ArrivalTime.WorldTime);
 		}
 
-		const double NewFrameSceneTime = NewFrameData.MetaData.SceneTime.AsSeconds();
-		if(!FMath::IsNearlyEqual(NewFrameSceneTime, LastSceneTime))
+		UpdateSmoothEngineTimeOffset(NewFrameData);
+	}
+
+	const double NewFrameSceneTime = NewFrameData.MetaData.SceneTime.AsSeconds();
+	if (!FMath::IsNearlyEqual(NewFrameSceneTime, LastSceneTime))
+	{
+		LastSceneTime = NewFrameSceneTime;
+
+		if (bUpdateContinuousClockOffset)
 		{
-			LastSceneTime = NewFrameSceneTime;
 			TimecodeClockOffset.UpdateEstimation(NewFrameSceneTime, NewFrameData.ArrivalTime.SceneTime.AsSeconds());
 		}
 	}
@@ -227,5 +244,76 @@ void FLiveLinkTimedDataInput::ProcessNewFrameTimingInfo(FLiveLinkBaseFrameData&
 	NewFrameData.WorldTime.SetClockOffset(EngineClockOffset.GetEstimatedOffset());
 }
 
+void FLiveLinkTimedDataInput::UpdateSmoothEngineTimeOffset(const FLiveLinkBaseFrameData& NewFrameData)
+{
+	//Remove the oldest frame from the buffer if it is full
+	if (FrameTimes.Num() >= FrameTimeBufferSize)
+	{
+		constexpr int32 Count = 1;
+		constexpr bool bAllowShrinking = false;
+		FrameTimes.RemoveAt(0, Count, bAllowShrinking);
+	}
+
+	//Add the newest frame time to the buffer
+	FrameTimes.Add(NewFrameData.WorldTime.GetSourceTime());
+
+	if (ULiveLinkSourceSettings* Settings = LiveLinkClient->GetSourceSettings(Source))
+	{
+		//Early-out if no frame offset is needed
+		const float NumFramesForSmoothOffset = LiveLinkTimedDataInput::CVarLiveLinkNumFramesForSmoothOffset.GetValueOnGameThread();
+
+		if (NumFramesForSmoothOffset <= 0)
+		{
+			Settings->BufferSettings.SmoothEngineTimeOffset = 0.0;
+			return;
+		}
+
+		//In order to compute a frame interval, at least 2 frames are needed
+		const int32 NumFramesInBuffer = FrameTimes.Num();
+		if (NumFramesInBuffer > 1)
+		{
+			//Compute the time interval between the two most recent frames (N+1 - N)
+			const double LatestFrameInterval = FrameTimes[NumFramesInBuffer - 1] - FrameTimes[NumFramesInBuffer - 2];
+
+			//Recover the average frame interval computed last frame
+			const double PreviousAverageFrameInterval = Settings->BufferSettings.SmoothEngineTimeOffset / NumFramesForSmoothOffset;
+
+			//Detect if the new frame interval is very different from the current moving average
+			if (FMath::Abs(LatestFrameInterval - PreviousAverageFrameInterval) > FrameIntervalThreshold)
+			{
+				++FrameIntervalChangeCount;
+
+				if (FrameIntervalChangeCount >= FrameIntervalSnapCount)
+				{
+					//Reduce the number of frames to consider to only those consecutive recent frames that differed so much from the previous average
+					NumFramesToConsiderForAverage = FrameIntervalSnapCount - 1;
+					FrameIntervalChangeCount = 0;
+				}
+			}
+			else
+			{
+				FrameIntervalChangeCount = 0;
+			}
+
+			//Slowly increase the number of frames to consider until it reaches the size of the frame buffer
+			if (NumFramesToConsiderForAverage < NumFramesInBuffer)
+			{
+				++NumFramesToConsiderForAverage;
+			}
+			else
+			{
+				NumFramesToConsiderForAverage = NumFramesInBuffer;
+			}
+
+			check(NumFramesToConsiderForAverage > 1);
+
+			const int32 OldestFrameToConsider = NumFramesInBuffer - NumFramesToConsiderForAverage;
+			const double AverageFrameInterval = (FrameTimes[NumFramesInBuffer - 1] - FrameTimes[OldestFrameToConsider]) / (NumFramesToConsiderForAverage - 1);
+
+			Settings->BufferSettings.SmoothEngineTimeOffset = AverageFrameInterval * NumFramesForSmoothOffset;
+		}
+	}
+}
+
 
 #undef LOCTEXT_NAMESPACE
diff --git a/Engine/Plugins/Animation/LiveLink/Source/LiveLink/Private/LiveLinkTimedDataInput.h b/Engine/Plugins/Animation/LiveLink/Source/LiveLink/Private/LiveLinkTimedDataInput.h
--- a/Engine/Plugins/Animation/LiveLink/Source/LiveLink/Private/LiveLinkTimedDataInput.h
+++ b/Engine/Plugins/Animation/LiveLink/Source/LiveLink/Private/LiveLinkTimedDataInput.h
@@ -6,13 +6,12 @@
 
 #include "CoreMinimal.h"
 #include "ClockOffsetEstimatorRamp.h"
-
+#include "Math/NumericLimits.h"
 
 class FLiveLinkClient;
 enum class ELiveLinkSourceMode : uint8;
 struct FLiveLinkBaseFrameData;
 
-
 class FLiveLinkTimedDataInput : public ITimedDataInput
 {
 public:
@@ -43,6 +42,9 @@ public:
 	//Tracks clock difference between each frame received and arrival time in engine referential.
 	void ProcessNewFrameTimingInfo(FLiveLinkBaseFrameData& NewFrameData);
 
+private:
+	//Computes the smooth offset based on the average frame time intervals of the source time
+	void UpdateSmoothEngineTimeOffset(const FLiveLinkBaseFrameData& NewFrameData);
 
 private:
 	FLiveLinkClient* LiveLinkClient;
@@ -53,7 +55,16 @@ private:
 	FClockOffsetEstimatorRamp EngineClockOffset;
 	FClockOffsetEstimatorRamp TimecodeClockOffset;
 
-	// We will receive each frame for each subject of this source. Stamp last source time/timecode to only update our offset estimation once per "source frame"
+	//We will receive each frame for each subject of this source. Stamp last source time/timecode to only update our offset estimation once per "source frame"
 	double LastWorldSourceTime = 0.0;
 	double LastSceneTime = 0.0;
+
+	static constexpr double FrameIntervalThreshold = 0.005;
+	static constexpr int32 FrameIntervalSnapCount = 5;
+	static constexpr int32 FrameTimeBufferSize = 200;
+
+	TArray<double> FrameTimes;
+
+	int32 FrameIntervalChangeCount = 0;
+	int32 NumFramesToConsiderForAverage = TNumericLimits<int32>::Max();
 };
diff --git a/Engine/Plugins/Animation/LiveLink/Source/LiveLinkEditor/Private/LiveLinkSourceSettingsDetailCustomization.cpp b/Engine/Plugins/Animation/LiveLink/Source/LiveLinkEditor/Private/LiveLinkSourceSettingsDetailCustomization.cpp
--- a/Engine/Plugins/Animation/LiveLink/Source/LiveLinkEditor/Private/LiveLinkSourceSettingsDetailCustomization.cpp
+++ b/Engine/Plugins/Animation/LiveLink/Source/LiveLinkEditor/Private/LiveLinkSourceSettingsDetailCustomization.cpp
@@ -115,6 +115,8 @@ void FLiveLinkSourceSettingsDetailCustomization::CustomizeDetails(IDetailLayoutB
 			, EPropertyLocation::Advanced);
 		CategoryBuilder.AddProperty(BufferSettingsPropertyHandle->GetChildHandle(GET_MEMBER_NAME_CHECKED(FLiveLinkSourceBufferManagementSettings, EngineTimeClockOffset))
 			, EPropertyLocation::Advanced);
+		CategoryBuilder.AddProperty(BufferSettingsPropertyHandle->GetChildHandle(GET_MEMBER_NAME_CHECKED(FLiveLinkSourceBufferManagementSettings, SmoothEngineTimeOffset))
+			, EPropertyLocation::Advanced);
 	}
 	else if (SourceMode == ELiveLinkSourceMode::Latest)
 	{
diff --git a/Engine/Source/Runtime/LiveLinkInterface/Public/ILiveLinkClient.h b/Engine/Source/Runtime/LiveLinkInterface/Public/ILiveLinkClient.h
--- a/Engine/Source/Runtime/LiveLinkInterface/Public/ILiveLinkClient.h
+++ b/Engine/Source/Runtime/LiveLinkInterface/Public/ILiveLinkClient.h
@@ -248,7 +248,7 @@ public:
 	virtual TArray<FLiveLinkSubjectKey> GetSubjectsSupportingRole(TSubclassOf<ULiveLinkRole> SupportedRole, bool bIncludeDisabledSubject, bool bIncludeVirtualSubject) const = 0;
 
 	/**
-	 * Get the time of all the frames for a specific source.
+	 * Get the time of all the frames for a specific subject, including computed offsets.
 	 * @note Use for debugging purposes.
 	 */
 	virtual TArray<FLiveLinkTime> GetSubjectFrameTimes(const FLiveLinkSubjectKey& SubjectKey) const = 0;
@@ -259,7 +259,7 @@ public:
 	virtual ULiveLinkSourceSettings* GetSourceSettings(const FGuid& SourceKey) const = 0;
 
 	/**
-	 * Get the time of all the frames for a source.
+	 * Get the time of all the frames for a specific subject, including computed offsets.
 	 * @note Use for debugging purposes.
 	 */
 	virtual TArray<FLiveLinkTime> GetSubjectFrameTimes(FLiveLinkSubjectName SubjectName) const = 0;
diff --git a/Engine/Source/Runtime/LiveLinkInterface/Public/LiveLinkFrameInterpolationProcessor.h b/Engine/Source/Runtime/LiveLinkInterface/Public/LiveLinkFrameInterpolationProcessor.h
--- a/Engine/Source/Runtime/LiveLinkInterface/Public/LiveLinkFrameInterpolationProcessor.h
+++ b/Engine/Source/Runtime/LiveLinkInterface/Public/LiveLinkFrameInterpolationProcessor.h
@@ -25,6 +25,10 @@ struct FLiveLinkInterpolationInfo
 	/** Distance in seconds between expected evaluation time and oldest sample */
 	float ExpectedEvaluationDistanceFromOldestSeconds;
 
+	/** Frame indices interpolated between */
+	int32 FrameIndexA = INDEX_NONE;
+	int32 FrameIndexB = INDEX_NONE;
+
 	/** Whether sampling was done below our oldest sample */
 	bool bUnderflowDetected = false;
 
diff --git a/Engine/Source/Runtime/LiveLinkInterface/Public/LiveLinkSourceSettings.h b/Engine/Source/Runtime/LiveLinkInterface/Public/LiveLinkSourceSettings.h
--- a/Engine/Source/Runtime/LiveLinkInterface/Public/LiveLinkSourceSettings.h
+++ b/Engine/Source/Runtime/LiveLinkInterface/Public/LiveLinkSourceSettings.h
@@ -55,6 +55,10 @@ struct FLiveLinkSourceBufferManagementSettings
 	UPROPERTY(VisibleAnywhere, Category = "Settings", AdvancedDisplay, meta = (ForceUnits = s))
 	double EngineTimeClockOffset = 0.0;
 
+	/** Continuously updated offset to achieve a smooth evaluation time (in seconds) */
+	UPROPERTY(VisibleAnywhere, Category = "Settings", AdvancedDisplay, meta = (ForceUnits = s))
+	double SmoothEngineTimeOffset = 0.0;
+
 #if WITH_EDITORONLY_DATA
 	/** DEPRECATED: TimecodeFrameRate is now read from each individual subject from FQualifiedFrameTime. 
 	 * It is expected that all subjects under a source have the same and it will be readable in DetectedFrameRate variable
