[LiveLink] Do not update the LiveLink component controller's "component to control" if the existing selection is valid.

diff --git a/Engine/Plugins/Animation/LiveLink/Source/LiveLinkComponents/Private/LiveLinkComponentController.cpp b/Engine/Plugins/Animation/LiveLink/Source/LiveLinkComponents/Private/LiveLinkComponentController.cpp
--- a/Engine/Plugins/Animation/LiveLink/Source/LiveLinkComponents/Private/LiveLinkComponentController.cpp
+++ b/Engine/Plugins/Animation/LiveLink/Source/LiveLinkComponents/Private/LiveLinkComponentController.cpp
@@ -55,6 +55,7 @@ void ULiveLinkComponentController::OnSubjectRoleChanged()
 	}
 	else
 	{
+		TSubclassOf<UActorComponent> DesiredClass = nullptr;
 		UActorComponent* DesiredActorComponent = nullptr;
 		TArray<TSubclassOf<ULiveLinkRole>> SelectedRoleHierarchy = GetSelectedRoleHierarchyClasses(SubjectRepresentation.Role);
 		ControllerMap.Empty(SelectedRoleHierarchy.Num());
@@ -73,7 +74,7 @@ void ULiveLinkComponentController::OnSubjectRoleChanged()
 				{
 					if (AActor* Actor = GetOwner())
 					{
-						TSubclassOf<UActorComponent> DesiredClass = SelectedControllerClass.GetDefaultObject()->GetDesiredComponentClass();
+						DesiredClass = SelectedControllerClass.GetDefaultObject()->GetDesiredComponentClass();
 						if (UActorComponent* ActorComponent = Actor->GetComponentByClass(DesiredClass))
 						{
 							DesiredActorComponent = ActorComponent;
@@ -85,11 +86,15 @@ void ULiveLinkComponentController::OnSubjectRoleChanged()
 
 		//After creating the controller hierarchy, update component to control to the highest in the hierarchy.
 #if WITH_EDITOR
-		if (ComponentToControl.ComponentProperty == NAME_None && DesiredActorComponent != nullptr)
+		if (DesiredActorComponent)
 		{
-			AActor* Actor = GetOwner();
-			check(Actor);
-			ComponentToControl = FComponentEditorUtils::MakeComponentReference(Actor, DesiredActorComponent);
+			UActorComponent* CurrentComponentToControl = ComponentToControl.GetComponent(GetOwner());
+			if ((CurrentComponentToControl == nullptr) || !CurrentComponentToControl->IsA(DesiredClass))
+			{
+				AActor* Actor = GetOwner();
+				check(Actor);
+				ComponentToControl = FComponentEditorUtils::MakeComponentReference(Actor, DesiredActorComponent);
+			}
 		}
 #endif
 	}
@@ -134,11 +139,15 @@ void ULiveLinkComponentController::SetControllerClassForRole(TSubclassOf<ULiveLi
 				if (RoleClass == SubjectRepresentation.Role)
 				{
 					TSubclassOf<UActorComponent> DesiredComponent = CurrentController->GetDesiredComponentClass();
-					if (AActor* Actor = GetOwner())
+					UActorComponent* CurrentComponentToControl = ComponentToControl.GetComponent(GetOwner());
+					if ((CurrentComponentToControl == nullptr) || !CurrentComponentToControl->IsA(DesiredComponent))
 					{
-						if (UActorComponent* ActorComponent = Actor->GetComponentByClass(DesiredComponent))
+						if (AActor* Actor = GetOwner())
 						{
-							ComponentToControl = FComponentEditorUtils::MakeComponentReference(Actor, ActorComponent);
+							if (UActorComponent* ActorComponent = Actor->GetComponentByClass(DesiredComponent))
+							{
+								ComponentToControl = FComponentEditorUtils::MakeComponentReference(Actor, ActorComponent);
+							}
 						}
 					}
 				}
