- Adding cvar to easily disable all LiveLink component evaluation of their subject to prevent path tracing invalidation

diff --git a/Engine/Plugins/Animation/LiveLink/Source/LiveLinkComponents/Private/LiveLinkComponentController.cpp b/Engine/Plugins/Animation/LiveLink/Source/LiveLinkComponents/Private/LiveLinkComponentController.cpp
--- a/Engine/Plugins/Animation/LiveLink/Source/LiveLinkComponents/Private/LiveLinkComponentController.cpp
+++ b/Engine/Plugins/Animation/LiveLink/Source/LiveLinkComponents/Private/LiveLinkComponentController.cpp
@@ -9,6 +9,7 @@
 #include "Engine/World.h"
 #include "Features/IModularFeatures.h"
 #include "GameFramework/Actor.h"
+#include "HAL/IConsoleManager.h"
 #include "Logging/LogMacros.h"
 #include "UObject/EnterpriseObjectVersion.h"
 #include "UObject/UObjectIterator.h"
@@ -23,6 +24,14 @@
 #define LOCTEXT_NAMESPACE "LiveLinkController"
 
 
+static TAutoConsoleVariable<bool> CVarEnableLiveLinkEvaluation(
+	TEXT("LiveLink.Component.EnableLiveLinkEvaluation"),
+	true,
+	TEXT("Whether LiveLink components should evaluate their subject."),
+	ECVF_Default);
+
+
+
 ULiveLinkComponentController::ULiveLinkComponentController()
 	: bUpdateInEditor(true)
 	, bIsDirty(false)
@@ -214,10 +223,12 @@ void ULiveLinkComponentController::TickComponent(float DeltaTime, ELevelTick Tic
 
 	ILiveLinkClient& LiveLinkClient = IModularFeatures::Get().GetModularFeature<ILiveLinkClient>(ILiveLinkClient::ModularFeatureName);
 
-	//Evaluate subject frame once and pass the data to our controllers
+	// Evaluate subject frame once and pass the data to our controllers
 	FLiveLinkSubjectFrameData SubjectData;
 
-	const bool bHasValidData = bEvaluateLiveLink ? LiveLinkClient.EvaluateFrame_AnyThread(SubjectRepresentation.Subject, SubjectRepresentation.Role, SubjectData) : false;
+	// Verify if global evaluation cvar is on or off. This can be used to stop LL evaluation at large for Pathtracing rendering for example
+	const bool bCanEvaluate = bEvaluateLiveLink && CVarEnableLiveLinkEvaluation.GetValueOnGameThread();
+	const bool bHasValidData = bCanEvaluate ? LiveLinkClient.EvaluateFrame_AnyThread(SubjectRepresentation.Subject, SubjectRepresentation.Role, SubjectData) : false;
 
 	//Go through each controllers and initialize them if we're dirty and tick them if there's valid data to process
 	for (auto& ControllerEntry : ControllerMap)
