Exposed ILiveLinkClient::GetSourceType to blueprints. Name changed to GetSourceTypeFromGuid to avoid conflict with existing GetSourceType function using a Source Handle

diff --git a/Engine/Plugins/Animation/LiveLink/Source/LiveLink/Private/LiveLinkBlueprintLibrary.cpp b/Engine/Plugins/Animation/LiveLink/Source/LiveLink/Private/LiveLinkBlueprintLibrary.cpp
--- a/Engine/Plugins/Animation/LiveLink/Source/LiveLink/Private/LiveLinkBlueprintLibrary.cpp
+++ b/Engine/Plugins/Animation/LiveLink/Source/LiveLink/Private/LiveLinkBlueprintLibrary.cpp
@@ -172,6 +172,18 @@ FText ULiveLinkBlueprintLibrary::GetSourceMachineName(UPARAM(ref) FLiveLinkSourc
 	}
 }
 
+FText ULiveLinkBlueprintLibrary::GetSourceTypeFromGuid(FGuid SourceGuid)
+{
+	IModularFeatures& ModularFeatures = IModularFeatures::Get();
+	if (ModularFeatures.IsModularFeatureAvailable(ILiveLinkClient::ModularFeatureName))
+	{
+		const ILiveLinkClient& LiveLinkClient = ModularFeatures.GetModularFeature<ILiveLinkClient>(ILiveLinkClient::ModularFeatureName);
+
+		return LiveLinkClient.GetSourceType(SourceGuid);
+	}
+	return FText::GetEmpty();
+}
+
 TArray<FLiveLinkSubjectName> ULiveLinkBlueprintLibrary::GetLiveLinkEnabledSubjectNames(bool bIncludeVirtualSubject)
 {
 	TArray<FLiveLinkSubjectName> Result;
diff --git a/Engine/Plugins/Animation/LiveLink/Source/LiveLink/Public/LiveLinkBlueprintLibrary.h b/Engine/Plugins/Animation/LiveLink/Source/LiveLink/Public/LiveLinkBlueprintLibrary.h
--- a/Engine/Plugins/Animation/LiveLink/Source/LiveLink/Public/LiveLinkBlueprintLibrary.h
+++ b/Engine/Plugins/Animation/LiveLink/Source/LiveLink/Public/LiveLinkBlueprintLibrary.h
@@ -117,7 +117,14 @@ class LIVELINK_API ULiveLinkBlueprintLibrary : public UBlueprintFunctionLibrary
 	static FText GetSourceMachineName(UPARAM(ref) FLiveLinkSourceHandle& SourceHandle);
 
 public:
-
+	/**
+	 * Get the type of a source from the given GUID
+	 * @param SourceGuid the GUID identifying the LiveLink Source
+	 * @return The type of the Source as Text
+	 */
+	UFUNCTION(BlueprintCallable, Category = "LiveLink")
+	static FText GetSourceTypeFromGuid(FGuid SourceGuid);
+	
 	/** Get a list of all enabled subject names */
 	UFUNCTION(BlueprintCallable, Category = "LiveLink")
 	static TArray<FLiveLinkSubjectName> GetLiveLinkEnabledSubjectNames(bool bIncludeVirtualSubject);
