Fix crash when trying to load LL preset that fails to load

diff --git a/Engine/Plugins/Animation/LiveLink/Source/LiveLinkEditor/Private/LiveLinkClientPanelToolbar.cpp b/Engine/Plugins/Animation/LiveLink/Source/LiveLinkEditor/Private/LiveLinkClientPanelToolbar.cpp
--- a/Engine/Plugins/Animation/LiveLink/Source/LiveLinkEditor/Private/LiveLinkClientPanelToolbar.cpp
+++ b/Engine/Plugins/Animation/LiveLink/Source/LiveLinkEditor/Private/LiveLinkClientPanelToolbar.cpp
@@ -31,6 +31,7 @@
 #include "Framework/Application/MenuStack.h"
 #include "Framework/Application/SlateApplication.h"
 #include "Framework/MultiBox/MultiBoxBuilder.h"
+#include "Framework/Notifications/NotificationManager.h"
 #include "Framework/SlateDelegates.h"
 #include "Styling/SlateIconFinder.h"
 #include "Widgets/Images/SImage.h"
@@ -44,6 +45,7 @@
 #include "Widgets/Layout/SSpacer.h"
 #include "Widgets/Layout/SUniformGridPanel.h"
 #include "Widgets/Notifications/SErrorText.h"
+#include "Widgets/Notifications/SNotificationList.h"
 #include "Widgets/Text/STextBlock.h"
 #include "Widgets/SBoxPanel.h"
 #include "Widgets/SNullWidget.h"
@@ -880,7 +882,19 @@ void SLiveLinkClientPanelToolbar::OnImportPreset(const FAssetData& InPreset)
 {
 	FSlateApplication::Get().DismissAllMenus();
 
-	ULiveLinkPreset* ImportedPreset = CastChecked<ULiveLinkPreset>(InPreset.GetAsset());
+	UObject* PresetAssetData = InPreset.GetAsset();
+	if (!PresetAssetData)
+	{
+		FNotificationInfo Info(LOCTEXT("LoadPresetFailed", "Failed to load preset"));
+		Info.ExpireDuration = 5.0f;
+		Info.Hyperlink = FSimpleDelegate::CreateStatic([]() { FMessageLog("LoadErrors").Open(EMessageSeverity::Info, true); });
+		Info.HyperlinkText = LOCTEXT("LoadObjectHyperlink", "Show Message Log");
+
+		FSlateNotificationManager::Get().AddNotification(Info);
+		return;
+	}
+
+	ULiveLinkPreset* ImportedPreset = Cast<ULiveLinkPreset>(PresetAssetData);
 	if (ImportedPreset)
 	{
 		FScopedTransaction Transaction(LOCTEXT("ImportPreset_Transaction", "Import LiveLink Preset"));
