Ensure that we transact modifications to the live link sources.

diff --git a/Engine/Plugins/Animation/LiveLink/Source/LiveLinkSequencer/Private/LiveLinkSequencerModule.cpp b/Engine/Plugins/Animation/LiveLink/Source/LiveLinkSequencer/Private/LiveLinkSequencerModule.cpp
--- a/Engine/Plugins/Animation/LiveLink/Source/LiveLinkSequencer/Private/LiveLinkSequencerModule.cpp
+++ b/Engine/Plugins/Animation/LiveLink/Source/LiveLinkSequencer/Private/LiveLinkSequencerModule.cpp
@@ -182,6 +182,10 @@ private:
 
 	static void AddLiveLinkSource(UTakeRecorderSources* Sources,  const FName& SubjectName)
 	{
+		FScopedTransaction Transaction(LOCTEXT("AddLiveLinkSource","Add Live Link Source"));
+
+		Sources->Modify();
+
 		UTakeRecorderLiveLinkSource* NewSource = Sources->AddSource<UTakeRecorderLiveLinkSource>();
 		NewSource->SubjectName = SubjectName;
 	}
