Fix smooth engine time offset to account for very large frame time intervals

diff --git a/Engine/Plugins/Animation/LiveLink/Source/LiveLink/Private/LiveLinkTimedDataInput.cpp b/Engine/Plugins/Animation/LiveLink/Source/LiveLink/Private/LiveLinkTimedDataInput.cpp
--- a/Engine/Plugins/Animation/LiveLink/Source/LiveLink/Private/LiveLinkTimedDataInput.cpp
+++ b/Engine/Plugins/Animation/LiveLink/Source/LiveLink/Private/LiveLinkTimedDataInput.cpp
@@ -275,6 +275,15 @@ void FLiveLinkTimedDataInput::UpdateSmoothEngineTimeOffset(const FLiveLinkBaseFr
 			//Compute the time interval between the two most recent frames (N+1 - N)
 			const double LatestFrameInterval = FrameTimes[NumFramesInBuffer - 1] - FrameTimes[NumFramesInBuffer - 2];
 
+			// Early-out if a very large amount of time has passed since the last frame was received (indicating that the device was previously idle)
+			if (LatestFrameInterval > VeryLargeFrameIntervalThreshold)
+			{
+				FrameTimes.RemoveAt(0, NumFramesInBuffer - 1, false);
+				Settings->BufferSettings.SmoothEngineTimeOffset = 0.0;
+				NumFramesToConsiderForAverage = 1;
+				return;
+			}
+
 			//Recover the average frame interval computed last frame
 			const double PreviousAverageFrameInterval = Settings->BufferSettings.SmoothEngineTimeOffset / NumFramesForSmoothOffset;
 
diff --git a/Engine/Plugins/Animation/LiveLink/Source/LiveLink/Private/LiveLinkTimedDataInput.h b/Engine/Plugins/Animation/LiveLink/Source/LiveLink/Private/LiveLinkTimedDataInput.h
--- a/Engine/Plugins/Animation/LiveLink/Source/LiveLink/Private/LiveLinkTimedDataInput.h
+++ b/Engine/Plugins/Animation/LiveLink/Source/LiveLink/Private/LiveLinkTimedDataInput.h
@@ -60,6 +60,7 @@ private:
 	double LastSceneTime = 0.0;
 
 	static constexpr double FrameIntervalThreshold = 0.005;
+	static constexpr double VeryLargeFrameIntervalThreshold = 0.5;
 	static constexpr int32 FrameIntervalSnapCount = 5;
 	static constexpr int32 FrameTimeBufferSize = 200;
 
