Refactor the LiveLink component to allow each LiveLink controller in its controller map to specifiy a component to control
Deprecate the "ComponentToControl" property of the LiveLinkComponentController
Add a customization for LiveLinkControllerBase that ensures the Component picker is displayed first in the details panel
Fix the fallout from the deprecation in USDImporter plugin

diff --git a/Engine/Plugins/Animation/LiveLink/Source/LiveLinkComponents/Private/Controllers/LiveLinkLightController.cpp b/Engine/Plugins/Animation/LiveLink/Source/LiveLinkComponents/Private/Controllers/LiveLinkLightController.cpp
--- a/Engine/Plugins/Animation/LiveLink/Source/LiveLinkComponents/Private/Controllers/LiveLinkLightController.cpp
+++ b/Engine/Plugins/Animation/LiveLink/Source/LiveLinkComponents/Private/Controllers/LiveLinkLightController.cpp
@@ -24,7 +24,7 @@ void ULiveLinkLightController::Tick(float DeltaTime, const FLiveLinkSubjectFrame
 
 	if (StaticData && FrameData)
 	{
-		if (ULightComponent* LightComponent = Cast<ULightComponent>(AttachedComponent))
+		if (ULightComponent* LightComponent = Cast<ULightComponent>(GetAttachedComponent()))
 		{
 			if (StaticData->bIsTemperatureSupported) { LightComponent->SetTemperature(FrameData->Temperature); }
 			if (StaticData->bIsIntensitySupported) { LightComponent->SetIntensity(FrameData->Intensity); }
@@ -88,10 +88,10 @@ void ULiveLinkLightController::PostLoad()
 				}
 
 				//if Subjects role direct controller is us, set the component to control to what we had
-				if (LiveLinkComponent->SubjectRepresentation.Role == ULiveLinkLightRole::StaticClass())
-				{
-					LiveLinkComponent->ComponentToControl = ComponentToControl_DEPRECATED;
-				}
+ 				if (LiveLinkComponent->SubjectRepresentation.Role == ULiveLinkLightRole::StaticClass())
+ 				{
+ 					ComponentPicker = ComponentToControl_DEPRECATED;
+ 				}
 			}
 		}
 	}
diff --git a/Engine/Plugins/Animation/LiveLink/Source/LiveLinkComponents/Private/Controllers/LiveLinkTransformController.cpp b/Engine/Plugins/Animation/LiveLink/Source/LiveLinkComponents/Private/Controllers/LiveLinkTransformController.cpp
--- a/Engine/Plugins/Animation/LiveLink/Source/LiveLinkComponents/Private/Controllers/LiveLinkTransformController.cpp
+++ b/Engine/Plugins/Animation/LiveLink/Source/LiveLinkComponents/Private/Controllers/LiveLinkTransformController.cpp
@@ -88,7 +88,7 @@ void FLiveLinkTransformControllerData::CheckForError(FName OwnerName, USceneComp
 void ULiveLinkTransformController::OnEvaluateRegistered()
 {
 	AActor* OuterActor = GetOuterActor();
-	TransformData.CheckForError(OuterActor ? OuterActor->GetFName() : NAME_None, Cast<USceneComponent>(AttachedComponent));
+	TransformData.CheckForError(OuterActor ? OuterActor->GetFName() : NAME_None, Cast<USceneComponent>(GetAttachedComponent()));
 }
 
 void ULiveLinkTransformController::Tick(float DeltaTime, const FLiveLinkSubjectFrameData& SubjectData)
@@ -98,7 +98,7 @@ void ULiveLinkTransformController::Tick(float DeltaTime, const FLiveLinkSubjectF
 
 	if (StaticData && FrameData)
 	{
-		if (USceneComponent* SceneComponent = Cast<USceneComponent>(AttachedComponent))
+		if (USceneComponent* SceneComponent = Cast<USceneComponent>(GetAttachedComponent()))
 		{
 			TransformData.ApplyTransform(SceneComponent, FrameData->Transform, *StaticData);
 		}
@@ -120,7 +120,7 @@ void ULiveLinkTransformController::SetAttachedComponent(UActorComponent* ActorCo
 	Super::SetAttachedComponent(ActorComponent);
 
 	AActor* OuterActor = GetOuterActor();
-	TransformData.CheckForError(OuterActor ? OuterActor->GetFName() : NAME_None, Cast<USceneComponent>(AttachedComponent));
+	TransformData.CheckForError(OuterActor ? OuterActor->GetFName() : NAME_None, Cast<USceneComponent>(GetAttachedComponent()));
 }
 
 void ULiveLinkTransformController::PostLoad()
@@ -145,7 +145,7 @@ void ULiveLinkTransformController::PostLoad()
 				//if Subjects role direct controller is us, set the component to control to what we had
 				if (LiveLinkComponent->SubjectRepresentation.Role == ULiveLinkTransformRole::StaticClass())
 				{
-					LiveLinkComponent->ComponentToControl = ComponentToControl_DEPRECATED;
+					ComponentPicker = ComponentToControl_DEPRECATED;
 				}
 			}
 		}
diff --git a/Engine/Plugins/Animation/LiveLink/Source/LiveLinkComponents/Private/LiveLinkComponentController.cpp b/Engine/Plugins/Animation/LiveLink/Source/LiveLinkComponents/Private/LiveLinkComponentController.cpp
--- a/Engine/Plugins/Animation/LiveLink/Source/LiveLinkComponents/Private/LiveLinkComponentController.cpp
+++ b/Engine/Plugins/Animation/LiveLink/Source/LiveLinkComponents/Private/LiveLinkComponentController.cpp
@@ -12,13 +12,12 @@
 #include "HAL/IConsoleManager.h"
 #include "Logging/LogMacros.h"
 #include "UObject/EnterpriseObjectVersion.h"
+#include "UObject/UE5MainStreamObjectVersion.h"
 #include "UObject/UObjectIterator.h"
 
 #if WITH_EDITOR
 #include "Editor.h"
-#include "Framework/Notifications/NotificationManager.h"
 #include "Kismet2/ComponentEditorUtils.h"
-#include "Widgets/Notifications/SNotificationList.h"
 #endif // WITH_EDITOR
 
 #define LOCTEXT_NAMESPACE "LiveLinkController"
@@ -46,12 +45,14 @@ ULiveLinkComponentController::ULiveLinkComponentController()
 #endif //WITH_EDITOR
 }
 
+PRAGMA_DISABLE_DEPRECATION_WARNINGS
 ULiveLinkComponentController::~ULiveLinkComponentController()
 {
 #if WITH_EDITOR
 	FEditorDelegates::EndPIE.RemoveAll(this);
 #endif //WITH_EDITOR
 }
+PRAGMA_ENABLE_DEPRECATION_WARNINGS
 
 void ULiveLinkComponentController::OnSubjectRoleChanged()
 {
@@ -64,8 +65,6 @@ void ULiveLinkComponentController::OnSubjectRoleChanged()
 	}
 	else
 	{
-		TSubclassOf<UActorComponent> DesiredClass = nullptr;
-		UActorComponent* DesiredActorComponent = nullptr;
 		TArray<TSubclassOf<ULiveLinkRole>> SelectedRoleHierarchy = GetSelectedRoleHierarchyClasses(SubjectRepresentation.Role);
 		ControllerMap.Empty(SelectedRoleHierarchy.Num());
 		for (const TSubclassOf<ULiveLinkRole>& RoleClass : SelectedRoleHierarchy)
@@ -77,35 +76,8 @@ void ULiveLinkComponentController::OnSubjectRoleChanged()
 
 				TSubclassOf<ULiveLinkControllerBase> SelectedControllerClass = GetControllerClassForRoleClass(RoleClass);
 				SetControllerClassForRole(RoleClass, SelectedControllerClass);
-
-				//Keep track of the most specific available component in the hierarchy
-				if (SelectedControllerClass)
-				{
-					if (AActor* Actor = GetOwner())
-					{
-						DesiredClass = SelectedControllerClass.GetDefaultObject()->GetDesiredComponentClass();
-						if (UActorComponent* ActorComponent = Actor->GetComponentByClass(DesiredClass))
-						{
-							DesiredActorComponent = ActorComponent;
-						}
-					}
-				}
-			}
-		}
-
-		//After creating the controller hierarchy, update component to control to the highest in the hierarchy.
-#if WITH_EDITOR
-		if (DesiredActorComponent)
-		{
-			UActorComponent* CurrentComponentToControl = ComponentToControl.GetComponent(GetOwner());
-			if ((CurrentComponentToControl == nullptr) || !CurrentComponentToControl->IsA(DesiredClass))
-			{
-				AActor* Actor = GetOwner();
-				check(Actor);
-				ComponentToControl = FComponentEditorUtils::MakeComponentReference(Actor, DesiredActorComponent);
 			}
 		}
-#endif
 	}
 
 	if (OnControllerMapUpdatedDelegate.IsBound())
@@ -142,25 +114,9 @@ void ULiveLinkComponentController::SetControllerClassForRole(TSubclassOf<ULiveLi
 			{
 				const EObjectFlags ControllerObjectFlags = GetMaskedFlags(RF_Public | RF_Transactional | RF_ArchetypeObject);
 				CurrentController = NewObject<ULiveLinkControllerBase>(this, DesiredControllerClass, NAME_None, ControllerObjectFlags);
+				InitializeController(CurrentController);
 
-#if WITH_EDITOR
-				//For the controller directly associated with the subject role, set the component to control to the desired component this controller wants
-				if (RoleClass == SubjectRepresentation.Role)
-				{
-					TSubclassOf<UActorComponent> DesiredComponent = CurrentController->GetDesiredComponentClass();
-					UActorComponent* CurrentComponentToControl = ComponentToControl.GetComponent(GetOwner());
-					if ((CurrentComponentToControl == nullptr) || !CurrentComponentToControl->IsA(DesiredComponent))
-					{
-						if (AActor* Actor = GetOwner())
-						{
-							if (UActorComponent* ActorComponent = Actor->GetComponentByClass(DesiredComponent))
-							{
-								ComponentToControl = FComponentEditorUtils::MakeComponentReference(Actor, ActorComponent);
-							}
-						}
-					}
-				}
-				
+#if WITH_EDITOR		
 				CurrentController->InitializeInEditor();
 #endif
 			}
@@ -242,7 +198,6 @@ void ULiveLinkComponentController::TickComponent(float DeltaTime, ELevelTick Tic
 		{
 			if (bIsDirty)
 			{
-				Controller->SetAttachedComponent(ComponentToControl.GetComponent(GetOwner()));
 				Controller->SetSelectedSubject(SubjectRepresentation);
 				Controller->OnEvaluateRegistered();
 			}
@@ -283,6 +238,30 @@ void ULiveLinkComponentController::Serialize(FArchive& Ar)
 #endif
 }
 
+void ULiveLinkComponentController::PostLoad()
+{
+	Super::PostLoad();
+
+#if WITH_EDITOR
+ 	const int32 Version = GetLinkerCustomVersion(FUE5MainStreamObjectVersion::GUID);
+ 	if (Version < FUE5MainStreamObjectVersion::LiveLinkComponentPickerPerController)
+ 	{
+		for (auto& ControllerEntry : ControllerMap)
+		{
+			ULiveLinkControllerBase* Controller = ControllerEntry.Value;
+			if (Controller)
+			{
+				Controller->ConditionalPostLoad();
+				PRAGMA_DISABLE_DEPRECATION_WARNINGS
+				UActorComponent* OldAttachedComponent = ComponentToControl_DEPRECATED.GetComponent(GetOwner());
+				Controller->SetAttachedComponent(OldAttachedComponent);
+				PRAGMA_ENABLE_DEPRECATION_WARNINGS
+			}
+		}
+	}
+#endif //WITH_EDITOR
+}
+
 #if WITH_EDITOR
 
 void ULiveLinkComponentController::PostEditChangeProperty(struct FPropertyChangedEvent& PropertyChangedEvent)
@@ -404,4 +383,59 @@ void ULiveLinkComponentController::CleanupControllersInMap()
 	}
 }
 
+void ULiveLinkComponentController::InitializeController(ULiveLinkControllerBase* InController)
+{
+#if WITH_EDITOR
+	// If the OuterActor has a component that matches the desired class of this controller, set that as the component to control. 
+	// Otherwise, the default root component will set as the component to control.
+	if (AActor* OuterActor = GetOwner())
+	{
+		TInlineComponentArray<UActorComponent*> ActorComponents;
+		OuterActor->GetComponents(InController->GetDesiredComponentClass(), ActorComponents);
+
+		bool bFoundValidComponent = false;
+
+		// Look through the list of components matching the desired component class, and choose the first editable instance
+		for (UActorComponent* ActorComponent : ActorComponents)
+		{
+			// Check that the selected component is editable (and thus appropriate to be driven by a LiveLink controller)
+			if (FComponentEditorUtils::CanEditComponentInstance(ActorComponent, Cast<USceneComponent>(ActorComponent), false))
+			{
+				InController->SetAttachedComponent(ActorComponent);
+				bFoundValidComponent = true;
+				break;
+			}
+		}
+
+		if (!bFoundValidComponent)
+		{
+			UE_LOG(LogLiveLinkComponents, Warning, TEXT("The desired component class for %s is %s, but %s does not have a component of that type."), *InController->GetName(), *InController->GetDesiredComponentClass()->GetName(), *OuterActor->GetActorLabel(false));
+		}
+	}
+#endif //WITH_EDITOR
+}
+
+UActorComponent* ULiveLinkComponentController::GetControlledComponent(TSubclassOf<ULiveLinkRole> InRoleClass) const
+{
+	if (const TObjectPtr<ULiveLinkControllerBase>* ControllerPtr = ControllerMap.Find(InRoleClass))
+	{	
+		if (const TObjectPtr<ULiveLinkControllerBase> Controller = *ControllerPtr)
+		{
+			return Controller->GetAttachedComponent();
+		}
+	}
+	return nullptr;
+}
+
+void ULiveLinkComponentController::SetControlledComponent(TSubclassOf<ULiveLinkRole> InRoleClass, UActorComponent* InComponent)
+{
+	if (TObjectPtr<ULiveLinkControllerBase>* ControllerPtr = ControllerMap.Find(InRoleClass))
+	{
+		if (TObjectPtr<ULiveLinkControllerBase> Controller = *ControllerPtr)
+		{
+			Controller->SetAttachedComponent(InComponent);
+		}
+	}
+}
+
 #undef LOCTEXT_NAMESPACE
\ No newline at end of file
diff --git a/Engine/Plugins/Animation/LiveLink/Source/LiveLinkComponents/Private/LiveLinkControllerBase.cpp b/Engine/Plugins/Animation/LiveLink/Source/LiveLinkComponents/Private/LiveLinkControllerBase.cpp
--- a/Engine/Plugins/Animation/LiveLink/Source/LiveLinkComponents/Private/LiveLinkControllerBase.cpp
+++ b/Engine/Plugins/Animation/LiveLink/Source/LiveLinkComponents/Private/LiveLinkControllerBase.cpp
@@ -4,9 +4,22 @@
 #include "GameFramework/Actor.h"
 #include "UObject/UObjectIterator.h"
 
+#if WITH_EDITOR
+#include "Kismet2/ComponentEditorUtils.h"
+#endif // WITH_EDITOR
+
 void ULiveLinkControllerBase::SetAttachedComponent(UActorComponent* ActorComponent)
 {
+#if WITH_EDITOR
+	if (GetAttachedComponent() != ActorComponent)
+	{
+		ComponentPicker = FComponentEditorUtils::MakeComponentReference(GetOuterActor(), ActorComponent);
+	}
+#endif // WITH_EDITOR
+
+PRAGMA_DISABLE_DEPRECATION_WARNINGS
 	AttachedComponent = ActorComponent;
+PRAGMA_ENABLE_DEPRECATION_WARNINGS
 }
 
 void ULiveLinkControllerBase::SetSelectedSubject(FLiveLinkSubjectRepresentation LiveLinkSubject)
@@ -60,3 +73,23 @@ TArray<TSubclassOf<ULiveLinkControllerBase>> ULiveLinkControllerBase::GetControl
 	return MoveTemp(Controllers);
 }
 
+UActorComponent* ULiveLinkControllerBase::GetAttachedComponent() const
+{
+	return ComponentPicker.GetComponent(GetOuterActor());
+}
+
+void ULiveLinkControllerBase::OnComponentToControlChanged()
+{
+PRAGMA_DISABLE_DEPRECATION_WARNINGS
+	AttachedComponent = GetAttachedComponent();
+PRAGMA_ENABLE_DEPRECATION_WARNINGS
+}
+
+void ULiveLinkControllerBase::PostLoad()
+{
+	Super::PostLoad();
+
+PRAGMA_DISABLE_DEPRECATION_WARNINGS
+	AttachedComponent = GetAttachedComponent();
+PRAGMA_ENABLE_DEPRECATION_WARNINGS
+}
\ No newline at end of file
diff --git a/Engine/Plugins/Animation/LiveLink/Source/LiveLinkComponents/Public/LiveLinkComponentController.h b/Engine/Plugins/Animation/LiveLink/Source/LiveLinkComponents/Public/LiveLinkComponentController.h
--- a/Engine/Plugins/Animation/LiveLink/Source/LiveLinkComponents/Public/LiveLinkComponentController.h
+++ b/Engine/Plugins/Animation/LiveLink/Source/LiveLinkComponents/Public/LiveLinkComponentController.h
@@ -44,8 +44,11 @@ public:
 	UPROPERTY(BlueprintAssignable, Category = "LiveLink")
 	FOnControllerMapUpdatedDelegate OnControllerMapUpdatedDelegate;
 
-	UPROPERTY(EditInstanceOnly, Category = "LiveLink", meta = (UseComponentPicker, AllowedClasses = "ActorComponent", DisallowedClasses = "LiveLinkComponentController"))
-	FComponentReference ComponentToControl;
+#if WITH_EDITORONLY_DATA
+	UE_DEPRECATED(5.1, "This property has been deprecated. Please use the ComponentPicker property of each controller in this LiveLink component's controller map.")
+	UPROPERTY()
+	FComponentReference ComponentToControl_DEPRECATED;
+#endif //WITH_EDITORONLY_DATA
 
 	// If true, will not evaluate LiveLink if the attached actor is a spawnable in Sequencer
 	UPROPERTY(EditAnywhere, BlueprintReadWrite, Category = "LiveLink")
@@ -85,6 +88,12 @@ public:
 	/** Used to notify that the subject role has changed. Mainly from Customization or C++ modification to the subject's Role */
 	void OnSubjectRoleChanged();
 
+	/** Returns the component controlled by the LiveLink controller of the input Role. Returns null if there is no controller for that Role */
+	UActorComponent* GetControlledComponent(TSubclassOf<ULiveLinkRole> InRoleClass) const;
+
+	/** Set the component to control for the LiveLink controller of the input Role */
+	void SetControlledComponent(TSubclassOf<ULiveLinkRole> InRoleClass, UActorComponent* InComponent);
+
 #if WITH_EDITOR
 	/** Used to cleanup controllers when exiting PIE */
 	void OnEndPIE(bool bIsSimulating);
@@ -97,6 +106,7 @@ public:
 	//~ End UActorComponent interface
 
 	//~ UObject interface
+	virtual void PostLoad() override;
 	virtual void Serialize(FArchive& Ar) override;
 #if WITH_EDITOR
 	virtual void PostEditChangeProperty(struct FPropertyChangedEvent& PropertyChangedEvent) override;
@@ -112,6 +122,9 @@ protected:
 	/** Loops through the controller map and calls Cleanup() on each entry */
 	void CleanupControllersInMap();
 
+	/** Initializes the component that the newly created input controller should control based on its specified desired component class */
+	void InitializeController(ULiveLinkControllerBase* InController);
+
 #if WITH_EDITOR
 	/** Called during loading to convert old data to new scheme. */
 	void ConvertOldControllerSystem();
diff --git a/Engine/Plugins/Animation/LiveLink/Source/LiveLinkComponents/Public/LiveLinkControllerBase.h b/Engine/Plugins/Animation/LiveLink/Source/LiveLinkComponents/Public/LiveLinkControllerBase.h
--- a/Engine/Plugins/Animation/LiveLink/Source/LiveLinkComponents/Public/LiveLinkControllerBase.h
+++ b/Engine/Plugins/Animation/LiveLink/Source/LiveLinkComponents/Public/LiveLinkControllerBase.h
@@ -16,12 +16,16 @@ class AActor;
 
 /**
  */
-UCLASS(Abstract, ClassGroup=(LiveLink), editinlinenew)
+UCLASS(Abstract, ClassGroup = (LiveLink), editinlinenew)
 class LIVELINKCOMPONENTS_API ULiveLinkControllerBase : public UObject
 {
 	GENERATED_BODY()
 
 public:
+	//~ Begin UObject interface
+	virtual void PostLoad() override;
+	//~ End UObject interface
+
 	/** Initialize the controller at the first tick of his owner component. */
 	virtual void OnEvaluateRegistered() { }
 
@@ -54,7 +58,6 @@ public:
 	 */
 	virtual void SetSelectedSubject(FLiveLinkSubjectRepresentation LiveLinkSubject);
 
-
 	/**
 	 * Cleanup controller state before getting removed
 	 */
@@ -67,6 +70,16 @@ public:
 	/** Get the selected LiveLink subject for this controller */
 	virtual FLiveLinkSubjectRepresentation GetSelectedSubject() { return SelectedSubject; }
 
+	/**
+	 * Returns the component controlled by this controller
+	 */
+	UActorComponent* GetAttachedComponent() const;
+
+	/**
+	 * Callback to reset the AttachedComponent when the value of the ComponentPicker is changed
+	 */
+	void OnComponentToControlChanged();
+
 protected:
 	AActor* GetOuterActor() const;
 
@@ -81,6 +94,14 @@ public:
 	static TArray<TSubclassOf<ULiveLinkControllerBase>> GetControllersForRole(const TSubclassOf<ULiveLinkRole>& RoleToSupport);
 
 protected:
+	/**
+	 * A component reference (customized) that allows the user to specify a component that this controller should control
+	 */
+	UPROPERTY(EditInstanceOnly, Category = "ComponentToControl", meta = (UseComponentPicker, AllowedClasses = "ActorComponent", DisallowedClasses = "LiveLinkComponentController"))
+	FComponentReference ComponentPicker;
+
+protected:
+	UE_DEPRECATED(5.1, "This property has been deprecated. Please use GetAttachedComponent() instead.")
 	TWeakObjectPtr<UActorComponent> AttachedComponent;
 
 	FLiveLinkSubjectRepresentation SelectedSubject;
diff --git a/Engine/Plugins/Animation/LiveLink/Source/LiveLinkEditor/Private/LiveLinkComponentDetailCustomization.cpp b/Engine/Plugins/Animation/LiveLink/Source/LiveLinkEditor/Private/LiveLinkComponentDetailCustomization.cpp
--- a/Engine/Plugins/Animation/LiveLink/Source/LiveLinkEditor/Private/LiveLinkComponentDetailCustomization.cpp
+++ b/Engine/Plugins/Animation/LiveLink/Source/LiveLinkEditor/Private/LiveLinkComponentDetailCustomization.cpp
@@ -12,6 +12,7 @@
 #include "LiveLinkEditorPrivate.h"
 #include "ScopedTransaction.h"
 
+#include "Widgets/Images/SImage.h"
 #include "Widgets/Input/SComboButton.h"
 
 #define LOCTEXT_NAMESPACE "LiveLinkComponentDetailsCustomization"
@@ -40,10 +41,8 @@ void FLiveLinkComponentDetailCustomization::CustomizeDetails(IDetailLayoutBuilde
 
 	if (ULiveLinkComponentController* SelectedPtr = Cast<ULiveLinkComponentController>(SelectedObjects[0].Get()))
 	{
-		//Register callback when LiveLinkSubjectRepresentation selection has changed to refresh the UI and update controller
-		TSharedRef<IPropertyHandle> SubjectRepresentationPropertyRef = DetailBuilder.GetProperty(GET_MEMBER_NAME_CHECKED(ULiveLinkComponentController, SubjectRepresentation));
-		SubjectRepresentationPropertyRef->MarkHiddenByCustomization();
-		DetailLayout->AddPropertyToCategory(SubjectRepresentationPropertyRef);
+		// Force "LiveLink" to be the topmost category in the details panel
+		DetailBuilder.EditCategory(TEXT("LiveLink"));
 
 		//Hide the Map default UI
 		TSharedRef<IPropertyHandle> ControllersProperty = DetailBuilder.GetProperty(GET_MEMBER_NAME_CHECKED(ULiveLinkComponentController, ControllerMap));
@@ -53,8 +52,10 @@ void FLiveLinkComponentDetailCustomization::CustomizeDetails(IDetailLayoutBuilde
 		TSharedPtr<IPropertyHandleMap> MapHandle = ControllersProperty->AsMap();
 		if (MapHandle.IsValid())
 		{
-			EditedObject = Cast<ULiveLinkComponentController>(SelectedObjects[0].Get());
+			EditedObject = SelectedPtr;
 
+			//Register callback when LiveLinkSubjectRepresentation selection has changed to refresh the UI and update controller
+			TSharedRef<IPropertyHandle> SubjectRepresentationPropertyRef = DetailBuilder.GetProperty(GET_MEMBER_NAME_CHECKED(ULiveLinkComponentController, SubjectRepresentation));
 			SubjectRepresentationPropertyRef->SetOnPropertyValueChanged(FSimpleDelegate::CreateSP(this, &FLiveLinkComponentDetailCustomization::OnSubjectRepresentationPropertyChanged));
 			SubjectRepresentationPropertyRef->SetOnChildPropertyValueChanged(FSimpleDelegate::CreateSP(this, &FLiveLinkComponentDetailCustomization::OnSubjectRepresentationPropertyChanged));
 
@@ -76,14 +77,13 @@ void FLiveLinkComponentDetailCustomization::CustomizeDetails(IDetailLayoutBuilde
 					continue;
 				}
 
-				//Map has a TSubClassof Key type
-				TSubclassOf<ULiveLinkRole> LiveLinkRoleClass;
-				UObject* LiveLinkRoleClassObj = nullptr;
 				TSharedPtr<IPropertyHandle> KeyHandle = EntryHandle->GetKeyHandle();
 
-				//Make sure we were able to query the UClass and that it's not null. 
-				check(KeyHandle->GetValue(LiveLinkRoleClassObj) == FPropertyAccess::Success)
-				LiveLinkRoleClass = Cast<UClass>(LiveLinkRoleClassObj);
+				//Map has a TSubClassof Key type. Make sure we were able to query the UClass and that it's not null. 
+				UObject* LiveLinkRoleClassObj = nullptr;
+				KeyHandle->GetValue(LiveLinkRoleClassObj);
+
+				TSubclassOf<ULiveLinkRole> LiveLinkRoleClass = Cast<UClass>(LiveLinkRoleClassObj);
 				if (LiveLinkRoleClass == nullptr)
 				{
 					continue;
@@ -92,30 +92,21 @@ void FLiveLinkComponentDetailCustomization::CustomizeDetails(IDetailLayoutBuilde
 				FText ControllerName = FText::Format(LOCTEXT("No Controller", "{0}"), FText::FromName(NAME_None));
 
 				//Map value is a pointer to the Controller. Can be null so it could be empty
-				uint32 ControllerValueCount = 0;
-				EntryHandle->GetNumChildren(ControllerValueCount);
-				TSharedPtr<IPropertyHandle> ValueHandle;
-				if (ControllerValueCount > 0)
+				UObject* ControllerPtr = nullptr;
+				EntryHandle->GetValue(ControllerPtr);
+				if (ControllerPtr != nullptr)
 				{
-					ValueHandle = EntryHandle->GetChildHandle(0);
-
-					UObject* ControllerPtr = nullptr;
-					EntryHandle->GetValue(ControllerPtr);
-					if (ControllerPtr != nullptr)
-					{
-						ControllerName = ControllerPtr->GetClass()->GetDisplayNameText();
-					}
+					ControllerName = ControllerPtr->GetClass()->GetDisplayNameText();
 				}
 
 				//Since we're displaying properties of another object, add it as external to the current one being edited. 
-				//Add each Map entry as in the row that is generated for that external objects : Name = Key (Component name) + Value = Dropdown menu with available controllers
 				TArray<UObject*> ExternalObjects;
-				if (EditedObject->ControllerMap.Contains(LiveLinkRoleClass))
-				{
-					ExternalObjects.Add(EditedObject->ControllerMap[LiveLinkRoleClass]);
-				}
-
+				ExternalObjects.Add(ControllerPtr);
 				IDetailPropertyRow* ThisRoleRow = DetailBuilder.EditCategory(TEXT("Role Controllers")).AddExternalObjects(ExternalObjects);
+
+				//As external objects, each map entry will generate a custom widget for its detail property row like this: 
+				//(NameWidget) | (ValueWidget)
+				//Role Name    | Dropdown menu with available controllers
 				if (ThisRoleRow)
 				{
 					ThisRoleRow->CustomWidget()
@@ -125,7 +116,7 @@ void FLiveLinkComponentDetailCustomization::CustomizeDetails(IDetailLayoutBuilde
 					]
 					.ValueContent()
 					[
-						BuildControllerValueWidget(KeyHandle, LiveLinkRoleClass, ControllerName)
+						BuildControllerValueWidget(LiveLinkRoleClass, ControllerName)
 					];
 				}
 			}
@@ -151,71 +142,59 @@ void FLiveLinkComponentDetailCustomization::OnSubjectRepresentationPropertyChang
 	}
 }
 
-TSharedRef<SWidget> FLiveLinkComponentDetailCustomization::HandleControllerComboButton(TSharedPtr<IPropertyHandle> KeyHandle) const
+TSharedRef<SWidget> FLiveLinkComponentDetailCustomization::HandleControllerComboButton(TSubclassOf<ULiveLinkRole> RoleClass) const
 {
 	// Generate menu
 	FMenuBuilder MenuBuilder(true, nullptr);
 	MenuBuilder.BeginSection("SupportedControllers", LOCTEXT("SupportedControllers", "Controllers"));
 	{
-		if (KeyHandle.IsValid())
+		if(RoleClass)
 		{
-			//Make sure we were able to query the UClass and that it's not null. 
-			TSubclassOf<ULiveLinkRole> RoleClass;
-			UObject* RoleClassObj = nullptr;
-			KeyHandle->GetValue(RoleClassObj);
-			RoleClass = Cast<UClass>(RoleClassObj);
-			if(RoleClass)
+			TArray<TSubclassOf<ULiveLinkControllerBase>> NewControllerClasses = ULiveLinkControllerBase::GetControllersForRole(RoleClass);
+			if (NewControllerClasses.Num() > 0)
 			{
-				TArray<TSubclassOf<ULiveLinkControllerBase>> NewControllerClasses = ULiveLinkControllerBase::GetControllersForRole(RoleClass);
-				if (NewControllerClasses.Num() > 0)
+				//Always add a None entry
+				MenuBuilder.AddMenuEntry(
+					FText::FromName(NAME_None),
+					FText::FromName(NAME_None),
+					FSlateIcon(),
+					FUIAction(
+						FExecuteAction::CreateSP(this, &FLiveLinkComponentDetailCustomization::HandleControllerSelection, RoleClass, TWeakObjectPtr<UClass>()),
+						FCanExecuteAction(),
+						FIsActionChecked::CreateSP(this, &FLiveLinkComponentDetailCustomization::IsControllerItemSelected, FName(), RoleClass)
+					),
+					NAME_None,
+					EUserInterfaceActionType::RadioButton
+				);
+
+				for (const TSubclassOf<ULiveLinkControllerBase>& ControllerClass : NewControllerClasses)
 				{
-					//Always add a None entry
-					MenuBuilder.AddMenuEntry(
-						FText::FromName(NAME_None),
-						FText::FromName(NAME_None),
-						FSlateIcon(),
-						FUIAction(
-							FExecuteAction::CreateSP(this, &FLiveLinkComponentDetailCustomization::HandleControllerSelection, RoleClass, TWeakObjectPtr<UClass>()),
-							FCanExecuteAction(),
-							FIsActionChecked::CreateSP(this, &FLiveLinkComponentDetailCustomization::IsControllerItemSelected, FName(), RoleClass)
-						),
-						NAME_None,
-						EUserInterfaceActionType::RadioButton
-					);
-
-					for (const TSubclassOf<ULiveLinkControllerBase>& ControllerClass : NewControllerClasses)
+					if (UClass* ControllerClassPtr = ControllerClass.Get())
 					{
-						if (UClass* ControllerClassPtr = ControllerClass.Get())
-						{
-							MenuBuilder.AddMenuEntry(
-								FText::Format(LOCTEXT("Controller Label", "{0}"), ControllerClassPtr->GetDisplayNameText()),
-								FText::Format(LOCTEXT("Controller ToolTip", "{0}"), ControllerClassPtr->GetDisplayNameText()),
-								FSlateIcon(),
-								FUIAction(
-									FExecuteAction::CreateSP(this, &FLiveLinkComponentDetailCustomization::HandleControllerSelection, RoleClass, MakeWeakObjectPtr(ControllerClassPtr)),
-									FCanExecuteAction(),
-									FIsActionChecked::CreateSP(this, &FLiveLinkComponentDetailCustomization::IsControllerItemSelected, ControllerClassPtr->GetFName(), RoleClass)
-								),
-								NAME_None,
-								EUserInterfaceActionType::RadioButton
-							);
-						}
-
+						MenuBuilder.AddMenuEntry(
+							FText::Format(LOCTEXT("Controller Label", "{0}"), ControllerClassPtr->GetDisplayNameText()),
+							FText::Format(LOCTEXT("Controller ToolTip", "{0}"), ControllerClassPtr->GetDisplayNameText()),
+							FSlateIcon(),
+							FUIAction(
+								FExecuteAction::CreateSP(this, &FLiveLinkComponentDetailCustomization::HandleControllerSelection, RoleClass, MakeWeakObjectPtr(ControllerClassPtr)),
+								FCanExecuteAction(),
+								FIsActionChecked::CreateSP(this, &FLiveLinkComponentDetailCustomization::IsControllerItemSelected, ControllerClassPtr->GetFName(), RoleClass)
+							),
+							NAME_None,
+							EUserInterfaceActionType::RadioButton
+						);
 					}
-				}
-				else
-				{
-					MenuBuilder.AddWidget(SNullWidget::NullWidget, LOCTEXT("NoControllersFound", "No Controllers were found for this role"), false, false);
+
 				}
 			}
 			else
 			{
-				MenuBuilder.AddWidget(SNullWidget::NullWidget, LOCTEXT("InvalidRoleClass", "Role is invalid. Can't find controllers for it"), false, false);
+				MenuBuilder.AddWidget(SNullWidget::NullWidget, LOCTEXT("NoControllersFound", "No Controllers were found for this role"), false, false);
 			}
 		}
 		else
 		{
-			MenuBuilder.AddWidget(SNullWidget::NullWidget, LOCTEXT("InvalidComponent", "Invalid component class handle received"), false, false);
+			MenuBuilder.AddWidget(SNullWidget::NullWidget, LOCTEXT("InvalidRoleClass", "Role is invalid. Can't find controllers for it"), false, false);
 		}
 	}
 	MenuBuilder.EndSection();
@@ -225,14 +204,14 @@ TSharedRef<SWidget> FLiveLinkComponentDetailCustomization::HandleControllerCombo
 
 void FLiveLinkComponentDetailCustomization::HandleControllerSelection(TSubclassOf<ULiveLinkRole> RoleClass, TWeakObjectPtr<UClass> SelectedControllerClass) const
 {
-	if (ULiveLinkComponentController* EditedObjectPtr = EditedObject.Get())
-	{
-		const FScopedTransaction Transaction(LOCTEXT("OnChangedController", "Property Controller Selection"));
-		EditedObjectPtr->Modify();
-
-		UClass* SelectedControllerClassPtr = SelectedControllerClass.Get();
-		EditedObjectPtr->SetControllerClassForRole(RoleClass, SelectedControllerClassPtr);
-	}
+ 	if (ULiveLinkComponentController* EditedObjectPtr = EditedObject.Get())
+ 	{
+ 		const FScopedTransaction Transaction(LOCTEXT("OnChangedController", "Property Controller Selection"));
+ 		EditedObjectPtr->Modify();
+ 
+ 		UClass* SelectedControllerClassPtr = SelectedControllerClass.Get();
+ 		EditedObjectPtr->SetControllerClassForRole(RoleClass, SelectedControllerClassPtr);
+ 	}
 }
 
 bool FLiveLinkComponentDetailCustomization::IsControllerItemSelected(FName Item, TSubclassOf<ULiveLinkRole> RoleClass) const
@@ -244,7 +223,6 @@ bool FLiveLinkComponentDetailCustomization::IsControllerItemSelected(FName Item,
 
 	if (ULiveLinkComponentController* EditedObjectPtr = EditedObject.Get())
 	{
-
 		TObjectPtr<ULiveLinkControllerBase>* CurrentClass = EditedObjectPtr->ControllerMap.Find(RoleClass);
 		if (CurrentClass == nullptr || *CurrentClass == nullptr)
 		{
@@ -259,117 +237,39 @@ bool FLiveLinkComponentDetailCustomization::IsControllerItemSelected(FName Item,
 	return false;
 }
 
-FSlateColor FLiveLinkComponentDetailCustomization::HandleControllerStatusColorAndOpacity(TSubclassOf<ULiveLinkRole> RoleClassEntry) const
+EVisibility FLiveLinkComponentDetailCustomization::HandleControllerWarningVisibility(TSubclassOf<ULiveLinkRole> RoleClassEntry) const
 {
-	FSlateColor Result = FLinearColor::Green;
-	bool bIsValid = true;
-
 	if (ULiveLinkComponentController* EditorObjectPtr = EditedObject.Get())
 	{
 		TObjectPtr<ULiveLinkControllerBase>* AssociatedControllerPtr = EditorObjectPtr->ControllerMap.Find(RoleClassEntry);
 		if (AssociatedControllerPtr && *AssociatedControllerPtr)
 		{
-			if (UActorComponent* SelectedComponent = EditorObjectPtr->ComponentToControl.GetComponent(EditorObjectPtr->GetOwner()))
+			const ULiveLinkControllerBase* AssociatedController = *AssociatedControllerPtr;
+			if (UActorComponent* SelectedComponent = AssociatedController->GetAttachedComponent())
 			{
-				const ULiveLinkControllerBase* AssociatedController = *AssociatedControllerPtr;
 				if (!SelectedComponent->IsA(AssociatedController->GetDesiredComponentClass()))
 				{
 					//Controller exists for RoleClass and desired component is not the kind it wants
-					bIsValid = false;
+					return EVisibility::Visible;
 				}
 			}
 			else
 			{
 				//Component is not valid
-				bIsValid = false;
+				return EVisibility::Visible;
 			}
 		}
 	}
 
-	if (!bIsValid)
-	{
-		Result = FLinearColor::Red;
-	}
-
-	return Result;
-}
-
-FText FLiveLinkComponentDetailCustomization::HandleControllerStatusText(TSubclassOf<ULiveLinkRole> RoleClassEntry) const
-{
-	FText Result;
-	bool bIsValid = true;
-
-	if (ULiveLinkComponentController* EditorObjectPtr = EditedObject.Get())
-	{
-		TObjectPtr<ULiveLinkControllerBase>* AssociatedControllerPtr = EditorObjectPtr->ControllerMap.Find(RoleClassEntry);
-		if (AssociatedControllerPtr && *AssociatedControllerPtr)
-		{
-			if (UActorComponent* SelectedComponent = EditorObjectPtr->ComponentToControl.GetComponent(EditorObjectPtr->GetOwner()))
-			{
-				const ULiveLinkControllerBase* AssociatedController = *AssociatedControllerPtr;
-				if (!SelectedComponent->IsA(AssociatedController->GetDesiredComponentClass()))
-				{
-					//Controller exists for RoleClass and desired component is not the kind it wants
-					bIsValid = false;
-				}
-			}
-			else
-			{
-				//Component is not valid
-				bIsValid = false;
-			}
-		}
-	}
-
-	if (!bIsValid)
-	{
-		Result = FEditorFontGlyphs::Ban;
-	}
-
-	return Result;
-}
-
-FText FLiveLinkComponentDetailCustomization::HandleControllerStatusToolTipText(TSubclassOf<ULiveLinkRole> RoleClassEntry) const
-{
-	FText Result;
-	bool bIsValid = true;
-
-	if (ULiveLinkComponentController* EditorObjectPtr = EditedObject.Get())
-	{
-		TObjectPtr<ULiveLinkControllerBase>* AssociatedControllerPtr = EditorObjectPtr->ControllerMap.Find(RoleClassEntry);
-		if (AssociatedControllerPtr && *AssociatedControllerPtr)
-		{
-			if (UActorComponent* SelectedComponent = EditorObjectPtr->ComponentToControl.GetComponent(EditorObjectPtr->GetOwner()))
-			{
-				const ULiveLinkControllerBase* AssociatedController = *AssociatedControllerPtr;
-				if (!SelectedComponent->IsA(AssociatedController->GetDesiredComponentClass()))
-				{
-					//Controller exists for RoleClass and desired component is not the kind it wants
-					bIsValid = false;
-				}
-			}
-			else
-			{
-				//Component is not valid
-				bIsValid = false;
-			}
-		}
-	}
-
-	if (!bIsValid)
-	{
-		Result = LOCTEXT("ControllerToolTip", "Controller can't control selected component");
-	}
-
-	return Result;
+	return EVisibility::Hidden;
 }
 
 TSharedRef<SWidget> FLiveLinkComponentDetailCustomization::BuildControllerNameWidget(TSharedPtr<IPropertyHandle> ControllersProperty, TSubclassOf<ULiveLinkRole> RoleClass) const
 {
-	return ControllersProperty->CreatePropertyNameWidget(RoleClass.Get()->GetDisplayNameText().ToUpper());
+	return ControllersProperty->CreatePropertyNameWidget(RoleClass.Get()->GetDisplayNameText());
 }
 
-TSharedRef<SWidget> FLiveLinkComponentDetailCustomization::BuildControllerValueWidget(TSharedPtr<IPropertyHandle> RoleKeyPropertyHandle, TSubclassOf<ULiveLinkRole> RoleClass, const FText& ControllerName) const
+TSharedRef<SWidget> FLiveLinkComponentDetailCustomization::BuildControllerValueWidget(TSubclassOf<ULiveLinkRole> RoleClass, const FText& ControllerName) const
 {
 	TSharedRef<SWidget> ValueWidget = 
 		SNew(SHorizontalBox)
@@ -378,7 +278,7 @@ TSharedRef<SWidget> FLiveLinkComponentDetailCustomization::BuildControllerValueW
 		.Padding(FMargin(4.0f, 0.0f, 0.0f, 0.0f))
 		[
 			SNew(SComboButton)
-			.OnGetMenuContent(this, &FLiveLinkComponentDetailCustomization::HandleControllerComboButton, RoleKeyPropertyHandle)
+			.OnGetMenuContent(this, &FLiveLinkComponentDetailCustomization::HandleControllerComboButton, RoleClass)
 			.ContentPadding(FMargin(4.0, 2.0))
 			.ButtonContent()
 			[
@@ -389,14 +289,14 @@ TSharedRef<SWidget> FLiveLinkComponentDetailCustomization::BuildControllerValueW
 		]
 		+ SHorizontalBox::Slot()
 		.AutoWidth()
+		.VAlign(VAlign_Center)
+		.HAlign(HAlign_Center)
 		.Padding(FMargin(4.0f, 0.0f, 0.0f, 0.0f))
-		.HAlign(EHorizontalAlignment::HAlign_Left)
 		[
-			SNew(STextBlock)
-			.Font(FEditorStyle::Get().GetFontStyle("FontAwesome.11"))
-			.Text(this, &FLiveLinkComponentDetailCustomization::HandleControllerStatusText, RoleClass)
-			.ColorAndOpacity(this, &FLiveLinkComponentDetailCustomization::HandleControllerStatusColorAndOpacity, RoleClass)
-			.ToolTipText(this, &FLiveLinkComponentDetailCustomization::HandleControllerStatusToolTipText, RoleClass)
+			SNew(SImage)
+			.Image(FLiveLinkEditorPrivate::GetStyleSet()->GetBrush("LiveLinkController.WarningIcon"))
+			.ToolTipText(LOCTEXT("ControllerWarningToolTip", "The selected component to control does not match this controller's desired component class"))
+ 			.Visibility(this, &FLiveLinkComponentDetailCustomization::HandleControllerWarningVisibility, RoleClass)
 		];
 
 	return ValueWidget;
diff --git a/Engine/Plugins/Animation/LiveLink/Source/LiveLinkEditor/Private/LiveLinkComponentDetailCustomization.h b/Engine/Plugins/Animation/LiveLink/Source/LiveLinkEditor/Private/LiveLinkComponentDetailCustomization.h
--- a/Engine/Plugins/Animation/LiveLink/Source/LiveLinkEditor/Private/LiveLinkComponentDetailCustomization.h
+++ b/Engine/Plugins/Animation/LiveLink/Source/LiveLinkEditor/Private/LiveLinkComponentDetailCustomization.h
@@ -5,18 +5,16 @@
 #include "IDetailCustomization.h"
 
 #include "CoreMinimal.h"
+#include "Layout/Visibility.h"
+#include "LiveLinkComponentController.h"
 #include "LiveLinkRole.h"
 #include "Styling/SlateColor.h"
 #include "Templates/SubclassOf.h"
 
-#include "LiveLinkComponentController.h"
-
-
 class IDetailLayoutBuilder;
 class IPropertyHandle;
 class SWidget;
 
-
 /**
 * Customizes a ULiveLinkComponentController details
 */
@@ -35,14 +33,12 @@ public:
 
 protected:
 	void OnSubjectRepresentationPropertyChanged();
-	TSharedRef<SWidget> HandleControllerComboButton(TSharedPtr<IPropertyHandle> KeyHandle) const;
+	TSharedRef<SWidget> HandleControllerComboButton(TSubclassOf<ULiveLinkRole> RoleClass) const;
 	void HandleControllerSelection(TSubclassOf<ULiveLinkRole> RoleClass, TWeakObjectPtr<UClass> SelectedControllerClass) const;
 	bool IsControllerItemSelected(FName Item, TSubclassOf<ULiveLinkRole> RoleClass) const;
-	FSlateColor HandleControllerStatusColorAndOpacity(TSubclassOf<ULiveLinkRole> RoleClassEntry) const;
-	FText HandleControllerStatusText(TSubclassOf<ULiveLinkRole> RoleClassEntry) const;
-	FText HandleControllerStatusToolTipText(TSubclassOf<ULiveLinkRole> RoleClassEntry) const;
+	EVisibility HandleControllerWarningVisibility(TSubclassOf<ULiveLinkRole> RoleClassEntry) const;
 	TSharedRef<SWidget> BuildControllerNameWidget(TSharedPtr<IPropertyHandle> ControllersProperty, TSubclassOf<ULiveLinkRole> RoleClass) const;
-	TSharedRef<SWidget> BuildControllerValueWidget(TSharedPtr<IPropertyHandle> RoleKeyPropertyHandle, TSubclassOf<ULiveLinkRole> RoleClass, const FText& ControllerName) const;
+	TSharedRef<SWidget> BuildControllerValueWidget(TSubclassOf<ULiveLinkRole> RoleClass, const FText& ControllerName) const;
 	void ForceRefreshDetails();
 
 protected:
diff --git a/Engine/Plugins/Animation/LiveLink/Source/LiveLinkEditor/Private/LiveLinkControllerBaseDetailCustomization.cpp b/Engine/Plugins/Animation/LiveLink/Source/LiveLinkEditor/Private/LiveLinkControllerBaseDetailCustomization.cpp
new file mode 100644
--- /dev/null
+++ b/Engine/Plugins/Animation/LiveLink/Source/LiveLinkEditor/Private/LiveLinkControllerBaseDetailCustomization.cpp
@@ -0,0 +1,81 @@
+// Copyright Epic Games, Inc. All Rights Reserved.
+
+#include "LiveLinkControllerBaseDetailCustomization.h"
+
+#include "DetailCategoryBuilder.h"
+#include "DetailLayoutBuilder.h"
+
+#define LOCTEXT_NAMESPACE "LiveLinkControllerBaseDetailsCustomization"
+
+void FLiveLinkControllerBaseDetailCustomization::CustomizeDetails(IDetailLayoutBuilder& DetailBuilder)
+{
+	TArray<TWeakObjectPtr<UObject>> SelectedObjects;
+	DetailBuilder.GetObjectsBeingCustomized(SelectedObjects);
+
+	if (SelectedObjects.Num() < 1)
+	{
+		return;
+	}
+
+	// The controller is not selectable in the details panel, the the controller map is hidden when multiple livelink components are selected
+	// Safe to assume there is only one selected controller being customized
+	LiveLinkControllerWeak = Cast<ULiveLinkControllerBase>(SelectedObjects[0].Get());
+
+	// Collapse the properties of each controller by default (cleaner look for controllers with a lot of settings)
+	TArray<FName> CategoryNames;
+	DetailBuilder.GetCategoryNames(CategoryNames);
+	for (FName CategoryName : CategoryNames)
+	{
+		IDetailCategoryBuilder& Category = DetailBuilder.EditCategory(CategoryName);
+		Category.InitiallyCollapsed(true);
+	}
+
+	DetailBuilder.SortCategories([](const TMap<FName, IDetailCategoryBuilder*>& CategoryMap)
+	{
+		for (const TPair<FName, IDetailCategoryBuilder*>& Pair : CategoryMap)
+		{
+			int32 SortOrder = Pair.Value->GetSortOrder();
+			const FName CategoryName = Pair.Key;
+
+			if (CategoryName == "ComponentToControl")
+			{
+				SortOrder = 0;
+			}
+			else
+			{
+				const int32 ValueSortOrder = Pair.Value->GetSortOrder();
+				if (ValueSortOrder >= SortOrder && ValueSortOrder < SortOrder + 10)
+				{
+					SortOrder += 10;
+				}
+				else
+				{
+					continue;
+				}
+			}
+
+			Pair.Value->SetSortOrder(SortOrder);
+		}
+	});
+
+	// Try to get the ComponentToControl property from the category
+	IDetailCategoryBuilder& ComponentToControlCategory = DetailBuilder.EditCategory(TEXT("ComponentToControl"));
+	TArray<TSharedRef<IPropertyHandle>> ComponentToControlProperties;
+	ComponentToControlCategory.GetDefaultProperties(ComponentToControlProperties);
+	if (ComponentToControlProperties.Num() > 0)
+	{
+		// Set callback when the user makes a change to the component picker
+		TSharedRef<IPropertyHandle> ComponentToControlProperty = ComponentToControlProperties[0];
+		ComponentToControlProperty->SetOnPropertyValueChanged(FSimpleDelegate::CreateSP(this, &FLiveLinkControllerBaseDetailCustomization::OnComponentChanged));
+	}
+ }
+
+void FLiveLinkControllerBaseDetailCustomization::OnComponentChanged()
+{
+	if (ULiveLinkControllerBase* LiveLinkController = LiveLinkControllerWeak.Get())
+	{
+		LiveLinkController->SetAttachedComponent(LiveLinkController->GetAttachedComponent());
+	}
+}
+
+#undef LOCTEXT_NAMESPACE
diff --git a/Engine/Plugins/Animation/LiveLink/Source/LiveLinkEditor/Private/LiveLinkControllerBaseDetailCustomization.h b/Engine/Plugins/Animation/LiveLink/Source/LiveLinkEditor/Private/LiveLinkControllerBaseDetailCustomization.h
new file mode 100644
--- /dev/null
+++ b/Engine/Plugins/Animation/LiveLink/Source/LiveLinkEditor/Private/LiveLinkControllerBaseDetailCustomization.h
@@ -0,0 +1,32 @@
+// Copyright Epic Games, Inc. All Rights Reserved.
+
+#pragma once
+
+#include "IDetailCustomization.h"
+
+#include "LiveLinkControllerBase.h"
+
+class IDetailLayoutBuilder;
+
+/**
+* Customizes ULiveLinkControllerBase details
+*/
+class FLiveLinkControllerBaseDetailCustomization : public IDetailCustomization
+{
+public:
+
+	static TSharedRef<IDetailCustomization> MakeInstance()
+	{
+		return MakeShared<FLiveLinkControllerBaseDetailCustomization>();
+	}
+
+	// IDetailCustomization interface
+	virtual void CustomizeDetails(class IDetailLayoutBuilder& DetailBuilder) override;
+	// End IDetailCustomization interface
+
+private:
+	void OnComponentChanged();
+
+protected:
+	TWeakObjectPtr<ULiveLinkControllerBase> LiveLinkControllerWeak;
+};
diff --git a/Engine/Plugins/Animation/LiveLink/Source/LiveLinkEditor/Private/LiveLinkEditorModule.cpp b/Engine/Plugins/Animation/LiveLink/Source/LiveLinkEditor/Private/LiveLinkEditorModule.cpp
--- a/Engine/Plugins/Animation/LiveLink/Source/LiveLinkEditor/Private/LiveLinkEditorModule.cpp
+++ b/Engine/Plugins/Animation/LiveLink/Source/LiveLinkEditor/Private/LiveLinkEditorModule.cpp
@@ -14,6 +14,7 @@
 #include "Modules/ModuleManager.h"
 #include "PropertyEditorModule.h"
 #include "Styling/SlateStyle.h"
+#include "Styling/StyleColors.h"
 #include "Styling/SlateTypes.h"
 #include "Styling/SlateStyleRegistry.h"
 #include "Widgets/Docking/SDockTab.h"
@@ -23,6 +24,8 @@
 #include "LiveLinkComponentController.h"
 #include "LiveLinkComponentDetailCustomization.h"
 #include "LiveLinkComponentSettings.h"
+#include "LiveLinkControllerBase.h"
+#include "LiveLinkControllerBaseDetailCustomization.h"
 #include "LiveLinkClient.h"
 #include "LiveLinkClientPanel.h"
 #include "LiveLinkClientCommands.h"
@@ -62,6 +65,7 @@ namespace LiveLinkEditorModuleUtils
 	}
 }
 
+#define IMAGE_BRUSH_SVG(RelativePath, ...) FSlateVectorImageBrush(StyleSet->RootToCoreContentDir(RelativePath, TEXT(".svg")), __VA_ARGS__)
 #define IMAGE_PLUGIN_BRUSH( RelativePath, ... ) FSlateImageBrush( LiveLinkEditorModuleUtils::InPluginContent( RelativePath, ".png" ), __VA_ARGS__ )
 #define IMAGE_PLUGIN_BRUSH_SVG( RelativePath, ... ) FSlateVectorImageBrush( LiveLinkEditorModuleUtils::InPluginContent( RelativePath, ".svg" ), __VA_ARGS__ )
 
@@ -113,6 +117,8 @@ public:
 		StyleSet->Set("LiveLinkClient.Common.RemoveSource", new IMAGE_PLUGIN_BRUSH(TEXT("icon_RemoveSource_40x"), Icon40x40));
 		StyleSet->Set("LiveLinkClient.Common.RemoveAllSources", new IMAGE_PLUGIN_BRUSH(TEXT("icon_RemoveSource_40x"), Icon40x40));
 
+		StyleSet->Set("LiveLinkController.WarningIcon", new IMAGE_BRUSH_SVG("Starship/Common/alert-circle", Icon16x16, FStyleColors::Warning));
+		
 		FButtonStyle Button = FButtonStyle()
 			.SetNormal(FSlateBoxBrush(StyleSet->RootToContentDir("Common/ButtonHoverHint.png"), FMargin(4 / 16.0f), FLinearColor(1, 1, 1, 0.15f)))
 			.SetHovered(FSlateBoxBrush(StyleSet->RootToContentDir("Common/ButtonHoverHint.png"), FMargin(4 / 16.0f), FLinearColor(1, 1, 1, 0.25f)))
@@ -250,6 +256,7 @@ private:
 		PropertyEditorModule.RegisterCustomPropertyTypeLayout(FLiveLinkSubjectName::StaticStruct()->GetFName(), FOnGetPropertyTypeCustomizationInstance::CreateStatic(&FLiveLinkSubjectNameDetailCustomization::MakeInstance));
 		PropertyEditorModule.RegisterCustomClassLayout(ULiveLinkSourceSettings::StaticClass()->GetFName(), FOnGetDetailCustomizationInstance::CreateStatic(&FLiveLinkSourceSettingsDetailCustomization::MakeInstance));
 		PropertyEditorModule.RegisterCustomClassLayout(ULiveLinkComponentController::StaticClass()->GetFName(), FOnGetDetailCustomizationInstance::CreateStatic(&FLiveLinkComponentDetailCustomization::MakeInstance));
+		PropertyEditorModule.RegisterCustomClassLayout(ULiveLinkControllerBase::StaticClass()->GetFName(), FOnGetDetailCustomizationInstance::CreateStatic(&FLiveLinkControllerBaseDetailCustomization::MakeInstance));
 
 		LiveLinkGraphPanelPinFactory = MakeShared<FLiveLinkGraphPanelPinFactory>();
 		FEdGraphUtilities::RegisterVisualPinFactory(LiveLinkGraphPanelPinFactory);
@@ -264,6 +271,7 @@ private:
 			if (PropertyEditorModule)
 			{
 				PropertyEditorModule->UnregisterCustomClassLayout(ULiveLinkComponentController::StaticClass()->GetFName());
+				PropertyEditorModule->UnregisterCustomClassLayout(ULiveLinkControllerBase::StaticClass()->GetFName());
 				PropertyEditorModule->UnregisterCustomClassLayout(ULiveLinkSourceSettings::StaticClass()->GetFName());
 				PropertyEditorModule->UnregisterCustomPropertyTypeLayout(FLiveLinkSubjectName::StaticStruct()->GetFName());
 				PropertyEditorModule->UnregisterCustomPropertyTypeLayout(FLiveLinkSubjectRepresentation::StaticStruct()->GetFName());
diff --git a/Engine/Plugins/Importers/USDImporter/Source/USDSchemas/Private/USDGeomXformableTranslator.cpp b/Engine/Plugins/Importers/USDImporter/Source/USDSchemas/Private/USDGeomXformableTranslator.cpp
--- a/Engine/Plugins/Importers/USDImporter/Source/USDSchemas/Private/USDGeomXformableTranslator.cpp
+++ b/Engine/Plugins/Importers/USDImporter/Source/USDSchemas/Private/USDGeomXformableTranslator.cpp
@@ -105,7 +105,7 @@ namespace UE::UsdXformableTranslatorImpl::Private
 
 				for ( ULiveLinkComponentController* LiveLinkComponent : LiveLinkComponents )
 				{
-					if ( LiveLinkComponent->ComponentToControl.GetComponent( Parent ) == Component )
+					if ( LiveLinkComponent->GetControlledComponent( ULiveLinkTransformRole::StaticClass() ) == Component )
 					{
 						// We found some other controller handling this component somehow, use that
 						Controller = LiveLinkComponent;
@@ -120,7 +120,6 @@ namespace UE::UsdXformableTranslatorImpl::Private
 
 					Controller = NewObject<ULiveLinkComponentController>( Parent, NAME_None, Context.ObjectFlags );
 					Controller->bUpdateInEditor = true;
-					Controller->ComponentToControl.PathToComponent = Component->GetPathName( Parent );
 
 					// Important because of how ULiveLinkComponentController::TickComponent also checks for the sequencer
 					// tag to try and guess if the controlled component is a spawnable
@@ -155,6 +154,9 @@ namespace UE::UsdXformableTranslatorImpl::Private
 			}
 			Controller->SetSubjectRepresentation( SubjectRepresentation );
 
+			// This should be done after setting the subject representation to ensure that the LiveLink component's ControllerMap has a transform controller
+			Controller->SetControlledComponent( ULiveLinkTransformRole::StaticClass(), Component );
+
 			if ( pxr::UsdAttribute Attr = Prim.GetAttribute( UnrealIdentifiers::UnrealLiveLinkEnabled ) )
 			{
 				bool bEnabled = true;
@@ -184,9 +186,9 @@ namespace UE::UsdXformableTranslatorImpl::Private
 
 		for ( ULiveLinkComponentController* LiveLinkComponent : LiveLinkComponents )
 		{
-			if ( LiveLinkComponent->ComponentToControl.GetComponent( Parent ) == Component )
+			if ( LiveLinkComponent->GetControlledComponent( ULiveLinkTransformRole::StaticClass() ) == Component )
 			{
-				LiveLinkComponent->ComponentToControl.ComponentProperty = NAME_None;
+				LiveLinkComponent->SetControlledComponent( ULiveLinkTransformRole::StaticClass(), nullptr );
 				Parent->RemoveInstanceComponent( LiveLinkComponent );
 				break;
 			}
diff --git a/Engine/Plugins/Importers/USDImporter/Source/USDStage/Private/USDStageActor.cpp b/Engine/Plugins/Importers/USDImporter/Source/USDStage/Private/USDStageActor.cpp
--- a/Engine/Plugins/Importers/USDImporter/Source/USDStage/Private/USDStageActor.cpp
+++ b/Engine/Plugins/Importers/USDImporter/Source/USDStage/Private/USDStageActor.cpp
@@ -60,6 +60,7 @@
 #include "Modules/ModuleManager.h"
 #include "PhysicsEngine/BodySetup.h"
 #include "Rendering/SkeletalMeshLODImporterData.h"
+#include "Roles/LiveLinkTransformRole.h"
 #include "StaticMeshAttributes.h"
 #include "StaticMeshOperations.h"
 #include "Tracks/MovieScene3DTransformTrack.h"
@@ -2489,7 +2490,10 @@ void AUsdStageActor::OnObjectPropertyChanged( UObject* ObjectBeingModified, FPro
 	ULiveLinkComponentController* Controller = Cast< ULiveLinkComponentController >( ObjectBeingModified );
 	if ( Controller )
 	{
-		ObjectBeingModified = Controller->ComponentToControl.GetComponent( Controller->GetOwner() );
+		if (UActorComponent* ControlledComponent = Controller->GetControlledComponent(ULiveLinkTransformRole::StaticClass()))
+		{
+			ObjectBeingModified = ControlledComponent;
+		}
 	}
 
 	UObject* PrimObject = ObjectBeingModified;
diff --git a/Engine/Plugins/Importers/USDImporter/Source/USDStage/USDStage.Build.cs b/Engine/Plugins/Importers/USDImporter/Source/USDStage/USDStage.Build.cs
--- a/Engine/Plugins/Importers/USDImporter/Source/USDStage/USDStage.Build.cs
+++ b/Engine/Plugins/Importers/USDImporter/Source/USDStage/USDStage.Build.cs
@@ -25,6 +25,7 @@ namespace UnrealBuildTool.Rules
 					"Engine",
 					"LevelSequence",
 					"LiveLinkComponents", // For tracking edits to LiveLinkComponentController properties and writing to USD
+					"LiveLinkInterface",
 					"MeshDescription",
 					"MovieScene",
 					"MovieSceneTracks",
diff --git a/Engine/Plugins/VirtualProduction/LiveLinkCamera/Source/LiveLinkCamera/Private/LiveLinkCameraController.cpp b/Engine/Plugins/VirtualProduction/LiveLinkCamera/Source/LiveLinkCamera/Private/LiveLinkCameraController.cpp
--- a/Engine/Plugins/VirtualProduction/LiveLinkCamera/Source/LiveLinkCamera/Private/LiveLinkCameraController.cpp
+++ b/Engine/Plugins/VirtualProduction/LiveLinkCamera/Source/LiveLinkCamera/Private/LiveLinkCameraController.cpp
@@ -85,9 +85,9 @@ void ULiveLinkCameraController::Tick(float DeltaTime, const FLiveLinkSubjectFram
 	// Invalidate the lens file evaluation data
 	LensFileEvalData.Invalidate();
 
-	if (UCameraComponent* CameraComponent = Cast<UCameraComponent>(AttachedComponent))
+	if (UCameraComponent* CameraComponent = Cast<UCameraComponent>(GetAttachedComponent()))
 	{
-		if (AActor* Camera = Cast<AActor>(AttachedComponent->GetOwner()))
+		if (AActor* Camera = Cast<AActor>(GetOuterActor()))
 		{
 			LensFileEvalData.Camera.UniqueId = Camera->GetUniqueID();
 		}
@@ -113,7 +113,7 @@ void ULiveLinkCameraController::Tick(float DeltaTime, const FLiveLinkSubjectFram
 			LensFileEvalData.Input.Zoom = FrameData->FocalLength;
 		}
 
-		if (UCameraComponent* CameraComponent = Cast<UCameraComponent>(AttachedComponent))
+		if (UCameraComponent* CameraComponent = Cast<UCameraComponent>(GetAttachedComponent()))
 		{
 			//Stamp previous values that have an impact on frustum visual representation
 			const float PreviousFOV = CameraComponent->FieldOfView;
@@ -235,7 +235,7 @@ void ULiveLinkCameraController::SetAttachedComponent(UActorComponent* ActorCompo
 {
 	Super::SetAttachedComponent(ActorComponent);
 
-	if (UCineCameraComponent* const CineCameraComponent = Cast<UCineCameraComponent>(AttachedComponent))
+	if (UCineCameraComponent* const CineCameraComponent = Cast<UCineCameraComponent>(GetAttachedComponent()))
 	{	
 		// Initialize the most recent filmback and component transform to the current values of the camera to properly detect changes to this property
 		LastFilmback = CineCameraComponent->Filmback;
@@ -255,7 +255,7 @@ void ULiveLinkCameraController::SetAttachedComponent(UActorComponent* ActorCompo
 
 void ULiveLinkCameraController::Cleanup()
 {
-	if (UCineCameraComponent* const CineCameraComponent = Cast<UCineCameraComponent>(AttachedComponent))
+	if (UCineCameraComponent* const CineCameraComponent = Cast<UCineCameraComponent>(GetAttachedComponent()))
 	{
 		if (UCameraCalibrationSubsystem* SubSystem = GEngine->GetEngineSubsystem<UCameraCalibrationSubsystem>())
 		{
@@ -308,7 +308,7 @@ void ULiveLinkCameraController::SetApplyNodalOffset(bool bInApplyNodalOffset)
 
 	bApplyNodalOffset = bInApplyNodalOffset;
 
-	if (UCineCameraComponent* const CineCameraComponent = Cast<UCineCameraComponent>(AttachedComponent))
+	if (UCineCameraComponent* const CineCameraComponent = Cast<UCineCameraComponent>(GetAttachedComponent()))
 	{
 		if (bApplyNodalOffset)
 		{
@@ -353,7 +353,7 @@ void ULiveLinkCameraController::PostEditChangeProperty(struct FPropertyChangedEv
 
 	if (PropertyName == GET_MEMBER_NAME_CHECKED(ULiveLinkCameraController, bApplyNodalOffset))
 	{
-		if (UCineCameraComponent* const CineCameraComponent = Cast<UCineCameraComponent>(AttachedComponent))
+		if (UCineCameraComponent* const CineCameraComponent = Cast<UCineCameraComponent>(GetAttachedComponent()))
 		{
 			if (bApplyNodalOffset)
 			{
@@ -373,7 +373,7 @@ void ULiveLinkCameraController::PostEditChangeProperty(struct FPropertyChangedEv
 	else if (PropertyName == GET_MEMBER_NAME_CHECKED(FLensFilePicker, LensFile)
 			|| PropertyName == GET_MEMBER_NAME_CHECKED(FLensFilePicker, bUseDefaultLensFile))
 	{
-		if (UCineCameraComponent* const CineCameraComponent = Cast<UCineCameraComponent>(AttachedComponent))
+		if (UCineCameraComponent* const CineCameraComponent = Cast<UCineCameraComponent>(GetAttachedComponent()))
 		{
 			ULensFile* SelectedLensFile = LensFilePicker.GetLensFile();
 			if (SelectedLensFile)
@@ -672,7 +672,7 @@ void ULiveLinkCameraController::OnPostActorTick(UWorld* World, ELevelTick TickTy
 {
 	if (World == GetWorld())
 	{
-		if (UCineCameraComponent* const CineCameraComponent = Cast<UCineCameraComponent>(AttachedComponent))
+		if (UCineCameraComponent* const CineCameraComponent = Cast<UCineCameraComponent>(GetAttachedComponent()))
 		{
 			if (ULensFile* CurrentLensFile = LensFilePicker.GetLensFile())
 			{
@@ -846,7 +846,7 @@ void ULiveLinkCameraController::PostLoad()
 				//if Subjects role direct controller is us, set the component to control to what we had
 				if (LiveLinkComponent->SubjectRepresentation.Role == ULiveLinkCameraRole::StaticClass())
 				{
-					LiveLinkComponent->ComponentToControl = ComponentToControl_DEPRECATED;
+					ComponentPicker = ComponentToControl_DEPRECATED;
 				}
 			}
 		}
diff --git a/Engine/Plugins/VirtualProduction/LiveLinkLens/Source/LiveLinkLens/Private/LiveLinkLensController.cpp b/Engine/Plugins/VirtualProduction/LiveLinkLens/Source/LiveLinkLens/Private/LiveLinkLensController.cpp
--- a/Engine/Plugins/VirtualProduction/LiveLinkLens/Source/LiveLinkLens/Private/LiveLinkLensController.cpp
+++ b/Engine/Plugins/VirtualProduction/LiveLinkLens/Source/LiveLinkLens/Private/LiveLinkLensController.cpp
@@ -25,7 +25,7 @@ void ULiveLinkLensController::Tick(float DeltaTime, const FLiveLinkSubjectFrameD
 
 	if (StaticData && FrameData)
 	{
-		if (UCineCameraComponent* const CineCameraComponent = Cast<UCineCameraComponent>(AttachedComponent))
+		if (UCineCameraComponent* const CineCameraComponent = Cast<UCineCameraComponent>(GetAttachedComponent()))
 		{
 			UCameraCalibrationSubsystem* const SubSystem = GEngine->GetEngineSubsystem<UCameraCalibrationSubsystem>();
 			const TSubclassOf<ULensModel> LensModel = SubSystem->GetRegisteredLensModel(StaticData->LensModel);
@@ -80,7 +80,7 @@ TSubclassOf<UActorComponent> ULiveLinkLensController::GetDesiredComponentClass()
 
 void ULiveLinkLensController::Cleanup()
 {
-	if (UCineCameraComponent* const CineCameraComponent = Cast<UCineCameraComponent>(AttachedComponent))
+	if (UCineCameraComponent* const CineCameraComponent = Cast<UCineCameraComponent>(GetAttachedComponent()))
 	{
 		if (UCameraCalibrationSubsystem* SubSystem = GEngine->GetEngineSubsystem<UCameraCalibrationSubsystem>())
 		{
diff --git a/Engine/Source/Runtime/Core/Public/UObject/UE5MainStreamObjectVersion.h b/Engine/Source/Runtime/Core/Public/UObject/UE5MainStreamObjectVersion.h
--- a/Engine/Source/Runtime/Core/Public/UObject/UE5MainStreamObjectVersion.h
+++ b/Engine/Source/Runtime/Core/Public/UObject/UE5MainStreamObjectVersion.h
@@ -221,6 +221,9 @@ struct CORE_API FUE5MainStreamObjectVersion
 		// Move some Chaos flags into a bitfield
 		AddRigidParticleControlFlags,
 
+		// Allow each LiveLink controller to specify its own component to control
+		LiveLinkComponentPickerPerController,
+
 		// -----<new versions can be added above this line>-------------------------------------------------
 		VersionPlusOne,
 		LatestVersion = VersionPlusOne - 1
