Adding AppStyle.h includes

diff --git a/Engine/Plugins/Animation/LiveLink/Source/LiveLinkGraphNode/Private/K2Node_EvaluateLiveLinkCustom.cpp b/Engine/Plugins/Animation/LiveLink/Source/LiveLinkGraphNode/Private/K2Node_EvaluateLiveLinkCustom.cpp
--- a/Engine/Plugins/Animation/LiveLink/Source/LiveLinkGraphNode/Private/K2Node_EvaluateLiveLinkCustom.cpp
+++ b/Engine/Plugins/Animation/LiveLink/Source/LiveLinkGraphNode/Private/K2Node_EvaluateLiveLinkCustom.cpp
@@ -14,6 +14,7 @@
 #include "LiveLinkTypes.h"
 #include "LiveLinkBlueprintLibrary.h"
 #include "UObject/PropertyPortFlags.h"
+#include "Styling/AppStyle.h"
 
 
 #define LOCTEXT_NAMESPACE "K2Node_EvaluateLiveLinkCustom"
diff --git a/Engine/Plugins/Animation/LiveLink/Source/LiveLinkGraphNode/Private/K2Node_UpdateVirtualSubjectDataBase.cpp b/Engine/Plugins/Animation/LiveLink/Source/LiveLinkGraphNode/Private/K2Node_UpdateVirtualSubjectDataBase.cpp
--- a/Engine/Plugins/Animation/LiveLink/Source/LiveLinkGraphNode/Private/K2Node_UpdateVirtualSubjectDataBase.cpp
+++ b/Engine/Plugins/Animation/LiveLink/Source/LiveLinkGraphNode/Private/K2Node_UpdateVirtualSubjectDataBase.cpp
@@ -10,6 +10,7 @@
 #include "LiveLinkRole.h"
 #include "Kismet2/BlueprintEditorUtils.h"
 #include "VirtualSubjects/LiveLinkBlueprintVirtualSubject.h"
+#include "Styling/AppStyle.h"
 
 #define LOCTEXT_NAMESPACE "K2Node_UpdateVirtualSubjectDataBase"
 
diff --git a/Engine/Plugins/EnhancedInput/Source/InputBlueprintNodes/Private/K2Node_InputDebugKey.cpp b/Engine/Plugins/EnhancedInput/Source/InputBlueprintNodes/Private/K2Node_InputDebugKey.cpp
--- a/Engine/Plugins/EnhancedInput/Source/InputBlueprintNodes/Private/K2Node_InputDebugKey.cpp
+++ b/Engine/Plugins/EnhancedInput/Source/InputBlueprintNodes/Private/K2Node_InputDebugKey.cpp
@@ -13,6 +13,7 @@
 #include "EditorCategoryUtils.h"
 #include "BlueprintActionDatabaseRegistrar.h"
 #include "InputActionValue.h"
+#include "Styling/AppStyle.h"
 #include "GameFramework/InputSettings.h"
 #include "Editor.h"								// for FEditorDelegates::OnEnableGestureRecognizerChanged
 
diff --git a/Engine/Plugins/Experimental/Animation/CurveExpression/Source/Editor/Private/K2Node_MakeCurveExpressionMap.cpp b/Engine/Plugins/Experimental/Animation/CurveExpression/Source/Editor/Private/K2Node_MakeCurveExpressionMap.cpp
--- a/Engine/Plugins/Experimental/Animation/CurveExpression/Source/Editor/Private/K2Node_MakeCurveExpressionMap.cpp
+++ b/Engine/Plugins/Experimental/Animation/CurveExpression/Source/Editor/Private/K2Node_MakeCurveExpressionMap.cpp
@@ -10,6 +10,7 @@
 #include "KismetCompilerMisc.h"
 #include "String/ParseLines.h"
 #include "Textures/SlateIcon.h"
+#include "Styling/AppStyle.h"
 
 #define LOCTEXT_NAMESPACE "MakeCurveExpressionMap"
 
diff --git a/Engine/Source/Editor/BlueprintGraph/Private/K2Node_EnumLiteral.cpp b/Engine/Source/Editor/BlueprintGraph/Private/K2Node_EnumLiteral.cpp
--- a/Engine/Source/Editor/BlueprintGraph/Private/K2Node_EnumLiteral.cpp
+++ b/Engine/Source/Editor/BlueprintGraph/Private/K2Node_EnumLiteral.cpp
@@ -14,6 +14,7 @@
 #include "EditorCategoryUtils.h"
 #include "FindInBlueprintManager.h"
 #include "BlueprintActionDatabaseRegistrar.h"
+#include "Styling/AppStyle.h"
 
 const FName UK2Node_EnumLiteral::GetEnumInputPinName()
 {
diff --git a/Engine/Source/Editor/BlueprintGraph/Private/K2Node_Event.cpp b/Engine/Source/Editor/BlueprintGraph/Private/K2Node_Event.cpp
--- a/Engine/Source/Editor/BlueprintGraph/Private/K2Node_Event.cpp
+++ b/Engine/Source/Editor/BlueprintGraph/Private/K2Node_Event.cpp
@@ -17,6 +17,7 @@
 #include "KismetCompiler.h"
 #include "EventEntryHandler.h"
 #include "DiffResults.h"
+#include "Styling/AppStyle.h"
 
 const FName UK2Node_Event::DelegateOutputName(TEXT("OutputDelegate"));
 
diff --git a/Engine/Source/Editor/BlueprintGraph/Private/K2Node_ExecutionSequence.cpp b/Engine/Source/Editor/BlueprintGraph/Private/K2Node_ExecutionSequence.cpp
--- a/Engine/Source/Editor/BlueprintGraph/Private/K2Node_ExecutionSequence.cpp
+++ b/Engine/Source/Editor/BlueprintGraph/Private/K2Node_ExecutionSequence.cpp
@@ -9,6 +9,7 @@
 #include "BlueprintNodeSpawner.h"
 #include "EditorCategoryUtils.h"
 #include "BlueprintActionDatabaseRegistrar.h"
+#include "Styling/AppStyle.h"
 
 #define LOCTEXT_NAMESPACE "K2Node_MultiGate"
 
diff --git a/Engine/Source/Editor/BlueprintGraph/Private/K2Node_GetInputAxisKeyValue.cpp b/Engine/Source/Editor/BlueprintGraph/Private/K2Node_GetInputAxisKeyValue.cpp
--- a/Engine/Source/Editor/BlueprintGraph/Private/K2Node_GetInputAxisKeyValue.cpp
+++ b/Engine/Source/Editor/BlueprintGraph/Private/K2Node_GetInputAxisKeyValue.cpp
@@ -9,6 +9,7 @@
 #include "EditorCategoryUtils.h"
 #include "Engine/InputAxisKeyDelegateBinding.h"
 #include "BlueprintActionDatabaseRegistrar.h"
+#include "Styling/AppStyle.h"
 
 #define LOCTEXT_NAMESPACE "K2Node_GetInputAxisKeyValue"
 
diff --git a/Engine/Source/Editor/BlueprintGraph/Private/K2Node_GetNumEnumEntries.cpp b/Engine/Source/Editor/BlueprintGraph/Private/K2Node_GetNumEnumEntries.cpp
--- a/Engine/Source/Editor/BlueprintGraph/Private/K2Node_GetNumEnumEntries.cpp
+++ b/Engine/Source/Editor/BlueprintGraph/Private/K2Node_GetNumEnumEntries.cpp
@@ -9,6 +9,7 @@
 #include "BlueprintFieldNodeSpawner.h"
 #include "EditorCategoryUtils.h"
 #include "BlueprintActionDatabaseRegistrar.h"
+#include "Styling/AppStyle.h"
 
 UK2Node_GetNumEnumEntries::UK2Node_GetNumEnumEntries(const FObjectInitializer& ObjectInitializer)
 	: Super(ObjectInitializer)
diff --git a/Engine/Source/Editor/BlueprintGraph/Private/K2Node_InputKey.cpp b/Engine/Source/Editor/BlueprintGraph/Private/K2Node_InputKey.cpp
--- a/Engine/Source/Editor/BlueprintGraph/Private/K2Node_InputKey.cpp
+++ b/Engine/Source/Editor/BlueprintGraph/Private/K2Node_InputKey.cpp
@@ -12,6 +12,7 @@
 #include "BlueprintNodeSpawner.h"
 #include "EditorCategoryUtils.h"
 #include "BlueprintActionDatabaseRegistrar.h"
+#include "Styling/AppStyle.h"
 #include "GameFramework/InputSettings.h"
 #include "Editor.h"								// for FEditorDelegates::OnEnableGestureRecognizerChanged
 
diff --git a/Engine/Source/Editor/BlueprintGraph/Private/K2Node_InputTouch.cpp b/Engine/Source/Editor/BlueprintGraph/Private/K2Node_InputTouch.cpp
--- a/Engine/Source/Editor/BlueprintGraph/Private/K2Node_InputTouch.cpp
+++ b/Engine/Source/Editor/BlueprintGraph/Private/K2Node_InputTouch.cpp
@@ -12,6 +12,7 @@
 #include "BlueprintNodeSpawner.h"
 #include "EditorCategoryUtils.h"
 #include "BlueprintActionDatabaseRegistrar.h"
+#include "Styling/AppStyle.h"
 
 UK2Node_InputTouch::UK2Node_InputTouch(const FObjectInitializer& ObjectInitializer)
 	: Super(ObjectInitializer)
diff --git a/Engine/Source/Editor/BlueprintGraph/Private/K2Node_Select.cpp b/Engine/Source/Editor/BlueprintGraph/Private/K2Node_Select.cpp
--- a/Engine/Source/Editor/BlueprintGraph/Private/K2Node_Select.cpp
+++ b/Engine/Source/Editor/BlueprintGraph/Private/K2Node_Select.cpp
@@ -17,6 +17,7 @@
 #include "BlueprintNodeSpawner.h"
 #include "EditorCategoryUtils.h"
 #include "BlueprintActionDatabaseRegistrar.h"
+#include "Styling/AppStyle.h"
 
 #define LOCTEXT_NAMESPACE "K2Node_Select"
 
diff --git a/Engine/Source/Editor/BlueprintGraph/Private/K2Node_SpawnActor.cpp b/Engine/Source/Editor/BlueprintGraph/Private/K2Node_SpawnActor.cpp
--- a/Engine/Source/Editor/BlueprintGraph/Private/K2Node_SpawnActor.cpp
+++ b/Engine/Source/Editor/BlueprintGraph/Private/K2Node_SpawnActor.cpp
@@ -11,6 +11,7 @@
 #include "Kismet2/BlueprintEditorUtils.h"
 #include "KismetCompilerMisc.h"
 #include "KismetCompiler.h"
+#include "Styling/AppStyle.h"
 
 struct FK2Node_SpawnActorHelper
 {
diff --git a/Engine/Source/Editor/BlueprintGraph/Private/K2Node_SpawnActorFromClass.cpp b/Engine/Source/Editor/BlueprintGraph/Private/K2Node_SpawnActorFromClass.cpp
--- a/Engine/Source/Editor/BlueprintGraph/Private/K2Node_SpawnActorFromClass.cpp
+++ b/Engine/Source/Editor/BlueprintGraph/Private/K2Node_SpawnActorFromClass.cpp
@@ -16,6 +16,7 @@
 #include "EditorCategoryUtils.h"
 #include "BlueprintActionDatabaseRegistrar.h"
 #include "K2Node_EnumLiteral.h"
+#include "Styling/AppStyle.h"
 
 struct FK2Node_SpawnActorFromClassHelper
 {
diff --git a/Engine/Source/Editor/BlueprintGraph/Private/K2Node_Timeline.cpp b/Engine/Source/Editor/BlueprintGraph/Private/K2Node_Timeline.cpp
--- a/Engine/Source/Editor/BlueprintGraph/Private/K2Node_Timeline.cpp
+++ b/Engine/Source/Editor/BlueprintGraph/Private/K2Node_Timeline.cpp
@@ -23,6 +23,7 @@
 #include "KismetCastingUtils.h"
 #include "KismetCompilerMisc.h"
 #include "KismetCompiler.h"
+#include "Styling/AppStyle.h"
 
 #define LOCTEXT_NAMESPACE "K2Node_Timeline"
 
