Move LiveLink setting "PresetSaveDir" into its own editor per project user setting

diff --git a/Engine/Plugins/Animation/LiveLink/Source/LiveLink/Private/LiveLinkSettings.cpp b/Engine/Plugins/Animation/LiveLink/Source/LiveLink/Private/LiveLinkSettings.cpp
--- a/Engine/Plugins/Animation/LiveLink/Source/LiveLink/Private/LiveLinkSettings.cpp
+++ b/Engine/Plugins/Animation/LiveLink/Source/LiveLink/Private/LiveLinkSettings.cpp
@@ -36,3 +36,24 @@ FLiveLinkRoleProjectSetting ULiveLinkSettings::GetDefaultSettingForRole(TSubclas
 	Result.Role = Role;
 	return Result;
 }
+
+void ULiveLinkSettings::PostInitProperties()
+{
+	Super::PostInitProperties();
+
+#if WITH_EDITOR
+PRAGMA_DISABLE_DEPRECATION_WARNINGS
+	// If the deprecated preset save directory has a valid path, update the replacement setting in LiveLinkUserSettings with its value
+	if (!PresetSaveDir_DEPRECATED.Path.IsEmpty())
+	{
+		ULiveLinkUserSettings* UserSettings = GetMutableDefault<ULiveLinkUserSettings>();
+		UserSettings->PresetSaveDir = PresetSaveDir_DEPRECATED;
+		UserSettings->SaveConfig();
+
+		// Empty the path of the deprecated setting so that this won't overwrite the user setting again
+		PresetSaveDir_DEPRECATED.Path.Empty();
+		SaveConfig();
+	}
+PRAGMA_ENABLE_DEPRECATION_WARNINGS
+#endif //WITH_EDITOR
+}
diff --git a/Engine/Plugins/Animation/LiveLink/Source/LiveLink/Public/LiveLinkSettings.h b/Engine/Plugins/Animation/LiveLink/Source/LiveLink/Public/LiveLinkSettings.h
--- a/Engine/Plugins/Animation/LiveLink/Source/LiveLink/Public/LiveLinkSettings.h
+++ b/Engine/Plugins/Animation/LiveLink/Source/LiveLink/Public/LiveLinkSettings.h
@@ -48,6 +48,19 @@ public:
 	TArray<TSubclassOf<ULiveLinkFramePreProcessor>> FramePreProcessors;
 };
 
+UCLASS(config=EditorPerProjectUserSettings)
+class LIVELINK_API ULiveLinkUserSettings : public UObject
+{
+	GENERATED_BODY()
+
+public:
+	/** The default location in which to save LiveLink presets */
+	UPROPERTY(config, EditAnywhere, Category = "LiveLink", meta = (DisplayName = "Preset Save Location"))
+	FDirectoryPath PresetSaveDir;
+
+public:
+	const FDirectoryPath& GetPresetSaveDir() const { return PresetSaveDir; }
+};
 
 /**
  * Settings for LiveLink.
@@ -60,6 +73,10 @@ class LIVELINK_API ULiveLinkSettings : public UObject
 public:
 	ULiveLinkSettings();
 
+	//~ Begin UObject Interface
+	virtual void PostInitProperties() override;
+	//~ End UObject Interface
+
 protected:
 	UPROPERTY(config, EditAnywhere, Category="LiveLink")
 	TArray<FLiveLinkRoleProjectSetting> DefaultRoleSettings;
@@ -73,9 +90,11 @@ public:
 	UPROPERTY(config, EditAnywhere, Category = "LiveLink")
 	TSoftObjectPtr<ULiveLinkPreset> DefaultLiveLinkPreset;
 
-	/** The default location in which to save take presets */
-	UPROPERTY(config, EditAnywhere, Category="LiveLink", meta=(DisplayName="Preset Save Location"))
-	FDirectoryPath PresetSaveDir;
+#if WITH_EDITORONLY_DATA
+	UE_DEPRECATED(5.1, "PresetSaveDir was moved into LiveLinkUserSettings. Please use ULiveLinkUserSettings::GetPresetSaveDir().")
+	UPROPERTY(config)
+	FDirectoryPath PresetSaveDir_DEPRECATED;
+#endif //WITH_EDITORONLY_DATA
 
 	/** Continuous clock offset correction step */
 	UPROPERTY(config, EditAnywhere, AdvancedDisplay, Category = "LiveLink")
@@ -128,7 +147,10 @@ public:
 
 public:
 	FLiveLinkRoleProjectSetting GetDefaultSettingForRole(TSubclassOf<ULiveLinkRole> Role) const;
-	const FDirectoryPath& GetPresetSaveDir() const { return PresetSaveDir; }
+
+	UE_DEPRECATED(5.1, "PresetSaveDir was moved into LiveLinkUserSettings. Please use ULiveLinkUserSettings::GetPresetSaveDir().")
+	const FDirectoryPath& GetPresetSaveDir() const { return GetDefault<ULiveLinkUserSettings>()->GetPresetSaveDir(); }
+
 	double GetTimeWithoutFrameToBeConsiderAsInvalid() const { return TimeWithoutFrameToBeConsiderAsInvalid; }
 	FLinearColor GetValidColor() const { return ValidColor; }
 	FLinearColor GetInvalidColor() const { return InvalidColor; }
diff --git a/Engine/Plugins/Animation/LiveLink/Source/LiveLinkEditor/Private/LiveLinkClientPanelToolbar.cpp b/Engine/Plugins/Animation/LiveLink/Source/LiveLinkEditor/Private/LiveLinkClientPanelToolbar.cpp
--- a/Engine/Plugins/Animation/LiveLink/Source/LiveLinkEditor/Private/LiveLinkClientPanelToolbar.cpp
+++ b/Engine/Plugins/Animation/LiveLink/Source/LiveLinkEditor/Private/LiveLinkClientPanelToolbar.cpp
@@ -798,7 +798,7 @@ static bool OpenSaveDialog(const FString& InDefaultPath, const FString& InNewNam
 
 bool GetSavePresetPackageName(FString& OutName)
 {
-	ULiveLinkSettings* ConfigSettings = GetMutableDefault<ULiveLinkSettings>();
+	ULiveLinkUserSettings* ConfigSettings = GetMutableDefault<ULiveLinkUserSettings>();
 
 	FDateTime Today = FDateTime::Now();
 
