Fix an attempt to load the TimeManagement module during shutdown phase.

diff --git a/Engine/Plugins/Animation/LiveLink/Source/LiveLink/Private/LiveLinkTimedDataInput.cpp b/Engine/Plugins/Animation/LiveLink/Source/LiveLink/Private/LiveLinkTimedDataInput.cpp
--- a/Engine/Plugins/Animation/LiveLink/Source/LiveLink/Private/LiveLinkTimedDataInput.cpp
+++ b/Engine/Plugins/Animation/LiveLink/Source/LiveLink/Private/LiveLinkTimedDataInput.cpp
@@ -96,7 +96,11 @@ FLiveLinkTimedDataInput::FLiveLinkTimedDataInput(FLiveLinkClient* InClient, FGui
 
 FLiveLinkTimedDataInput::~FLiveLinkTimedDataInput()
 {
-	ITimeManagementModule::Get().GetTimedDataInputCollection().Remove(this);
+	ITimeManagementModule* TimeManagementModule = FModuleManager::GetModulePtr<ITimeManagementModule>("TimeManagement");
+	if (TimeManagementModule)
+	{
+		TimeManagementModule->GetTimedDataInputCollection().Remove(this);
+	}
 }
 
 FText FLiveLinkTimedDataInput::GetDisplayName() const
