- Fixing crash when adding livelink component

diff --git a/Engine/Plugins/Animation/LiveLink/Source/LiveLinkComponents/Private/LiveLinkComponentController.cpp b/Engine/Plugins/Animation/LiveLink/Source/LiveLinkComponents/Private/LiveLinkComponentController.cpp
--- a/Engine/Plugins/Animation/LiveLink/Source/LiveLinkComponents/Private/LiveLinkComponentController.cpp
+++ b/Engine/Plugins/Animation/LiveLink/Source/LiveLinkComponents/Private/LiveLinkComponentController.cpp
@@ -349,18 +349,21 @@ TArray<TSubclassOf<ULiveLinkRole>> ULiveLinkComponentController::GetSelectedRole
 {
 	TArray<TSubclassOf<ULiveLinkRole>> ClassHierarchy;
 
-	for (TObjectIterator<UClass> It; It; ++It)
+	if (InCurrentRoleClass)
 	{
-		if (!It->HasAnyClassFlags(CLASS_Abstract | CLASS_Deprecated))
+		for (TObjectIterator<UClass> It; It; ++It)
 		{
-			if (InCurrentRoleClass->IsChildOf(*It))
+			if (!It->HasAnyClassFlags(CLASS_Abstract | CLASS_Deprecated))
 			{
-				ClassHierarchy.AddUnique(*It);
+				if (InCurrentRoleClass->IsChildOf(*It))
+				{
+					ClassHierarchy.AddUnique(*It);
+				}
 			}
 		}
 	}
 
-	return MoveTemp(ClassHierarchy);
+	return ClassHierarchy;
 }
 
 TSubclassOf<ULiveLinkControllerBase> ULiveLinkComponentController::GetControllerClassForRoleClass(const TSubclassOf<ULiveLinkRole> RoleClass) const
