PR #9192: Fixed LiveLinkLightController not setting the LightComponent's "Use Temperature" (Contributed by [REDACTED_PERSON])

diff --git a/Engine/Plugins/Animation/LiveLink/Source/LiveLinkComponents/Private/Controllers/LiveLinkLightController.cpp b/Engine/Plugins/Animation/LiveLink/Source/LiveLinkComponents/Private/Controllers/LiveLinkLightController.cpp
--- a/Engine/Plugins/Animation/LiveLink/Source/LiveLinkComponents/Private/Controllers/LiveLinkLightController.cpp
+++ b/Engine/Plugins/Animation/LiveLink/Source/LiveLinkComponents/Private/Controllers/LiveLinkLightController.cpp
@@ -26,6 +26,10 @@ void ULiveLinkLightController::Tick(float DeltaTime, const FLiveLinkSubjectFrame
 	{
 		if (ULightComponent* LightComponent = Cast<ULightComponent>(GetAttachedComponent()))
 		{
+			if (StaticData->bIsTemperatureSupported != LightComponent->bUseTemperature)
+			{
+				LightComponent->SetUseTemperature(StaticData->bIsTemperatureSupported);
+			}
 			if (StaticData->bIsTemperatureSupported) { LightComponent->SetTemperature(FrameData->Temperature); }
 			if (StaticData->bIsIntensitySupported) { LightComponent->SetIntensity(FrameData->Intensity); }
 			if (StaticData->bIsLightColorSupported) { LightComponent->SetLightColor(FrameData->LightColor); }
