- Adding a dependency to LiveLinkAnimationCore in LL plugin to make sure it's pulled when packaging and avoid crashing when LL anim node is used

diff --git a/Engine/Plugins/Animation/LiveLink/Source/LiveLink/LiveLink.Build.cs b/Engine/Plugins/Animation/LiveLink/Source/LiveLink/LiveLink.Build.cs
--- a/Engine/Plugins/Animation/LiveLink/Source/LiveLink/LiveLink.Build.cs
+++ b/Engine/Plugins/Animation/LiveLink/Source/LiveLink/LiveLink.Build.cs
@@ -21,6 +21,7 @@ namespace UnrealBuildTool.Rules
 			{
 				"HeadMountedDisplay",
 				"InputCore",
+				"LiveLinkAnimationCore", //Adding dependency to make sure this runtime module is pulled when packaging if LiveLink plugin is enabled. 
 				"Media",
 				"Projects",
 				"SlateCore",
