- fixing typo in livelink log

diff --git a/Engine/Plugins/Animation/LiveLink/Source/LiveLink/Private/LiveLinkLogInstance.cpp b/Engine/Plugins/Animation/LiveLink/Source/LiveLink/Private/LiveLinkLogInstance.cpp
--- a/Engine/Plugins/Animation/LiveLink/Source/LiveLink/Private/LiveLinkLogInstance.cpp
+++ b/Engine/Plugins/Animation/LiveLink/Source/LiveLink/Private/LiveLinkLogInstance.cpp
@@ -292,7 +292,7 @@ void FLiveLinkLogInstance::LogMessage(EMessageSeverity::Type Severity, FName Mes
 
 			if (bLog)
 			{
-				Message = FString::Printf(TEXT("Occured %d time. Last occurence %s. %s")
+				Message = FString::Printf(TEXT("Occurred %d time. Last occurrence %s. %s")
 					, FoundRepeatableMessage->Data->Occurrence.Load()
 					, *(FoundRepeatableMessage->Data->LastTimeOccured.Load().ToString())
 					, *Message);
