- Removing UE4 reference in livelink plugin description

diff --git a/Engine/Plugins/Animation/LiveLink/LiveLink.uplugin b/Engine/Plugins/Animation/LiveLink/LiveLink.uplugin
--- a/Engine/Plugins/Animation/LiveLink/LiveLink.uplugin
+++ b/Engine/Plugins/Animation/LiveLink/LiveLink.uplugin
@@ -3,7 +3,7 @@
 	"Version": 1,
 	"VersionName": "2.0",
 	"FriendlyName": "Live Link",
-	"Description": "LiveLink allows streaming of animated data into UE4",
+	"Description": "LiveLink allows streaming of animated data into Unreal Engine",
 	"Category": "Animation",
 	"CreatedBy": "Epic Games, Inc.",
 	"CreatedByURL": "http://epicgames.com",
