Remove FLiveLinkMessageBusDiscoveryManager thread from client

diff --git a/Engine/Plugins/Animation/LiveLink/Source/LiveLink/Private/ILiveLinkModule.h b/Engine/Plugins/Animation/LiveLink/Source/LiveLink/Private/ILiveLinkModule.h
--- a/Engine/Plugins/Animation/LiveLink/Source/LiveLink/Private/ILiveLinkModule.h
+++ b/Engine/Plugins/Animation/LiveLink/Source/LiveLink/Private/ILiveLinkModule.h
@@ -6,6 +6,10 @@
 #include "Containers/Array.h"
 #include "Modules/ModuleManager.h"
 
+#ifndef WITH_LIVELINK_DISCOVERY_MANAGER_THREAD
+#define WITH_LIVELINK_DISCOVERY_MANAGER_THREAD 1
+#endif
+
 class FLiveLinkHeartbeatEmitter;
 class FLiveLinkMessageBusDiscoveryManager;
 class FSlateStyleSet;
@@ -35,7 +39,9 @@ public:
 
 	virtual TSharedPtr<FSlateStyleSet> GetStyle() = 0;
 	virtual FLiveLinkHeartbeatEmitter& GetHeartbeatEmitter() = 0;
+#if WITH_LIVELINK_DISCOVERY_MANAGER_THREAD
 	virtual FLiveLinkMessageBusDiscoveryManager& GetMessageBusDiscoveryManager() = 0;
+#endif
 
 public:
 
diff --git a/Engine/Plugins/Animation/LiveLink/Source/LiveLink/Private/LiveLinkMessageBusSource.cpp b/Engine/Plugins/Animation/LiveLink/Source/LiveLink/Private/LiveLinkMessageBusSource.cpp
--- a/Engine/Plugins/Animation/LiveLink/Source/LiveLink/Private/LiveLinkMessageBusSource.cpp
+++ b/Engine/Plugins/Animation/LiveLink/Source/LiveLink/Private/LiveLinkMessageBusSource.cpp
@@ -7,7 +7,9 @@
 #include "LiveLinkClient.h"
 #include "LiveLinkHeartbeatEmitter.h"
 #include "LiveLinkLog.h"
+#if WITH_LIVELINK_DISCOVERY_MANAGER_THREAD
 #include "LiveLinkMessageBusDiscoveryManager.h"
+#endif
 #include "LiveLinkMessageBusSourceSettings.h"
 #include "LiveLinkMessages.h"
 #include "LiveLinkRoleTrait.h"
@@ -49,7 +51,9 @@ void FLiveLinkMessageBusSource::ReceiveClient(ILiveLinkClient* InClient, FGuid I
 	}
 	else
 	{
+#if WITH_LIVELINK_DISCOVERY_MANAGER_THREAD
 		ILiveLinkModule::Get().GetMessageBusDiscoveryManager().AddDiscoveryMessageRequest();
+#endif
 		bIsValid = false;
 	}
 
@@ -60,6 +64,7 @@ void FLiveLinkMessageBusSource::Update()
 {
 	if (!ConnectionAddress.IsValid())
 	{
+#if WITH_LIVELINK_DISCOVERY_MANAGER_THREAD
 		FLiveLinkMessageBusDiscoveryManager& DiscoveryManager = ILiveLinkModule::Get().GetMessageBusDiscoveryManager();
 		for (const FProviderPollResultPtr& Result : DiscoveryManager.GetDiscoveryResults())
 		{
@@ -74,6 +79,7 @@ void FLiveLinkMessageBusSource::Update()
 				break;
 			}
 		}
+#endif
 	}
 	else
 	{
@@ -296,11 +302,13 @@ void FLiveLinkMessageBusSource::SendConnectMessage()
 
 bool FLiveLinkMessageBusSource::RequestSourceShutdown()
 {
+#if WITH_LIVELINK_DISCOVERY_MANAGER_THREAD
 	FLiveLinkMessageBusDiscoveryManager& DiscoveryManager = ILiveLinkModule::Get().GetMessageBusDiscoveryManager();
 	if (DiscoveryManager.IsRunning() && !ConnectionAddress.IsValid())
 	{
 		DiscoveryManager.RemoveDiscoveryMessageRequest();
 	}
+#endif
 
 	FLiveLinkHeartbeatEmitter& HeartbeatEmitter = ILiveLinkModule::Get().GetHeartbeatEmitter();
 	HeartbeatEmitter.StopHeartbeat(ConnectionAddress, MessageEndpoint);
diff --git a/Engine/Plugins/Animation/LiveLink/Source/LiveLink/Private/LiveLinkModule.cpp b/Engine/Plugins/Animation/LiveLink/Source/LiveLink/Private/LiveLinkModule.cpp
--- a/Engine/Plugins/Animation/LiveLink/Source/LiveLink/Private/LiveLinkModule.cpp
+++ b/Engine/Plugins/Animation/LiveLink/Source/LiveLink/Private/LiveLinkModule.cpp
@@ -18,7 +18,9 @@ FLiveLinkModule::FLiveLinkModule()
 	: LiveLinkClient()
 	, LiveLinkMotionController(LiveLinkClient)
 	, HeartbeatEmitter(MakeUnique<FLiveLinkHeartbeatEmitter>())
+#if WITH_LIVELINK_DISCOVERY_MANAGER_THREAD
 	, DiscoveryManager(MakeUnique<FLiveLinkMessageBusDiscoveryManager>())
+#endif
 	, LiveLinkDebugCommand(MakeUnique<FLiveLinkDebugCommand>(LiveLinkClient))
 {
 }
@@ -41,7 +43,9 @@ void FLiveLinkModule::ShutdownModule()
 	FCoreDelegates::OnFEngineLoopInitComplete.RemoveAll(this);
 
 	HeartbeatEmitter->Exit();
+#if WITH_LIVELINK_DISCOVERY_MANAGER_THREAD
 	DiscoveryManager->Stop();
+#endif
 	LiveLinkMotionController.UnregisterController();
 
 	IModularFeatures::Get().UnregisterModularFeature(FLiveLinkClient::ModularFeatureName, &LiveLinkClient);
diff --git a/Engine/Plugins/Animation/LiveLink/Source/LiveLink/Private/LiveLinkModule.h b/Engine/Plugins/Animation/LiveLink/Source/LiveLink/Private/LiveLinkModule.h
--- a/Engine/Plugins/Animation/LiveLink/Source/LiveLink/Private/LiveLinkModule.h
+++ b/Engine/Plugins/Animation/LiveLink/Source/LiveLink/Private/LiveLinkModule.h
@@ -33,7 +33,9 @@ public:
 
 	//~ Begin ILiveLinkModule interface
 	virtual FLiveLinkHeartbeatEmitter& GetHeartbeatEmitter() override { return *HeartbeatEmitter; }
+#if WITH_LIVELINK_DISCOVERY_MANAGER_THREAD
 	virtual FLiveLinkMessageBusDiscoveryManager& GetMessageBusDiscoveryManager() override { return *DiscoveryManager; }
+#endif
 	virtual TSharedPtr<FSlateStyleSet> GetStyle() override { return StyleSet; }
 	//~ End ILiveLinkModule interface
 
diff --git a/Engine/Plugins/Animation/LiveLink/Source/LiveLink/Private/SLiveLinkMessageBusSourceFactory.cpp b/Engine/Plugins/Animation/LiveLink/Source/LiveLink/Private/SLiveLinkMessageBusSourceFactory.cpp
--- a/Engine/Plugins/Animation/LiveLink/Source/LiveLink/Private/SLiveLinkMessageBusSourceFactory.cpp
+++ b/Engine/Plugins/Animation/LiveLink/Source/LiveLink/Private/SLiveLinkMessageBusSourceFactory.cpp
@@ -4,7 +4,9 @@
 
 #include "ILiveLinkClient.h"
 #include "ILiveLinkModule.h"
+#if WITH_LIVELINK_DISCOVERY_MANAGER_THREAD
 #include "LiveLinkMessageBusDiscoveryManager.h"
+#endif
 #include "LiveLinkMessageBusFinder.h"
 #include "Widgets/Layout/SBox.h"
 
@@ -76,17 +78,21 @@ private:
 
 SLiveLinkMessageBusSourceFactory::~SLiveLinkMessageBusSourceFactory()
 {
+#if WITH_LIVELINK_DISCOVERY_MANAGER_THREAD
 	if (ILiveLinkModule* ModulePtr = FModuleManager::GetModulePtr<ILiveLinkModule>("LiveLink"))
 	{
 		ModulePtr->GetMessageBusDiscoveryManager().RemoveDiscoveryMessageRequest();
 	}
+#endif
 }
 
 void SLiveLinkMessageBusSourceFactory::Construct(const FArguments& Args)
 {
 	OnSourceSelected = Args._OnSourceSelected;
 
+#if WITH_LIVELINK_DISCOVERY_MANAGER_THREAD
 	ILiveLinkModule::Get().GetMessageBusDiscoveryManager().AddDiscoveryMessageRequest();
+#endif
 
 	ChildSlot
 	[
@@ -123,7 +129,9 @@ void SLiveLinkMessageBusSourceFactory::Tick(const FGeometry& AllottedGeometry, c
 	if (FApp::GetCurrentTime() - LastUIUpdateSeconds > 0.5)
 	{
 		PollData.Reset();
+#if WITH_LIVELINK_DISCOVERY_MANAGER_THREAD
 		PollData.Append(ILiveLinkModule::Get().GetMessageBusDiscoveryManager().GetDiscoveryResults());
+#endif
 		
 		Sources.RemoveAllSwap([this](TSharedPtr<FLiveLinkSource> Source)
 		{
