LiveLink: Prevent LiveLink component CDO from trying to PostLoad anything in its ControllerMap

diff --git a/Engine/Plugins/Animation/LiveLink/Source/LiveLinkComponents/Private/LiveLinkComponentController.cpp b/Engine/Plugins/Animation/LiveLink/Source/LiveLinkComponents/Private/LiveLinkComponentController.cpp
--- a/Engine/Plugins/Animation/LiveLink/Source/LiveLinkComponents/Private/LiveLinkComponentController.cpp
+++ b/Engine/Plugins/Animation/LiveLink/Source/LiveLinkComponents/Private/LiveLinkComponentController.cpp
@@ -257,6 +257,11 @@ void ULiveLinkComponentController::PostLoad()
 {
 	Super::PostLoad();
 
+	if (HasAnyFlags(RF_ClassDefaultObject | RF_ArchetypeObject))
+	{
+		return;
+	}
+
 #if WITH_EDITOR
  	const int32 Version = GetLinkerCustomVersion(FUE5MainStreamObjectVersion::GUID);
  	if (Version < FUE5MainStreamObjectVersion::LiveLinkComponentPickerPerController)
