[LiveLink]
* LiveLinkDebugCommand - Fixed non editor non unity/pch compile error

diff --git a/Engine/Plugins/Animation/LiveLink/Source/LiveLink/Private/LiveLinkDebugCommand.cpp b/Engine/Plugins/Animation/LiveLink/Source/LiveLink/Private/LiveLinkDebugCommand.cpp
--- a/Engine/Plugins/Animation/LiveLink/Source/LiveLink/Private/LiveLinkDebugCommand.cpp
+++ b/Engine/Plugins/Animation/LiveLink/Source/LiveLink/Private/LiveLinkDebugCommand.cpp
@@ -9,6 +9,8 @@
 #include "IAssetViewport.h"
 #include "LevelEditor.h"
 #include "Modules/ModuleManager.h"
+#else
+#include "Engine/Engine.h"
 #endif
 
 static const FName LevelEditorName("LevelEditor");
