Fixed non unity/pch errors in LiveLink, OnlineFramework and OnlineSubsystemMcp

diff --git a/Engine/Plugins/Animation/LiveLink/Source/LiveLink/Private/LiveLinkHeartbeatEmitter.cpp b/Engine/Plugins/Animation/LiveLink/Source/LiveLink/Private/LiveLinkHeartbeatEmitter.cpp
--- a/Engine/Plugins/Animation/LiveLink/Source/LiveLink/Private/LiveLinkHeartbeatEmitter.cpp
+++ b/Engine/Plugins/Animation/LiveLink/Source/LiveLink/Private/LiveLinkHeartbeatEmitter.cpp
@@ -23,6 +23,8 @@ FLiveLinkHeartbeatEmitter::FLiveLinkHeartbeatEmitter()
 
 }
 
+FLiveLinkHeartbeatEmitter::~FLiveLinkHeartbeatEmitter() = default;
+
 void FLiveLinkHeartbeatEmitter::StartHeartbeat(const FMessageAddress& RecipientAddress, const TSharedPtr<FMessageEndpoint, ESPMode::ThreadSafe>& MessageEndpoint)
 {
 	if (!bIsRunning)
diff --git a/Engine/Plugins/Animation/LiveLink/Source/LiveLink/Private/LiveLinkHeartbeatEmitter.h b/Engine/Plugins/Animation/LiveLink/Source/LiveLink/Private/LiveLinkHeartbeatEmitter.h
--- a/Engine/Plugins/Animation/LiveLink/Source/LiveLink/Private/LiveLinkHeartbeatEmitter.h
+++ b/Engine/Plugins/Animation/LiveLink/Source/LiveLink/Private/LiveLinkHeartbeatEmitter.h
@@ -23,6 +23,7 @@ class FLiveLinkHeartbeatEmitter : public FRunnable
 {
 public:
 	FLiveLinkHeartbeatEmitter();
+	~FLiveLinkHeartbeatEmitter();
 
 	/**
 	 * Start sending a heartbeat to the specified Recipient by using the given MessageEndpoint.
diff --git a/Engine/Plugins/Online/OnlineFramework/Source/Party/Private/SocialManager.cpp b/Engine/Plugins/Online/OnlineFramework/Source/Party/Private/SocialManager.cpp
--- a/Engine/Plugins/Online/OnlineFramework/Source/Party/Private/SocialManager.cpp
+++ b/Engine/Plugins/Online/OnlineFramework/Source/Party/Private/SocialManager.cpp
@@ -21,6 +21,10 @@
 
 #include UE_INLINE_GENERATED_CPP_BY_NAME(SocialManager)
 
+#if PARTY_PLATFORM_SESSIONS_XBL 
+#include "Misc/Base64.h"
+#endif
+
 static TAutoConsoleVariable<bool> CVarForceDisconnectedToPartyService(
 	TEXT("SocialUI.ForceDisconnectedToPartyService"),
 	false,
