[LiveLink]
* FLiveLinkPropertyData - fixed dllexport of ctor/dtor

diff --git a/Engine/Plugins/Animation/LiveLink/Source/LiveLinkMovieScene/Public/MovieScene/MovieSceneLiveLinkStructProperties.h b/Engine/Plugins/Animation/LiveLink/Source/LiveLinkMovieScene/Public/MovieScene/MovieSceneLiveLinkStructProperties.h
--- a/Engine/Plugins/Animation/LiveLink/Source/LiveLinkMovieScene/Public/MovieScene/MovieSceneLiveLinkStructProperties.h
+++ b/Engine/Plugins/Animation/LiveLink/Source/LiveLinkMovieScene/Public/MovieScene/MovieSceneLiveLinkStructProperties.h
@@ -24,8 +24,8 @@ struct FLiveLinkPropertyData
 	GENERATED_BODY()
 
 public:
-	FLiveLinkPropertyData();
-	~FLiveLinkPropertyData();
+	LIVELINKMOVIESCENE_API FLiveLinkPropertyData();
+	LIVELINKMOVIESCENE_API ~FLiveLinkPropertyData();
 
 	UPROPERTY()
 	FName PropertyName;
