Update UseSourceTimecode tooltip

diff --git a/Engine/Plugins/Animation/LiveLink/Source/LiveLinkSequencer/Private/TakeRecorderSource/TakeRecorderLiveLinkSource.h b/Engine/Plugins/Animation/LiveLink/Source/LiveLinkSequencer/Private/TakeRecorderSource/TakeRecorderLiveLinkSource.h
--- a/Engine/Plugins/Animation/LiveLink/Source/LiveLinkSequencer/Private/TakeRecorderSource/TakeRecorderLiveLinkSource.h
+++ b/Engine/Plugins/Animation/LiveLink/Source/LiveLinkSequencer/Private/TakeRecorderSource/TakeRecorderLiveLinkSource.h
@@ -65,7 +65,10 @@ public:
 	UPROPERTY(EditAnywhere, BlueprintReadWrite, Category = "Subjects")
 	bool bSaveSubjectSettings;
 
-	/** If true we use timecode even if not synchronized, else we use world time*/
+	/**
+	 * Whether the livelink subject's timecode or the system time should be used for the recording.
+	 * @note If set, the livelink subject's timecode will be used even if it does not match the engine timecode.
+	 */
 	UPROPERTY(EditAnywhere, BlueprintReadWrite, Category = "Subjects")
 	bool bUseSourceTimecode;
 
