LiveLink: Fix MovieScene LiveLink sources to always use a buffer size of 1 frame (fixes bug with MRQ renders where the wrong frame is chosen)

diff --git a/Engine/Plugins/Animation/LiveLink/Source/LiveLinkMovieScene/Private/MovieScene/MovieSceneLiveLinkSource.cpp b/Engine/Plugins/Animation/LiveLink/Source/LiveLinkMovieScene/Private/MovieScene/MovieSceneLiveLinkSource.cpp
--- a/Engine/Plugins/Animation/LiveLink/Source/LiveLinkMovieScene/Private/MovieScene/MovieSceneLiveLinkSource.cpp
+++ b/Engine/Plugins/Animation/LiveLink/Source/LiveLinkMovieScene/Private/MovieScene/MovieSceneLiveLinkSource.cpp
@@ -3,6 +3,7 @@
 #include "MovieSceneLiveLinkSource.h"
 #include "Features/IModularFeatures.h"
 #include "ILiveLinkClient.h"
+#include "ILiveLinkSource.h"
 
 FMovieSceneLiveLinkSource::FMovieSceneLiveLinkSource()
 	: Client(nullptr)
@@ -68,6 +69,13 @@ void FMovieSceneLiveLinkSource::ReceiveClient(ILiveLinkClient* InClient, FGuid I
 	Client->CreateSubject(SubjectPreset);
 }
 
+void FMovieSceneLiveLinkSource::InitializeSettings(ULiveLinkSourceSettings* Settings)
+{
+	ILiveLinkSource::InitializeSettings(Settings);
+
+	Settings->BufferSettings.MaxNumberOfFrameToBuffered = 1;
+}
+
 bool FMovieSceneLiveLinkSource::IsSourceStillValid() const
 {
 	return Client != nullptr;
diff --git a/Engine/Plugins/Animation/LiveLink/Source/LiveLinkMovieScene/Private/MovieScene/MovieSceneLiveLinkSource.h b/Engine/Plugins/Animation/LiveLink/Source/LiveLinkMovieScene/Private/MovieScene/MovieSceneLiveLinkSource.h
--- a/Engine/Plugins/Animation/LiveLink/Source/LiveLinkMovieScene/Private/MovieScene/MovieSceneLiveLinkSource.h
+++ b/Engine/Plugins/Animation/LiveLink/Source/LiveLinkMovieScene/Private/MovieScene/MovieSceneLiveLinkSource.h
@@ -21,6 +21,7 @@ public:
 
 	// ILiveLinkSource interface
 	virtual void ReceiveClient(ILiveLinkClient* InClient, FGuid InSourceGuid) override;
+	virtual void InitializeSettings(ULiveLinkSourceSettings* Settings) override;
 	virtual bool IsSourceStillValid() const override;
 	virtual bool RequestSourceShutdown() override;
 	virtual FText GetSourceMachineName() const override;
