LiveLink: Fix warning when no editable components are found for LiveLink to control

diff --git a/Engine/Plugins/Animation/LiveLink/Source/LiveLinkComponents/Private/LiveLinkComponentController.cpp b/Engine/Plugins/Animation/LiveLink/Source/LiveLinkComponents/Private/LiveLinkComponentController.cpp
--- a/Engine/Plugins/Animation/LiveLink/Source/LiveLinkComponents/Private/LiveLinkComponentController.cpp
+++ b/Engine/Plugins/Animation/LiveLink/Source/LiveLinkComponents/Private/LiveLinkComponentController.cpp
@@ -417,6 +417,11 @@ void ULiveLinkComponentController::InitializeController(ULiveLinkControllerBase*
 		TInlineComponentArray<UActorComponent*> ActorComponents;
 		OuterActor->GetComponents(InController->GetDesiredComponentClass(), ActorComponents);
 
+		if (ActorComponents.Num() == 0)
+		{
+			UE_LOG(LogLiveLinkComponents, Warning, TEXT("The desired component class for %s is %s, but %s does not have a component of that type."), *InController->GetName(), *InController->GetDesiredComponentClass()->GetName(), *OuterActor->GetActorLabel(false));
+		}
+
 		bool bFoundValidComponent = false;
 
 		// Look through the list of components matching the desired component class, and choose the first editable instance
@@ -433,7 +438,7 @@ void ULiveLinkComponentController::InitializeController(ULiveLinkControllerBase*
 
 		if (!bFoundValidComponent)
 		{
-			UE_LOG(LogLiveLinkComponents, Warning, TEXT("The desired component class for %s is %s, but %s does not have a component of that type."), *InController->GetName(), *InController->GetDesiredComponentClass()->GetName(), *OuterActor->GetActorLabel(false));
+			UE_LOG(LogLiveLinkComponents, Verbose, TEXT("%s has no editable components of type %s."), *OuterActor->GetActorLabel(false), *InController->GetDesiredComponentClass()->GetName());
 		}
 	}
 #endif //WITH_EDITOR
