Fix LiveLinkHub compilation errors caused by some modules checking for target type editor instead of CompilingAgainstEditor

diff --git a/Engine/Plugins/Animation/LiveLink/Source/LiveLinkComponents/LiveLinkComponents.Build.cs b/Engine/Plugins/Animation/LiveLink/Source/LiveLinkComponents/LiveLinkComponents.Build.cs
--- a/Engine/Plugins/Animation/LiveLink/Source/LiveLinkComponents/LiveLinkComponents.Build.cs
+++ b/Engine/Plugins/Animation/LiveLink/Source/LiveLinkComponents/LiveLinkComponents.Build.cs
@@ -15,7 +15,7 @@ namespace UnrealBuildTool.Rules
 				"LiveLinkInterface",
 			});
 
-			if (Target.Type == TargetType.Editor)
+			if (Target.bCompileAgainstEditor)
 			{
 				PrivateDependencyModuleNames.AddRange(
 				new string[]
diff --git a/Engine/Plugins/Experimental/QuicMessaging/Source/QuicMessaging/QuicMessaging.Build.cs b/Engine/Plugins/Experimental/QuicMessaging/Source/QuicMessaging/QuicMessaging.Build.cs
--- a/Engine/Plugins/Experimental/QuicMessaging/Source/QuicMessaging/QuicMessaging.Build.cs
+++ b/Engine/Plugins/Experimental/QuicMessaging/Source/QuicMessaging/QuicMessaging.Build.cs
@@ -50,7 +50,7 @@ namespace UnrealBuildTool.Rules
 				}
 			);
 
-			if (Target.Type == TargetType.Editor)
+			if (Target.bCompileAgainstEditor)
 			{
 
 				DynamicallyLoadedModuleNames.AddRange(
diff --git a/Engine/Plugins/Messaging/TcpMessaging/Source/TcpMessaging/TcpMessaging.Build.cs b/Engine/Plugins/Messaging/TcpMessaging/Source/TcpMessaging/TcpMessaging.Build.cs
--- a/Engine/Plugins/Messaging/TcpMessaging/Source/TcpMessaging/TcpMessaging.Build.cs
+++ b/Engine/Plugins/Messaging/TcpMessaging/Source/TcpMessaging/TcpMessaging.Build.cs
@@ -37,7 +37,7 @@ namespace UnrealBuildTool.Rules
 					"TcpMessaging/Private/Transport",
 				});
 
-			if (Target.Type == TargetType.Editor)
+			if (Target.bCompileAgainstEditor)
 			{
 				DynamicallyLoadedModuleNames.AddRange(
 					new string[] {
diff --git a/Engine/Plugins/Messaging/UdpMessaging/Source/UdpMessaging/UdpMessaging.Build.cs b/Engine/Plugins/Messaging/UdpMessaging/Source/UdpMessaging/UdpMessaging.Build.cs
--- a/Engine/Plugins/Messaging/UdpMessaging/Source/UdpMessaging/UdpMessaging.Build.cs
+++ b/Engine/Plugins/Messaging/UdpMessaging/Source/UdpMessaging/UdpMessaging.Build.cs
@@ -30,7 +30,7 @@ namespace UnrealBuildTool.Rules
 					"MessagingCommon",
 				});
 
-			if (Target.Type == TargetType.Editor)
+			if (Target.bCompileAgainstEditor)
 			{
 				DynamicallyLoadedModuleNames.AddRange(
 					new string[] {
diff --git a/Engine/Source/Developer/AnimationDataController/AnimationDataController.build.cs b/Engine/Source/Developer/AnimationDataController/AnimationDataController.build.cs
--- a/Engine/Source/Developer/AnimationDataController/AnimationDataController.build.cs
+++ b/Engine/Source/Developer/AnimationDataController/AnimationDataController.build.cs
@@ -22,7 +22,7 @@ namespace UnrealBuildTool.Rules
 				}
 			);
 			
-			if (Target.Type == TargetType.Editor)
+			if (Target.bCompileAgainstEditor)
 			{
 				PrivateDependencyModuleNames.AddRange(
 					new string[]
diff --git a/Engine/Source/Runtime/GameplayTags/GameplayTags.Build.cs b/Engine/Source/Runtime/GameplayTags/GameplayTags.Build.cs
--- a/Engine/Source/Runtime/GameplayTags/GameplayTags.Build.cs
+++ b/Engine/Source/Runtime/GameplayTags/GameplayTags.Build.cs
@@ -25,7 +25,7 @@ namespace UnrealBuildTool.Rules
 				}
 			);
 
-			if (Target.Type == TargetType.Editor)
+			if (Target.bCompileAgainstEditor)
             {
                 PrivateDependencyModuleNames.AddRange(
                 new string[]
