Add a stub livelink hub module to fix GenerateProjectFiles warning caused by missing module.

diff --git a/Engine/Plugins/Animation/LiveLink/LiveLink.uplugin b/Engine/Plugins/Animation/LiveLink/LiveLink.uplugin
--- a/Engine/Plugins/Animation/LiveLink/LiveLink.uplugin
+++ b/Engine/Plugins/Animation/LiveLink/LiveLink.uplugin
@@ -14,6 +14,7 @@
 	"CanContainContent": false,
 	"IsBetaVersion": false,
 	"Installed": false,
+	"SupportedPrograms": [ "LiveLinkHub" ],
 	"Plugins": [
 		{
 			"Name": "Takes",
@@ -23,14 +24,29 @@
 			"Name": "ConcertSyncClient",
 			"Enabled": true,
 			"Optional": true
+		},
+		{
+			"Name": "StructUtils",
+			"Enabled": true
+		},
+		{
+			"Name": "ContentBrowserAssetDataSource",
+			"Enabled":  true
 		}
 	],
 	"Modules": [
 		{
 			"Name": "LiveLink",
 			"Type": "Runtime",
+			"ProgramAllowList": [ "LiveLinkHub" ],
 			"LoadingPhase": "Default"
 		},
+		{
+			"Name": "LiveLinkHub",
+			"Type": "Runtime",
+			"ProgramAllowList": [ "LiveLinkHub" ],
+			"LoadingPhase": "PreDefault"
+		},
 		{
 			"Name": "LiveLinkComponents",
 			"Type": "Runtime",
@@ -39,7 +55,8 @@
 		{
 			"Name": "LiveLinkEditor",
 			"Type": "Editor",
-			"LoadingPhase": "Default"
+			"LoadingPhase": "Default",
+			"ProgramAllowList": [ "LiveLinkHub" ]
 		},
 		{
 			"Name": "LiveLinkGraphNode",
@@ -54,7 +71,8 @@
 		{
 			"Name": "LiveLinkSequencer",
 			"Type": "UncookedOnly",
-			"LoadingPhase": "Default"
+			"LoadingPhase": "Default",
+			"ProgramDenyList": [ "LiveLinkHub" ]
 		},
 		{
 			"Name": "LiveLinkMultiUser",
diff --git a/Engine/Plugins/Animation/LiveLink/Source/LiveLinkHub/LiveLinkHub.Build.cs b/Engine/Plugins/Animation/LiveLink/Source/LiveLinkHub/LiveLinkHub.Build.cs
new file mode 100644
--- /dev/null
+++ b/Engine/Plugins/Animation/LiveLink/Source/LiveLinkHub/LiveLinkHub.Build.cs
@@ -0,0 +1,22 @@
+// Copyright Epic Games, Inc. All Rights Reserved.
+
+namespace UnrealBuildTool.Rules
+{
+	public class LiveLinkHub : ModuleRules
+	{
+		public LiveLinkHub(ReadOnlyTargetRules Target) : base(Target)
+		{
+			PublicDependencyModuleNames.AddRange(
+			new string[]
+			{
+				"Core",
+				"CoreUObject",
+			});
+
+			PrivateDependencyModuleNames.AddRange(
+			new string[]
+			{
+			});
+		}
+	}
+}
diff --git a/Engine/Plugins/Animation/LiveLink/Source/LiveLinkHub/Private/LiveLinkHubModule.cpp b/Engine/Plugins/Animation/LiveLink/Source/LiveLinkHub/Private/LiveLinkHubModule.cpp
new file mode 100644
--- /dev/null
+++ b/Engine/Plugins/Animation/LiveLink/Source/LiveLinkHub/Private/LiveLinkHubModule.cpp
@@ -0,0 +1,11 @@
+// Copyright Epic Games, Inc. All Rights Reserved.
+
+#include "LiveLinkHubModule.h"
+#include "Modules/ModuleManager.h"
+
+void FLiveLinkHubModule::StartLiveLinkHub()
+{
+}
+
+IMPLEMENT_MODULE(FLiveLinkHubModule, LiveLinkHub);
+
diff --git a/Engine/Plugins/Animation/LiveLink/Source/LiveLinkHub/Private/LiveLinkHubModule.h b/Engine/Plugins/Animation/LiveLink/Source/LiveLinkHub/Private/LiveLinkHubModule.h
new file mode 100644
--- /dev/null
+++ b/Engine/Plugins/Animation/LiveLink/Source/LiveLinkHub/Private/LiveLinkHubModule.h
@@ -0,0 +1,13 @@
+// Copyright Epic Games, Inc. All Rights Reserved.
+
+#pragma once
+
+#include "ILiveLinkHubModule.h"
+
+class FLiveLinkHubModule : public ILiveLinkHubModule
+{
+public:
+	//~ Begin ILiveLinkHubModule interface
+	virtual void StartLiveLinkHub() override;
+	//~ End ILiveLinkHubModule interface
+};
diff --git a/Engine/Plugins/Animation/LiveLink/Source/LiveLinkHub/Public/ILiveLinkHubModule.h b/Engine/Plugins/Animation/LiveLink/Source/LiveLinkHub/Public/ILiveLinkHubModule.h
new file mode 100644
--- /dev/null
+++ b/Engine/Plugins/Animation/LiveLink/Source/LiveLinkHub/Public/ILiveLinkHubModule.h
@@ -0,0 +1,11 @@
+// Copyright Epic Games, Inc. All Rights Reserved.
+
+#include "Modules/ModuleInterface.h"
+
+class LIVELINKHUB_API ILiveLinkHubModule : public IModuleInterface
+{
+public:
+	/** Launch the slate application hosting the live link hub. */
+	virtual void StartLiveLinkHub() = 0;
+};
+
