Add missing header guard to ILiveLinkHub module

diff --git a/Engine/Plugins/Animation/LiveLink/Source/LiveLinkHub/Public/ILiveLinkHubModule.h b/Engine/Plugins/Animation/LiveLink/Source/LiveLinkHub/Public/ILiveLinkHubModule.h
--- a/Engine/Plugins/Animation/LiveLink/Source/LiveLinkHub/Public/ILiveLinkHubModule.h
+++ b/Engine/Plugins/Animation/LiveLink/Source/LiveLinkHub/Public/ILiveLinkHubModule.h
@@ -1,5 +1,7 @@
 // Copyright Epic Games, Inc. All Rights Reserved.
 
+#pragma once
+
 #include "Modules/ModuleInterface.h"
 
 class LIVELINKHUB_API ILiveLinkHubModule : public IModuleInterface
