Refactor LiveLink to support LiveLinkHub
Separate LiveLinkClientPanel's widgets into components that live within a LiveLinkPanelController, which allows LiveLinkHub to reuse these widgets in a different layout than the livelink client panel.

Move LiveLinkSourceCollection, LiveLinkSubject from Private to Internal
Move LiveLinkClientPanelToolbar, SLiveLinkDataView from Private to Internal
(Meaning that they can only be accessed from other engine modules)

diff --git a/Engine/Plugins/Animation/LiveLink/LiveLink.uplugin b/Engine/Plugins/Animation/LiveLink/LiveLink.uplugin
--- a/Engine/Plugins/Animation/LiveLink/LiveLink.uplugin
+++ b/Engine/Plugins/Animation/LiveLink/LiveLink.uplugin
@@ -43,7 +43,7 @@
 		},
 		{
 			"Name": "LiveLinkHub",
-			"Type": "Runtime",
+			"Type": "EditorAndProgram",
 			"ProgramAllowList": [ "LiveLinkHub" ],
 			"LoadingPhase": "PreDefault"
 		},
diff --git a/Engine/Plugins/Animation/LiveLink/Source/LiveLink/Internal/ClockOffsetEstimatorRamp.h b/Engine/Plugins/Animation/LiveLink/Source/LiveLink/Internal/ClockOffsetEstimatorRamp.h
new file mode 100644
--- /dev/null
+++ b/Engine/Plugins/Animation/LiveLink/Source/LiveLink/Internal/ClockOffsetEstimatorRamp.h
@@ -0,0 +1,52 @@
+// Copyright Epic Games, Inc. All Rights Reserved.
+
+#pragma once
+
+#include "HAL/Platform.h"
+
+/**
+ * Continually estimates the time offset between two live time sources.
+ *
+ * Samples of clock times corresponding to the same instant are provided,
+ * and it will update its estimation of the time difference between the clocks.
+ */
+class FClockOffsetEstimatorRamp
+{
+public:
+
+	/** Get current estimated offset between the two tracked clocks */
+	double GetEstimatedOffset() const;
+
+	/** Update the current time offset estimation with a new data sample */
+	void UpdateEstimation(double Atime, double Btime);
+
+	/** Set required consecutive big errors to clamp */
+	void SetRequiredConsecutiveBigError(int32 NewCount);
+
+	/** Set maximum allowed estimation error before looking to clamp */
+	void SetMaxAllowedEstimationError(double NewValue);
+
+	/** Sets incremental correction step */
+	void SetCorrectionStep(double NewValue);
+
+	/** Puts back the estimator in its initial state */
+	void Reset();
+
+protected:
+
+	/** Current estimated offset between two clocks */
+	double EstimatedOffset = 0.0;
+	
+
+	/** Maximum error tolerated before considering clamping to current offset */
+	double MaxAllowedEstimationError = 0.25;
+
+	/** Correction step to apply, positive or negative depending on error sign */
+	double CorrectionStep = 100e-6;
+
+	/** Consecutive time a too big estimation error is required before clamping */
+	int32 ConsecutiveTooBigErrorRequired = 5;
+	
+	/** Current consecutive big error count. Initialized at required count so first pass snaps */
+	int32 ConsecutiveBigErrorCount = 5;
+};
diff --git a/Engine/Plugins/Animation/LiveLink/Source/LiveLink/Internal/LiveLinkSourceCollection.h b/Engine/Plugins/Animation/LiveLink/Source/LiveLink/Internal/LiveLinkSourceCollection.h
new file mode 100644
--- /dev/null
+++ b/Engine/Plugins/Animation/LiveLink/Source/LiveLink/Internal/LiveLinkSourceCollection.h
@@ -0,0 +1,157 @@
+// Copyright Epic Games, Inc. All Rights Reserved.
+
+#pragma once
+
+
+#include "ILiveLinkClient.h"
+#include "LiveLinkSubject.h"
+#include "LiveLinkSubjectSettings.h"
+#include "LiveLinkVirtualSubject.h"
+#include "UObject/GCObject.h"
+
+class FLiveLinkSourceCollection;
+
+class FLiveLinkSubject;
+class ILiveLinkSource;
+class ILiveLinkSubject;
+
+class ULiveLinkSourceSettings;
+class ULiveLinkSubjectSettings;
+class ULiveLinkVirtualSubject;
+
+struct LIVELINK_API FLiveLinkCollectionSourceItem
+{
+	FLiveLinkCollectionSourceItem() = default;
+	FLiveLinkCollectionSourceItem(FLiveLinkCollectionSourceItem&&) = default;
+	FLiveLinkCollectionSourceItem& operator=(FLiveLinkCollectionSourceItem&&) = default;
+
+	FLiveLinkCollectionSourceItem(const FLiveLinkCollectionSourceItem&) = delete;
+	FLiveLinkCollectionSourceItem& operator=(const FLiveLinkCollectionSourceItem&) = delete;
+
+	FGuid Guid;
+	TObjectPtr<ULiveLinkSourceSettings> Setting; // GC by FLiveLinkSourceCollection::AddReferencedObjects
+	TSharedPtr<ILiveLinkSource> Source;
+	TSharedPtr<FLiveLinkTimedDataInput> TimedData;
+	bool bPendingKill = false;
+	bool bIsVirtualSource = false;
+
+public:
+	bool IsVirtualSource() const;
+};
+
+
+struct LIVELINK_API FLiveLinkCollectionSubjectItem
+{
+	FLiveLinkCollectionSubjectItem(FLiveLinkSubjectKey InKey, TUniquePtr<FLiveLinkSubject> InLiveSubject, ULiveLinkSubjectSettings* InSettings, bool bInEnabled);
+	FLiveLinkCollectionSubjectItem(FLiveLinkSubjectKey InKey, ULiveLinkVirtualSubject* InVirtualSubject, bool bInEnabled);
+
+public:
+	FLiveLinkSubjectKey Key;
+	bool bEnabled;
+	bool bPendingKill;
+
+	ILiveLinkSubject* GetSubject() { return VirtualSubject ? static_cast<ILiveLinkSubject*>(VirtualSubject) : static_cast<ILiveLinkSubject*>(LiveSubject.Get()); }
+	FLiveLinkSubject* GetLiveSubject() { return LiveSubject.Get(); }
+	ULiveLinkVirtualSubject* GetVirtualSubject() { return VirtualSubject; }
+	ILiveLinkSubject* GetSubject() const { return VirtualSubject ? VirtualSubject : static_cast<ILiveLinkSubject*>(LiveSubject.Get()); }
+	FLiveLinkSubject* GetLiveSubject() const { return LiveSubject.Get(); }
+	ULiveLinkVirtualSubject* GetVirtualSubject() const { return VirtualSubject; }
+	UObject* GetSettings() const { return VirtualSubject ? static_cast<UObject*>(VirtualSubject) : static_cast<UObject*>(Setting); }
+	ULiveLinkSubjectSettings* GetLinkSettings() const { return Setting; }
+
+private:
+	TObjectPtr<ULiveLinkSubjectSettings> Setting; // GC by FLiveLinkSourceCollection::AddReferencedObjects
+	TUniquePtr<FLiveLinkSubject> LiveSubject;
+	TObjectPtr<ULiveLinkVirtualSubject> VirtualSubject; // GC by FLiveLinkSourceCollection::AddReferencedObjects
+
+public:
+	FLiveLinkCollectionSubjectItem(const FLiveLinkCollectionSubjectItem&) = delete;
+	FLiveLinkCollectionSubjectItem& operator=(const FLiveLinkCollectionSubjectItem&) = delete;
+	FLiveLinkCollectionSubjectItem(FLiveLinkCollectionSubjectItem&&) = default;
+	FLiveLinkCollectionSubjectItem& operator=(FLiveLinkCollectionSubjectItem&&) = default;
+
+	friend FLiveLinkSourceCollection;
+};
+
+
+class LIVELINK_API FLiveLinkSourceCollection : public FGCObject
+{
+public:
+	// "source guid" for virtual subjects
+	static const FGuid DefaultVirtualSubjectGuid;
+	FLiveLinkSourceCollection();
+
+	UE_NONCOPYABLE(FLiveLinkSourceCollection)
+
+public:
+	//~ Begin FGCObject implementation
+	virtual void AddReferencedObjects(FReferenceCollector & Collector) override;
+	virtual FString GetReferencerName() const override
+	{
+		return TEXT("FLiveLinkSourceCollection");
+	}
+	//~ End FGCObject implementation
+
+public:
+
+	
+	TArray<FLiveLinkCollectionSourceItem>& GetSources() { return Sources; }
+	const TArray<FLiveLinkCollectionSourceItem>& GetSources() const { return Sources; }
+	const TArray<FLiveLinkCollectionSubjectItem>& GetSubjects() const { return Subjects; }
+
+	void AddSource(FLiveLinkCollectionSourceItem Source);
+	void RemoveSource(FGuid SourceGuid);
+	void RemoveAllSources();
+	FLiveLinkCollectionSourceItem* FindSource(TSharedPtr<ILiveLinkSource> Source);
+	const FLiveLinkCollectionSourceItem* FindSource(TSharedPtr<ILiveLinkSource> Source) const;
+	FLiveLinkCollectionSourceItem* FindSource(FGuid SourceGuid);
+	const FLiveLinkCollectionSourceItem* FindSource(FGuid SourceGuid) const;
+	FLiveLinkCollectionSourceItem* FindVirtualSource(FName VirtualSourceName);
+	const FLiveLinkCollectionSourceItem* FindVirtualSource(FName VirtualSourceName) const;
+
+	void AddSubject(FLiveLinkCollectionSubjectItem Subject);
+	void RemoveSubject(FLiveLinkSubjectKey SubjectKey);
+	FLiveLinkCollectionSubjectItem* FindSubject(FLiveLinkSubjectKey SubjectKey);
+	const FLiveLinkCollectionSubjectItem* FindSubject(FLiveLinkSubjectKey SubjectKey) const;
+	const FLiveLinkCollectionSubjectItem* FindEnabledSubject(FLiveLinkSubjectName SubjectName) const;
+
+	bool IsSubjectEnabled(FLiveLinkSubjectKey SubjectKey) const;
+	void SetSubjectEnabled(FLiveLinkSubjectKey SubjectKey, bool bEnabled);
+
+	void RemovePendingKill();
+	bool RequestShutdown();
+
+	FSimpleMulticastDelegate& OnLiveLinkSourcesChanged() { return OnLiveLinkSourcesChangedDelegate; }
+	FSimpleMulticastDelegate& OnLiveLinkSubjectsChanged() { return OnLiveLinkSubjectsChangedDelegate; }
+
+	FOnLiveLinkSourceChangedDelegate& OnLiveLinkSourceAdded() { return OnLiveLinkSourceAddedDelegate; }
+	FOnLiveLinkSourceChangedDelegate& OnLiveLinkSourceRemoved() { return OnLiveLinkSourceRemovedDelegate; }
+	FOnLiveLinkSubjectChangedDelegate& OnLiveLinkSubjectAdded() { return OnLiveLinkSubjectAddedDelegate; }
+	FOnLiveLinkSubjectChangedDelegate& OnLiveLinkSubjectRemoved() { return OnLiveLinkSubjectRemovedDelegate; }
+
+private:
+	void RemoveSource(int32 Index);
+
+private:
+	TArray<FLiveLinkCollectionSourceItem> Sources;
+
+	TArray<FLiveLinkCollectionSubjectItem> Subjects;
+
+	/** Notify when the client sources list has changed */
+	FSimpleMulticastDelegate OnLiveLinkSourcesChangedDelegate;
+
+	/** Notify when a client subjects list has changed */
+	FSimpleMulticastDelegate OnLiveLinkSubjectsChangedDelegate;
+
+	/** Notify when a client source's is added */
+	FOnLiveLinkSourceChangedDelegate OnLiveLinkSourceAddedDelegate;
+
+	/** Notify when a client source's is removed */
+	FOnLiveLinkSourceChangedDelegate OnLiveLinkSourceRemovedDelegate;
+
+	/** Notify when a client subject's is added */
+	FOnLiveLinkSubjectChangedDelegate OnLiveLinkSubjectAddedDelegate;
+
+	/** Notify when a client subject's is removed */
+	FOnLiveLinkSubjectChangedDelegate OnLiveLinkSubjectRemovedDelegate;
+};
diff --git a/Engine/Plugins/Animation/LiveLink/Source/LiveLink/Internal/LiveLinkSubject.h b/Engine/Plugins/Animation/LiveLink/Source/LiveLink/Internal/LiveLinkSubject.h
new file mode 100644
--- /dev/null
+++ b/Engine/Plugins/Animation/LiveLink/Source/LiveLink/Internal/LiveLinkSubject.h
@@ -0,0 +1,224 @@
+// Copyright Epic Games, Inc. All Rights Reserved.
+
+#pragma once
+
+#include "ILiveLinkSubject.h"
+
+#include "Containers/Queue.h"
+#include "ITimedDataInput.h"
+#include "LiveLinkFrameInterpolationProcessor.h"
+#include "LiveLinkFramePreProcessor.h"
+#include "LiveLinkSourceSettings.h"
+
+class FLiveLinkTimedDataInput;
+class ULiveLinkSubjectSettings;
+struct FLiveLinkSubjectTimeSyncData;
+
+
+struct FLiveLinkInterpolationInfo;
+
+struct FLiveLinkTimeSynchronizationData
+{
+	// Whether or not synchronization has been established.
+	bool bHasEstablishedSync = false;
+
+	// The frame in our buffer where a rollover was detected. Only applicable for time synchronized sources.
+	int32 RolloverFrame = INDEX_NONE;
+
+	// Frame offset that will be used for this source.
+	int32 Offset = 0;
+
+	// Frame Time value modulus. When this value is not set, we assume no rollover occurs.
+	TOptional<FFrameTime> RolloverModulus;
+
+	// Frame rate used as the base for synchronization.
+	FFrameRate SyncFrameRate;
+
+	// Frame time that synchronization was established (relative to SynchronizationFrameRate).
+	FFrameTime SyncStartTime;
+};
+
+/**
+ * Manages subject manipulation either to add or get frame data for specific roles
+ */
+class LIVELINK_API FLiveLinkSubject : public ILiveLinkSubject, public ITimedDataInputChannel
+{
+private:
+	using Super = ILiveLinkSubject;
+
+public:
+	explicit FLiveLinkSubject(TSharedPtr<FLiveLinkTimedDataInput> InTimedDataGroup);
+	FLiveLinkSubject(const FLiveLinkSubject&) = delete;
+	FLiveLinkSubject& operator=(const FLiveLinkSubject&) = delete;
+	virtual ~FLiveLinkSubject();
+
+	//~ Begin ILiveLinkSubject Interface
+public:
+	virtual void Initialize(FLiveLinkSubjectKey SubjectKey, TSubclassOf<ULiveLinkRole> InRole, ILiveLinkClient* LiveLinkClient) override;
+	virtual void Update() override;
+	virtual void ClearFrames() override;
+	virtual FLiveLinkSubjectKey GetSubjectKey() const override { return SubjectKey; }
+	virtual TSubclassOf<ULiveLinkRole> GetRole() const override { return Role; }
+	virtual bool HasValidFrameSnapshot() const override;
+	virtual FLiveLinkStaticDataStruct& GetStaticData() override { return StaticData; }
+	virtual const FLiveLinkStaticDataStruct& GetStaticData() const override { return StaticData; }
+	virtual TArray<FLiveLinkTime> GetFrameTimes() const override;
+	virtual const TArray<ULiveLinkFrameTranslator::FWorkerSharedPtr> GetFrameTranslators() const override { return FrameTranslators; }
+	virtual bool IsRebroadcasted() const override { return bRebroadcastSubject; }
+	virtual bool HasStaticDataBeenRebroadcasted() const override { return bRebroadcastStaticDataSent; }
+	virtual void SetStaticDataAsRebroadcasted(const bool bInSent) override { bRebroadcastStaticDataSent = bInSent; }
+protected:
+	virtual const FLiveLinkSubjectFrameData& GetFrameSnapshot() const override { return FrameSnapshot; }
+	//~ End ILiveLinkSubject Interface
+
+	//~Begin ITimedDataSource Interface
+public:
+	virtual FText GetDisplayName() const override;
+	virtual ETimedDataInputState GetState() const override;
+	virtual FTimedDataChannelSampleTime GetOldestDataTime() const override;
+	virtual FTimedDataChannelSampleTime GetNewestDataTime() const override;
+	virtual TArray<FTimedDataChannelSampleTime> GetDataTimes() const override;
+	virtual int32 GetNumberOfSamples() const override;
+	virtual bool IsBufferStatsEnabled() const override;
+	virtual void SetBufferStatsEnabled(bool bEnable) override;
+	virtual int32 GetBufferUnderflowStat() const override;
+	virtual int32 GetBufferOverflowStat() const override;
+	virtual int32 GetFrameDroppedStat() const override;
+	virtual void GetLastEvaluationData(FTimedDataInputEvaluationData& OutEvaluationData) const override;
+	virtual void ResetBufferStats() override;
+	//~End ITimedDataSource Interface
+
+public:
+	bool EvaluateFrameAtWorldTime(double InWorldTime, TSubclassOf<ULiveLinkRole> InDesiredRole, FLiveLinkSubjectFrameData& OutFrame);
+	bool EvaluateFrameAtSceneTime(const FQualifiedFrameTime& InSceneTime, TSubclassOf<ULiveLinkRole> InDesiredRole, FLiveLinkSubjectFrameData& OutFrame);
+
+	bool HasStaticData() const;
+
+	// Handling setting a new static data. Create role data if not found in map.
+	void SetStaticData(TSubclassOf<ULiveLinkRole> InRole, FLiveLinkStaticDataStruct&& InStaticData);
+
+	// Add a frame of data from a FLiveLinkFrameData
+	void AddFrameData(FLiveLinkFrameDataStruct&& InFrameData);
+
+	void CacheSettings(ULiveLinkSourceSettings* SourceSetting, ULiveLinkSubjectSettings* SubjectSetting);
+
+	ELiveLinkSourceMode GetMode() const { return CachedSettings.SourceMode; }
+	FLiveLinkSubjectTimeSyncData GetTimeSyncData();
+	bool IsTimeSynchronized() const;
+
+	double GetLastPushTime() const { return LastPushTime; }
+
+private:
+	int32 FindNewFrame_WorldTime(const FLiveLinkWorldTime& FrameTime) const;
+	int32 FindNewFrame_WorldTimeInternal(const FLiveLinkWorldTime& FrameTime) const;
+	int32 FindNewFrame_SceneTime(const FQualifiedFrameTime& FrameTime, const FLiveLinkWorldTime& WorldTime) const;
+	int32 FindNewFrame_Latest(const FLiveLinkWorldTime& FrameTime) const;
+
+	// Reorder frame with the same timecode and create subframes
+	void AdjustSubFrame_SceneTime(int32 FrameIndex);
+
+	// Populate OutFrame with a frame based off of the supplied time (pre offsetted)
+	bool GetFrameAtWorldTime(const double InSeconds, FLiveLinkSubjectFrameData& OutFrame);
+	bool GetFrameAtWorldTime_Closest(const double InSeconds, FLiveLinkSubjectFrameData& OutFrame);
+	bool GetFrameAtWorldTime_Interpolated(const double InSeconds, FLiveLinkSubjectFrameData& OutFrame);
+
+	// Populate OutFrame with a frame based off of the supplied scene time (pre offsetted).
+	bool GetFrameAtSceneTime(const FQualifiedFrameTime& InSceneTime, FLiveLinkSubjectFrameData& OutFrame);
+	bool GetFrameAtSceneTime_Closest(const FQualifiedFrameTime& InSceneTime, FLiveLinkSubjectFrameData& OutFrame);
+	bool GetFrameAtSceneTime_Interpolated(const FQualifiedFrameTime& InSceneTime, FLiveLinkSubjectFrameData& OutFrame);
+	
+	// Verify interpolation result to update our internal statistics
+	void VerifyInterpolationInfo(const FLiveLinkInterpolationInfo& InterpolationInfo);
+
+	// Populate OutFrame with the latest frame.
+	bool GetLatestFrame(FLiveLinkSubjectFrameData& OutFrame);
+
+	void ResetFrame(FLiveLinkSubjectFrameData& OutFrame) const;
+
+	// Update our internal statistics
+	void IncreaseFrameDroppedStat();
+	void IncreaseBufferUnderFlowStat();
+	void IncreaseBufferOverFlowStat();
+	void UpdateEvaluationData(const FTimedDataInputEvaluationData& EvaluationData);
+
+	//Remove frames from our buffer - based on receiving order
+	void RemoveFrames(int32 Count);
+
+protected:
+	// The role the subject was build with
+	TSubclassOf<ULiveLinkRole> Role;
+
+	TArray<ULiveLinkFramePreProcessor::FWorkerSharedPtr> FramePreProcessors;
+
+	ULiveLinkFrameInterpolationProcessor::FWorkerSharedPtr FrameInterpolationProcessor = nullptr;
+
+	/** List of available translator the subject can use. */
+	TArray<ULiveLinkFrameTranslator::FWorkerSharedPtr> FrameTranslators;
+
+private:
+	struct FLiveLinkCachedSettings
+	{
+		ELiveLinkSourceMode SourceMode = ELiveLinkSourceMode::EngineTime;
+		FLiveLinkSourceBufferManagementSettings BufferSettings;
+	};
+	
+	struct FSubjectEvaluationStatistics
+	{
+		TAtomic<int32> BufferUnderflow;
+		TAtomic<int32> BufferOverflow;
+		TAtomic<int32> FrameDrop;
+		FTimedDataInputEvaluationData LastEvaluationData;
+
+		FSubjectEvaluationStatistics();
+		FSubjectEvaluationStatistics(const FSubjectEvaluationStatistics&) = delete;
+		FSubjectEvaluationStatistics& operator=(const FSubjectEvaluationStatistics&) = delete;
+	};
+
+	// Static data of the subject
+	FLiveLinkStaticDataStruct StaticData;
+
+	// Frames added to the subject
+	TArray<FLiveLinkFrameDataStruct> FrameData;
+
+	// Contains identifier of each frame in the order they were received
+	TQueue<FLiveLinkFrameIdentifier> ReceivedOrderedFrames;
+
+	// Next identifier to assign to next received frame
+	FLiveLinkFrameIdentifier NextIdentifier = 0;
+
+	// Current frame snapshot of the evaluation
+	FLiveLinkSubjectFrameData FrameSnapshot;
+
+	// Name of the subject
+	FLiveLinkSubjectKey SubjectKey;
+
+	// Timed data input group for the subject
+	TWeakPtr<FLiveLinkTimedDataInput> TimedDataGroup;
+
+	// Connection settings specified by user
+	FLiveLinkCachedSettings CachedSettings;
+
+	// Last time a frame was pushed
+	double LastPushTime = 0.0;
+
+	// Logging stats is enabled by default. If monitor opens at a later stage,previous stats will be able to be seen
+	bool bIsStatLoggingEnabled = true;
+
+	// Some stats compiled by the subject.
+	FSubjectEvaluationStatistics EvaluationStatistics;
+
+	// Last Timecode FrameRate received
+	FFrameRate LastTimecodeFrameRate;
+
+	// If enabled, rebroadcast this subject
+    bool bRebroadcastSubject = false;
+
+	// If true, static data has been sent for this rebroadcast
+	bool bRebroadcastStaticDataSent = false;
+	
+	/** 
+	 * Evaluation can be done on any thread so we need to protect statistic logging 
+	 * Some stats requires more than atomic sized vars so a critical section is used to protect when necessary
+	 */
+	mutable FCriticalSection StatisticCriticalSection;
+};
diff --git a/Engine/Plugins/Animation/LiveLink/Source/LiveLink/Internal/LiveLinkTimedDataInput.h b/Engine/Plugins/Animation/LiveLink/Source/LiveLink/Internal/LiveLinkTimedDataInput.h
new file mode 100644
--- /dev/null
+++ b/Engine/Plugins/Animation/LiveLink/Source/LiveLink/Internal/LiveLinkTimedDataInput.h
@@ -0,0 +1,69 @@
+// Copyright Epic Games, Inc. All Rights Reserved.
+
+#pragma once
+
+#include "ITimedDataInput.h"
+
+#include "ClockOffsetEstimatorRamp.h"
+
+class FLiveLinkClient;
+enum class ELiveLinkSourceMode : uint8;
+struct FLiveLinkBaseFrameData;
+
+class LIVELINK_API FLiveLinkTimedDataInput : public ITimedDataInput
+{
+public:
+	FLiveLinkTimedDataInput(FLiveLinkClient* Client, FGuid Source);
+	FLiveLinkTimedDataInput(const FLiveLinkTimedDataInput&) = delete;
+	FLiveLinkTimedDataInput& operator=(const FLiveLinkTimedDataInput&) = delete;
+	virtual ~FLiveLinkTimedDataInput();
+
+	//~ Begin ITimedDataInput API
+	virtual FText GetDisplayName() const override;
+	virtual TArray<ITimedDataInputChannel*> GetChannels() const override;
+	virtual ETimedDataInputEvaluationType GetEvaluationType() const override;
+	virtual void SetEvaluationType(ETimedDataInputEvaluationType Evaluation) override;
+	virtual double GetEvaluationOffsetInSeconds() const override;
+	virtual void SetEvaluationOffsetInSeconds(double Offset) override;
+	virtual FFrameRate GetFrameRate() const override;
+	virtual int32 GetDataBufferSize() const override;
+	virtual void SetDataBufferSize(int32 BufferSize) override;
+	virtual bool IsDataBufferSizeControlledByInput() const override { return true; }
+	virtual void AddChannel(ITimedDataInputChannel* Channel) override { Channels.Add(Channel); }
+	virtual void RemoveChannel(ITimedDataInputChannel* Channel) override { Channels.RemoveSingleSwap(Channel); }
+#if WITH_EDITOR
+	virtual const FSlateBrush* GetDisplayIcon() const override;
+#endif
+	//~ End ITimedDataInput API
+
+public:
+	//Tracks clock difference between each frame received and arrival time in engine referential.
+	void ProcessNewFrameTimingInfo(FLiveLinkBaseFrameData& NewFrameData);
+
+private:
+	//Computes the smooth offset based on the average frame time intervals of the source time
+	void UpdateSmoothEngineTimeOffset(const FLiveLinkBaseFrameData& NewFrameData);
+
+private:
+	FLiveLinkClient* LiveLinkClient;
+	TArray<ITimedDataInputChannel*> Channels;
+	FGuid Source;
+
+	//Continuous clock offset estimator for engine time and timecode
+	FClockOffsetEstimatorRamp EngineClockOffset;
+	FClockOffsetEstimatorRamp TimecodeClockOffset;
+
+	//We will receive each frame for each subject of this source. Stamp last source time/timecode to only update our offset estimation once per "source frame"
+	double LastWorldSourceTime = 0.0;
+	double LastSceneTime = 0.0;
+
+	static constexpr double FrameIntervalThreshold = 0.005;
+	static constexpr double VeryLargeFrameIntervalThreshold = 0.5;
+	static constexpr int32 FrameIntervalSnapCount = 5;
+	static constexpr int32 FrameTimeBufferSize = 200;
+
+	TArray<double> FrameTimes;
+
+	int32 FrameIntervalChangeCount = 0;
+	int32 NumFramesToConsiderForAverage = TNumericLimits<int32>::Max();
+};
diff --git a/Engine/Plugins/Animation/LiveLink/Source/LiveLink/Private/ClockOffsetEstimatorRamp.h b/Engine/Plugins/Animation/LiveLink/Source/LiveLink/Private/ClockOffsetEstimatorRamp.h
deleted file mode 100644
diff --git a/Engine/Plugins/Animation/LiveLink/Source/LiveLink/Private/LiveLinkClientReference.cpp b/Engine/Plugins/Animation/LiveLink/Source/LiveLink/Private/LiveLinkClientReference.cpp
--- a/Engine/Plugins/Animation/LiveLink/Source/LiveLink/Private/LiveLinkClientReference.cpp
+++ b/Engine/Plugins/Animation/LiveLink/Source/LiveLink/Private/LiveLinkClientReference.cpp
@@ -1,9 +1,20 @@
 // Copyright Epic Games, Inc. All Rights Reserved.
 
 #include "LiveLinkClientReference.h"
+#include "ILiveLinkClient.h"
 #include "LiveLinkModule.h"
+#include "Features/IModularFeatures.h"
+
+#ifndef WITH_LIVELINK_HUB
+#define WITH_LIVELINK_HUB 0
+#endif
 
 ILiveLinkClient* FLiveLinkClientReference::GetClient()const
 {
+#if WITH_LIVELINK_HUB
+	// Compiling for LiveLinkHub should return the livelink client we registered as a modular feature.
+	return &IModularFeatures::Get().GetModularFeature<ILiveLinkClient>(ILiveLinkClient::ModularFeatureName);
+#else
 	return FLiveLinkModule::LiveLinkClient_AnyThread;
+#endif
 }
diff --git a/Engine/Plugins/Animation/LiveLink/Source/LiveLink/Private/LiveLinkModule.cpp b/Engine/Plugins/Animation/LiveLink/Source/LiveLink/Private/LiveLinkModule.cpp
--- a/Engine/Plugins/Animation/LiveLink/Source/LiveLink/Private/LiveLinkModule.cpp
+++ b/Engine/Plugins/Animation/LiveLink/Source/LiveLink/Private/LiveLinkModule.cpp
@@ -18,6 +18,11 @@
 LLM_DEFINE_TAG(LiveLink);
 #define LOCTEXT_NAMESPACE "LiveLinkModule"
 
+#ifndef WITH_LIVELINK_HUB
+#define WITH_LIVELINK_HUB 0
+#endif
+
+
 FLiveLinkClient* FLiveLinkModule::LiveLinkClient_AnyThread = nullptr;
 
 FLiveLinkModule::FLiveLinkModule()
@@ -37,8 +42,11 @@ void FLiveLinkModule::StartupModule()
 	FLiveLinkLogInstance::CreateInstance();
 	CreateStyle();
 
+#if !WITH_LIVELINK_HUB
 	FPlatformAtomics::InterlockedExchangePtr((void**)&LiveLinkClient_AnyThread, &LiveLinkClient);
 	IModularFeatures::Get().RegisterModularFeature(FLiveLinkClient::ModularFeatureName, &LiveLinkClient);
+#endif
+
 	LiveLinkMotionController.RegisterController();
 
 	//Register for engine initialization completed so we can load default preset if any. Presets could depend on plugins loaded at a later stage.
@@ -56,8 +64,11 @@ void FLiveLinkModule::ShutdownModule()
 #endif
 	LiveLinkMotionController.UnregisterController();
 
+
+#if !WITH_LIVELINK_HUB
 	IModularFeatures::Get().UnregisterModularFeature(FLiveLinkClient::ModularFeatureName, &LiveLinkClient);
 	FPlatformAtomics::InterlockedExchangePtr((void**)&LiveLinkClient_AnyThread, nullptr);
+#endif
 
 	FSlateStyleRegistry::UnRegisterSlateStyle(*StyleSet.Get());
 	FLiveLinkLogInstance::DestroyInstance();
diff --git a/Engine/Plugins/Animation/LiveLink/Source/LiveLink/Private/LiveLinkSourceCollection.h b/Engine/Plugins/Animation/LiveLink/Source/LiveLink/Private/LiveLinkSourceCollection.h
deleted file mode 100644
diff --git a/Engine/Plugins/Animation/LiveLink/Source/LiveLink/Private/LiveLinkSubject.h b/Engine/Plugins/Animation/LiveLink/Source/LiveLink/Private/LiveLinkSubject.h
deleted file mode 100644
diff --git a/Engine/Plugins/Animation/LiveLink/Source/LiveLink/Private/LiveLinkTimedDataInput.h b/Engine/Plugins/Animation/LiveLink/Source/LiveLink/Private/LiveLinkTimedDataInput.h
deleted file mode 100644
diff --git a/Engine/Plugins/Animation/LiveLink/Source/LiveLink/Public/LiveLinkClient.h b/Engine/Plugins/Animation/LiveLink/Source/LiveLink/Public/LiveLinkClient.h
--- a/Engine/Plugins/Animation/LiveLink/Source/LiveLink/Public/LiveLinkClient.h
+++ b/Engine/Plugins/Animation/LiveLink/Source/LiveLink/Public/LiveLinkClient.h
@@ -30,7 +30,7 @@ struct FLiveLinkSubjectTimeSyncData
 PRAGMA_DISABLE_DEPRECATION_WARNINGS
 struct FLiveLinkSkeletonStaticData;
 
-class FLiveLinkClient_Base_DEPRECATED : public ILiveLinkClient
+class LIVELINK_API FLiveLinkClient_Base_DEPRECATED : public ILiveLinkClient
 {
 public:
 	//~ Begin ILiveLinkClient implementation
@@ -78,6 +78,8 @@ public:
 	virtual TArray<FGuid> GetVirtualSources(bool bEvenIfPendingKill = false) const override;
 	virtual FLiveLinkSourcePreset GetSourcePreset(FGuid SourceGuid, UObject* DuplicatedObjectOuter) const override;
 	virtual FText GetSourceType(FGuid EntryGuid) const override;
+	virtual FText GetSourceStatus(FGuid EntryGuid) const override;
+
 
 	virtual void PushSubjectStaticData_AnyThread(const FLiveLinkSubjectKey& SubjectKey, TSubclassOf<ULiveLinkRole> Role, FLiveLinkStaticDataStruct&& StaticData) override;
 	virtual void PushSubjectFrameData_AnyThread(const FLiveLinkSubjectKey& SubjectKey, FLiveLinkFrameDataStruct&& FrameData) override;
@@ -160,7 +162,6 @@ public:
 	FLiveLinkSubjectTimeSyncData GetTimeSyncData(FLiveLinkSubjectName SubjectName);
 
 	FText GetSourceMachineName(FGuid EntryGuid) const;
-	FText GetSourceStatus(FGuid EntryGuid) const;
 	bool IsSourceStillValid(FGuid EntryGuid) const;
 	UE_DEPRECATED(4.23, "FLiveLinkClient::GetSourceTypeForEntry is deprecated. Please use GetSourceType instead!")
 	FText GetSourceTypeForEntry(FGuid EntryGuid) const { return GetSourceType(EntryGuid); }
@@ -239,10 +240,10 @@ private:
 	/** Removes a subject from the rebroadcast provider and resets it if there are no more subjects */
 	void RemoveRebroadcastedSubject(FLiveLinkSubjectKey InSubjectKey);
 
-private:
+protected:
 	/** The current collection used. */
 	TUniquePtr<FLiveLinkSourceCollection> Collection;
-
+private:
 	/** Pending static info to add to a subject. */
 	TArray<FPendingSubjectStatic> SubjectStaticToPush;
 
diff --git a/Engine/Plugins/Animation/LiveLink/Source/LiveLinkEditor/Internal/LiveLinkClientPanelToolbar.h b/Engine/Plugins/Animation/LiveLink/Source/LiveLinkEditor/Internal/LiveLinkClientPanelToolbar.h
new file mode 100644
--- /dev/null
+++ b/Engine/Plugins/Animation/LiveLink/Source/LiveLinkEditor/Internal/LiveLinkClientPanelToolbar.h
@@ -0,0 +1,71 @@
+// Copyright Epic Games, Inc. All Rights Reserved.
+
+#pragma once
+
+#include "Widgets/SCompoundWidget.h"
+#include "UObject/GCObject.h"
+
+
+class ULiveLinkPreset;
+template <typename T> class TSubclassOf;
+
+class FMenuBuilder;
+class FLiveLinkClient;
+class ILiveLinkSource;
+class IMenu;
+class SComboButton;
+class SEditableTextBox;
+class STextEntryPopup;
+class ULiveLinkSourceFactory;
+struct FAssetData;
+
+class LIVELINKEDITOR_API SLiveLinkClientPanelToolbar : public SCompoundWidget, public FGCObject
+{
+	SLATE_BEGIN_ARGS(SLiveLinkClientPanelToolbar){}
+	SLATE_ARGUMENT(TSharedPtr<SWindow>, ParentWindow)
+	SLATE_END_ARGS()
+
+
+	void Construct(const FArguments& Args, FLiveLinkClient* InClient);
+
+	//~ FEditorUndoClient interface
+	virtual void AddReferencedObjects(FReferenceCollector& Collector) override;
+	virtual FString GetReferencerName() const override
+	{
+		return TEXT("SLiveLinkClientPanelToolbar");
+	}
+
+private:
+	TSharedRef<SWidget> OnGenerateSourceMenu();
+	void RetrieveFactorySourcePanel(FMenuBuilder& MenuBuilder, int32 FactoryIndex);
+	void ExecuteCreateSource(int32 FactoryIndex);
+	void OnSourceCreated(TSharedPtr<ILiveLinkSource> NewSource, FString ConnectionString, TSubclassOf<ULiveLinkSourceFactory> Factory);
+
+	TSharedRef<SWidget> OnPresetGeneratePresetsMenu();
+	TSharedRef<SWidget> HandleSourceSelectionComboButton();
+	void HandleVirtualSourceSelection(const FGuid& InSourceGuid);
+	bool IsVirtualSourceSelected(const FGuid& InSourceGuid);
+	void OnSaveAsPreset();
+	void OnImportPreset(const FAssetData& InPreset);
+	FReply OnRevertChanges();
+	bool HasLoadedLiveLinkPreset() const;
+
+	void PopulateVirtualSubjectSourceCreationMenu(FMenuBuilder& InMenuBuilder);
+	FReply OnAddVirtualSubjectSource();
+	void AddVirtualSubject();
+
+
+private:
+	FLiveLinkClient* Client;
+
+	TWeakPtr<IMenu> AddSubjectMenu;
+	TWeakPtr<IMenu> VirtualSubjectMenu;
+	TWeakPtr<STextEntryPopup> VirtualSubjectPopup;
+	TWeakPtr<SEditableTextBox> VirtualSubjectSourceName;
+
+	TWeakObjectPtr<ULiveLinkPreset> LiveLinkPreset;
+	TArray<TObjectPtr<ULiveLinkSourceFactory>> Factories;
+
+	/** Parent window override used for creating asset dialogs when running in LiveLinkHub. */
+	TSharedPtr<SWindow> ParentWindowOverride;
+};
diff --git a/Engine/Plugins/Animation/LiveLink/Source/LiveLinkEditor/Internal/LiveLinkClientPanelViews.h b/Engine/Plugins/Animation/LiveLink/Source/LiveLinkEditor/Internal/LiveLinkClientPanelViews.h
new file mode 100644
--- /dev/null
+++ b/Engine/Plugins/Animation/LiveLink/Source/LiveLinkEditor/Internal/LiveLinkClientPanelViews.h
@@ -0,0 +1,205 @@
+// Copyright Epic Games, Inc. All Rights Reserved.
+
+#pragma once
+
+#include "Delegates/Delegate.h"
+#include "Delegates/DelegateCombinations.h"
+#include "Input/Events.h"
+#include "Layout/Geometry.h"
+#include "LiveLinkRole.h"
+#include "LiveLinkSourceSettings.h"
+#include "Misc/Guid.h"
+#include "Templates/SharedPointer.h"
+#include "Widgets/Views/SListView.h"
+#include "Widgets/Views/STreeView.h"
+
+class FLiveLinkClient;
+struct FLiveLinkSourceUIEntry;
+class FLiveLinkSourcesView;
+struct FLiveLinkSubjectUIEntry;
+class FLiveLinkSubjectsView;
+class FUICommandList;
+class IDetailsView;
+class ITableRow;
+struct FPropertyChangedEvent;
+class SLiveLinkSourceListView;
+class SLiveLinkDataView;
+class STableViewBase;
+class SWidget;
+
+typedef TSharedPtr<FLiveLinkSourceUIEntry> FLiveLinkSourceUIEntryPtr;
+typedef TSharedPtr<FLiveLinkSubjectUIEntry> FLiveLinkSubjectUIEntryPtr;
+
+namespace UE::LiveLink
+{
+	TSharedPtr<IDetailsView> CreateSourcesDetailsView(const TSharedPtr<FLiveLinkSourcesView>& InSourcesView);
+	TSharedPtr<SLiveLinkDataView> CreateSubjectsDetailsView(FLiveLinkClient* InLiveLinkClient);
+}
+
+// Structure that defines a single entry in the subject UI
+struct FLiveLinkSubjectUIEntry
+{
+	FLiveLinkSubjectUIEntry(const FLiveLinkSubjectKey& InSubjectKey, FLiveLinkClient* InClient);
+
+	// Subject key
+	FLiveLinkSubjectKey SubjectKey;
+	// LiveLink Client
+	FLiveLinkClient* Client;
+
+	// Children (if this entry represents a source instead of a specific subject
+	TArray<TSharedPtr<FLiveLinkSubjectUIEntry>> Children;
+
+	// Whether the subject entry is a subject
+	bool IsSubject() const;
+	// Whether the subject entry is a source
+	bool IsSource() const;
+	// Whether the subject is virtual
+	bool IsVirtualSubject() const;
+	// Get the subject or source settings
+	UObject* GetSettings() const;
+	// Whether the subject is enabled
+	bool IsSubjectEnabled() const;
+	// Whether the subject is valid
+	bool IsSubjectValid() const;
+	// Enable or disable a subject
+	void SetSubjectEnabled(bool bIsEnabled);
+	// Get a textual representation of the ui entry 
+	FText GetItemText() const;
+	// Get the livelink role of this entry
+	TSubclassOf<ULiveLinkRole> GetItemRole() const;
+	// Remove the subject or source from the livelink client
+	void RemoveFromClient() const;
+
+private:
+	// Whether the subject is virtual
+	bool bIsVirtualSubject = false;
+};
+
+// Structure that defines a single entry in the source UI
+struct FLiveLinkSourceUIEntry
+{
+public:
+	FLiveLinkSourceUIEntry(FGuid InEntryGuid, FLiveLinkClient* InClient)
+		: EntryGuid(InEntryGuid)
+		, Client(InClient)
+	{}
+
+	// Get the guid of the source
+	FGuid GetGuid() const;
+	// Get the type of the source
+	FText GetSourceType() const;
+	// Get the machine name of the source
+	FText GetMachineName() const;
+	// Get the source's status
+	FText GetStatus() const;
+	// Get the source's settings
+	ULiveLinkSourceSettings* GetSourceSettings() const;
+	// Remove the source from the client
+	void RemoveFromClient() const;
+	// Get the display name of the source
+	FText GetDisplayName() const;
+
+private:
+	// GUID of the source
+	FGuid EntryGuid;
+	// Pointer to the livelink client
+	FLiveLinkClient* Client = nullptr;
+};
+
+namespace UE::LiveLink
+{
+	/** Base class for live link list/tree views, handles removing list element by pressing delete. */
+	template <typename ListType, typename ListElementType>
+	class SLiveLinkListView : public ListType
+	{
+	public:
+		virtual FReply OnKeyDown(const FGeometry& MyGeometry, const FKeyEvent& InKeyEvent) override
+		{
+			if (InKeyEvent.GetKey() == EKeys::Delete || InKeyEvent.GetKey() == EKeys::BackSpace)
+			{
+				TArray<ListElementType> SelectedItem = ListType::GetSelectedItems();
+				for (ListElementType Item : SelectedItem)
+				{
+					Item->RemoveFromClient();
+				}
+				return FReply::Handled();
+			}
+
+			return FReply::Unhandled();
+		}
+	};
+}
+
+class SLiveLinkSourceListView : public UE::LiveLink::SLiveLinkListView<SListView<FLiveLinkSourceUIEntryPtr>, FLiveLinkSourceUIEntryPtr>
+{
+};
+
+class SLiveLinkSubjectsTreeView : public UE::LiveLink::SLiveLinkListView<STreeView<FLiveLinkSubjectUIEntryPtr>, FLiveLinkSubjectUIEntryPtr>
+{
+};
+
+class FLiveLinkSubjectsView : public TSharedFromThis<FLiveLinkSubjectsView>
+{
+public:
+	DECLARE_DELEGATE_TwoParams(FOnSubjectSelectionChanged, FLiveLinkSubjectUIEntryPtr, ESelectInfo::Type);
+
+	FLiveLinkSubjectsView(FOnSubjectSelectionChanged InOnSubjectSelectionChanged, const TSharedPtr<FUICommandList>& InCommandList);
+
+	// Helper functions for building the subject tree UI
+	TSharedRef<ITableRow> MakeTreeRowWidget(FLiveLinkSubjectUIEntryPtr InInfo, const TSharedRef<STableViewBase>& OwnerTable);
+	// Get a subjectsview's children
+	void GetChildrenForInfo(FLiveLinkSubjectUIEntryPtr InInfo, TArray< FLiveLinkSubjectUIEntryPtr >& OutChildren);
+	// Handle subject selection changed
+	void OnSubjectSelectionChanged(FLiveLinkSubjectUIEntryPtr SubjectEntry, ESelectInfo::Type SelectInfo);
+	// Handler for the subject tree context menu opening
+	TSharedPtr<SWidget> OnOpenVirtualSubjectContextMenu(TSharedPtr<FUICommandList> InCommandList);
+	// Return whether a subject can be removed
+	bool CanRemoveSubject() const;
+
+private:
+	// Create the subjects tree view
+	void CreateSubjectsTreeView(const TSharedPtr<FUICommandList>& InCommandList);
+
+public:
+	// Subject tree widget
+	TSharedPtr<class SLiveLinkSubjectsTreeView> SubjectsTreeView;
+	// Subject tree items
+	TArray<FLiveLinkSubjectUIEntryPtr> SubjectData;
+	// Subject Selection Changed delegate
+	FOnSubjectSelectionChanged SubjectSelectionChangedDelegate;
+};
+
+class FLiveLinkSourcesView : public TSharedFromThis<FLiveLinkSourcesView>
+{
+public:
+	DECLARE_DELEGATE_TwoParams(FOnSourceSelectionChanged, FLiveLinkSourceUIEntryPtr, ESelectInfo::Type);
+
+	FLiveLinkSourcesView(FLiveLinkClient* InLiveLinkClient, TSharedPtr<FUICommandList> InCommandList, FOnSourceSelectionChanged InOnSourceSelectionChanged);
+
+	// Gather information about all sources and update the list view 
+	void RefreshSourceData(bool bRefreshUI);
+	// Handler that creates a widget row for a given ui entry
+	TSharedRef<ITableRow> MakeSourceListViewWidget(FLiveLinkSourceUIEntryPtr Entry, const TSharedRef<STableViewBase>& OwnerTable) const;
+	// Handles constructing a context menu for the sources
+	TSharedPtr<SWidget> OnSourceConstructContextMenu(TSharedPtr<FUICommandList> InCommandList);
+	// Return whether a source can be removed
+	bool CanRemoveSource();
+	// Removes a livelink source from the livelink client
+	void HandleRemoveSource();
+	// Callback when property changes on source settings
+	void OnPropertyChanged(const FPropertyChangedEvent& InEvent);
+	// Handle selection change, triggering the OnSourceSelectionChangedDelegate delegate.
+	void OnSourceListSelectionChanged(FLiveLinkSourceUIEntryPtr Entry, ESelectInfo::Type SelectionType) const;
+private:
+	// Create the sources list view
+	void CreateSourcesListView(const TSharedPtr<FUICommandList>& InCommandList);
+public:
+	// Holds the data that will be displayed in the list view
+	TArray<FLiveLinkSourceUIEntryPtr> SourceData;
+	// Holds the sources list view widget
+	TSharedPtr<SLiveLinkSourceListView> SourcesListView;
+	// LiveLink Client
+	FLiveLinkClient* Client;
+	// Source selection changed delegate
+	FOnSourceSelectionChanged OnSourceSelectionChangedDelegate;
+};
diff --git a/Engine/Plugins/Animation/LiveLink/Source/LiveLinkEditor/Internal/LiveLinkPanelController.h b/Engine/Plugins/Animation/LiveLink/Source/LiveLinkEditor/Internal/LiveLinkPanelController.h
new file mode 100644
--- /dev/null
+++ b/Engine/Plugins/Animation/LiveLink/Source/LiveLinkEditor/Internal/LiveLinkPanelController.h
@@ -0,0 +1,76 @@
+// Copyright Epic Games, Inc. All Rights Reserved.
+
+#pragma once
+
+#include "Delegates/IDelegateInstance.h"
+#include "Templates/SharedPointer.h"
+
+class IDetailsView;
+class FLiveLinkClient;
+struct FLiveLinkSourceUIEntry;
+struct FLiveLinkSubjectUIEntry;
+class FLiveLinkSourcesView;
+class SLiveLinkSourceListView;
+class FLiveLinkSubjectsView;
+class FUICommandList;
+class SLiveLinkDataView;
+
+namespace ESelectInfo { enum Type : int; }
+typedef TSharedPtr<FLiveLinkSourceUIEntry> FLiveLinkSourceUIEntryPtr;
+typedef TSharedPtr<FLiveLinkSubjectUIEntry> FLiveLinkSubjectUIEntryPtr;
+
+/** Handles callback connections between the sources, subjects and details views. */
+class LIVELINKEDITOR_API FLiveLinkPanelController : public TSharedFromThis<FLiveLinkPanelController>
+{
+public:
+	FLiveLinkPanelController();
+	~FLiveLinkPanelController();
+
+private:
+	// Bind live link commands 
+	void BindCommands();
+	// Handles the source collection changing.
+	void OnSourcesChangedHandler();
+	// Handles the subject collection changing.
+	void OnSubjectsChangedHandler();
+	// Remove all sources command handlers
+	bool HasSource() const;
+	// Handles the remove source command.
+	void HandleRemoveSource();
+	// Handles the remove all sources command.
+	void HandleRemoveAllSources();
+	// Returns whether a source could be removed.
+	bool CanRemoveSource() const;
+	// Returns whether a subject can be removed.
+	bool CanRemoveSubject() const;
+	// Handles the remove subject command.
+	void HandleRemoveSubject();
+	// Recreates the source list data behind the list view.
+	void RebuildSourceList();
+	// Recreates the subject list data behind the tree view.
+	void RebuildSubjectList();
+	// Handles source selection changing.
+	void OnSourceSelectionChanged(FLiveLinkSourceUIEntryPtr Entry, ESelectInfo::Type SelectionType) const;
+	// Hadnles subject selection changing.
+	void OnSubjectSelectionChanged(FLiveLinkSubjectUIEntryPtr SubjectEntry, ESelectInfo::Type SelectInfo);
+
+public:
+	// Sources view
+	TSharedPtr<FLiveLinkSourcesView> SourcesView;
+	// Subjects view
+	TSharedPtr<FLiveLinkSubjectsView> SubjectsView;
+	// Reference to connection settings struct details panel
+	TSharedPtr<IDetailsView> SourcesDetailsView;
+	// Reference to the data value struct details panel
+	TSharedPtr<SLiveLinkDataView> SubjectsDetailsView;
+	// LiveLink Client
+	FLiveLinkClient* Client;
+	// Handle to delegate when client sources list has changed
+	FDelegateHandle OnSourcesChangedHandle;
+	// Handle to delegate when a client subjects list has changed
+	FDelegateHandle OnSubjectsChangedHandle;
+	// Guard from reentrant selection
+	mutable bool bSelectionChangedGuard = false;
+	// Command list
+	TSharedPtr<FUICommandList> CommandList;
+};
\ No newline at end of file
diff --git a/Engine/Plugins/Animation/LiveLink/Source/LiveLinkEditor/Internal/SLiveLinkDataView.h b/Engine/Plugins/Animation/LiveLink/Source/LiveLinkEditor/Internal/SLiveLinkDataView.h
new file mode 100644
--- /dev/null
+++ b/Engine/Plugins/Animation/LiveLink/Source/LiveLinkEditor/Internal/SLiveLinkDataView.h
@@ -0,0 +1,59 @@
+// Copyright Epic Games, Inc. All Rights Reserved.
+
+#pragma once
+
+#include "Widgets/SCompoundWidget.h"
+
+#include "LiveLinkTypes.h"
+
+namespace ETextCommit { enum Type : int; }
+struct FPropertyChangedEvent;
+
+class IDetailsView;
+class FLiveLinkClient;
+class IStructureDetailsView;
+
+class SLiveLinkDataView : public SCompoundWidget
+{
+private:
+	using Super = SCompoundWidget;
+
+public:
+	SLATE_BEGIN_ARGS(SLiveLinkDataView){}
+	SLATE_END_ARGS()
+
+	void Construct(const FArguments& Args, FLiveLinkClient* InClient);
+
+	virtual void Tick(const FGeometry& AllottedGeometry, const double InCurrentTime, const float InDeltaTime) override;
+
+	void SetSubjectKey(FLiveLinkSubjectKey InSubjectKey);
+	FLiveLinkSubjectKey GetSubjectKey() const { return SubjectKey; }
+
+	void SetRefreshDelay(double DelaySeconds) { UpdateDelay = DelaySeconds; }
+	double GetRefreshDelay() const { return UpdateDelay; }
+
+private:
+	enum class EDetailType : uint32
+	{
+		Property,
+		StaticData,
+		FrameData,
+	};
+
+	void OnPropertyChanged(const FPropertyChangedEvent& InEvent);
+	int32 GetDetailWidgetIndex() const;
+	void OnSelectDetailWidget(EDetailType InDetailType);
+	bool IsSelectedDetailWidget(EDetailType InDetailType) const { return InDetailType == DetailType; }
+	bool CanEditRefreshDelay() const { return DetailType != EDetailType::Property; }
+	void SetRefreshDelayInternal(double InDelaySeconds, ETextCommit::Type) { SetRefreshDelay(InDelaySeconds); }
+
+private:
+	FLiveLinkClient* Client;
+	FLiveLinkSubjectKey SubjectKey;
+	double LastUpdateSeconds;
+	double UpdateDelay;
+	EDetailType DetailType;
+
+	TSharedPtr<IStructureDetailsView> StructureDetailsView;
+	TSharedPtr<IDetailsView> SettingsDetailsView;
+};
diff --git a/Engine/Plugins/Animation/LiveLink/Source/LiveLinkEditor/Private/LiveLinkClientPanel.cpp b/Engine/Plugins/Animation/LiveLink/Source/LiveLinkEditor/Private/LiveLinkClientPanel.cpp
--- a/Engine/Plugins/Animation/LiveLink/Source/LiveLinkEditor/Private/LiveLinkClientPanel.cpp
+++ b/Engine/Plugins/Animation/LiveLink/Source/LiveLinkEditor/Private/LiveLinkClientPanel.cpp
@@ -2,394 +2,29 @@
 
 #include "LiveLinkClientPanel.h"
 
-
-#include "DetailsViewArgs.h"
+#include "Editor.h"
+#include "Editor/EditorPerformanceSettings.h"
 #include "Framework/MultiBox/MultiBoxBuilder.h"
-#include "Framework/Commands/UICommandList.h"
-#include "Widgets/Layout/SWidgetSwitcher.h"
-#include "Framework/Views/TableViewMetadata.h"
-#include "Widgets/Input/SButton.h"
 #include "IDetailsView.h"
-#include "Widgets/Input/SCheckBox.h"
-#include "LiveLinkRole.h"
-#include "SWarningOrErrorBox.h"
-
-#include "LiveLinkClientCommands.h"
 #include "LiveLinkClient.h"
 #include "LiveLinkClientPanelToolbar.h"
+#include "LiveLinkClientPanelViews.h"
 #include "LiveLinkLog.h"
-#include "LiveLinkSettings.h"
+#include "LiveLinkRole.h"
 #include "LiveLinkSourceFactory.h"
 #include "LiveLinkSourceSettings.h"
-#include "SLiveLinkDataView.h"
-
-#include "Editor.h"
-#include "Editor/EditorPerformanceSettings.h"
-#include "EditorFontGlyphs.h"
 #include "MessageLogModule.h"
 #include "Modules/ModuleManager.h"
-#include "PropertyEditorModule.h"
-#include "Widgets/Views/STreeView.h"
+#include "SLiveLinkDataView.h"
+#include "SWarningOrErrorBox.h"
+#include "Widgets/Input/SButton.h"
+#include "Widgets/Layout/SWidgetSwitcher.h"
 
 
 #define LOCTEXT_NAMESPACE "LiveLinkClientPanel"
 
-// Static Source UI FNames
-namespace SourceListUI
-{
-	static const FName TypeColumnName(TEXT("Type"));
-	static const FName MachineColumnName(TEXT("Machine"));
-	static const FName StatusColumnName(TEXT("Status"));
-	static const FName ActionsColumnName(TEXT("Action"));
-};
-
-// Static Subject UI FNames
-namespace SubjectTreeUI
-{
-	static const FName EnabledColumnName(TEXT("Enabled"));
-	static const FName NameColumnName(TEXT("Name"));
-	static const FName RoleColumnName(TEXT("Role"));
-	static const FName ActionsColumnName(TEXT("Action"));
-};
-
-namespace
-{
-	/** Base class for live link list/tree views, handles removing list element by pressing delete. */
-	template <typename ListType, typename ListElementType>
-	class SLiveLinkListView : public ListType
-	{
-	public:
-		virtual FReply OnKeyDown(const FGeometry& MyGeometry, const FKeyEvent& InKeyEvent) override
-		{
-			if (InKeyEvent.GetKey() == EKeys::Delete || InKeyEvent.GetKey() == EKeys::BackSpace)
-			{
-				TArray<ListElementType> SelectedItem = ListType::GetSelectedItems();
-				for (ListElementType Item : SelectedItem)
-				{
-					Item->RemoveFromClient();
-				}
-				return FReply::Handled();
-			}
-
-			return FReply::Unhandled();
-		}
-	};
-}
-
-// Structure that defines a single entry in the source UI
-struct FLiveLinkSourceUIEntry
-{
-public:
-	FLiveLinkSourceUIEntry(FGuid InEntryGuid, FLiveLinkClient* InClient)
-		: EntryGuid(InEntryGuid)
-		, Client(InClient)
-	{}
-
-	FGuid GetGuid() const { return EntryGuid; }
-	FText GetSourceType() const { return Client->GetSourceType(EntryGuid); }
-	FText GetMachineName() const { return Client->GetSourceMachineName(EntryGuid); }
-	FText GetStatus() const { return Client->GetSourceStatus(EntryGuid); }
-	ULiveLinkSourceSettings* GetSourceSettings() const { return Client->GetSourceSettings(EntryGuid); }
-	void RemoveFromClient() const { Client->RemoveSource(EntryGuid); }
-	FText GetDisplayName() const { return GetSourceType(); }
-
-private:
-	FGuid EntryGuid;
-	FLiveLinkClient* Client;
-};
-
-// Structure that defines a single entry in the subject UI
-struct FLiveLinkSubjectUIEntry
-{
-	FLiveLinkSubjectUIEntry(const FLiveLinkSubjectKey& InSubjectKey, FLiveLinkClient* InClient)
-		: SubjectKey(InSubjectKey)
-		, Client(InClient)
-		, bIsVirtualSubject(false)
-	{
-		if (InClient)
-		{
-			bIsVirtualSubject = InClient->IsVirtualSubject(InSubjectKey);
-		}
-	}
-	FLiveLinkSubjectKey SubjectKey;
-	FLiveLinkClient* Client;
-
-	// Children (if this entry represents a source instead of a specific subject
-	TArray<FLiveLinkSubjectUIEntryPtr> Children;
-
-
-	bool IsSubject() const { return !SubjectKey.SubjectName.IsNone(); }
-	bool IsSource() const { return SubjectKey.SubjectName.IsNone(); }
-	bool IsVirtualSubject() const { return (IsSubject() && bIsVirtualSubject); }
-
-	UObject* GetSettings() const
-	{
-		if (IsSource())
-		{
-			return Client->GetSourceSettings(SubjectKey.Source);
-		}
-		else
-		{
-			return Client->GetSubjectSettings(SubjectKey);
-		}
-	}
-
-	bool IsSubjectEnabled() const
-	{
-		return IsSubject() ? Client->IsSubjectEnabled(SubjectKey, false) : false;
-	}
-
-	bool IsSubjectValid() const
-	{
-		return IsSubject() ? Client->IsSubjectValid(SubjectKey) : false;
-	}
-
-	void SetSubjectEnabled(bool bIsEnabled)
-	{
-		if (IsSubject())
-		{
-			Client->SetSubjectEnabled(SubjectKey, bIsEnabled);
-		}
-	}
-
-	FText GetItemText() const
-	{
-		if (IsSource())
-		{
-			return Client->GetSourceType(SubjectKey.Source);
-		}
-		else
-		{
-			return FText::FromName(SubjectKey.SubjectName);
-		}
-	}
-
-	TSubclassOf<ULiveLinkRole> GetItemRole() const { return IsSubject() ? Client->GetSubjectRole(SubjectKey) : TSubclassOf<ULiveLinkRole>(); }
-
-	void RemoveFromClient() const { Client->RemoveSubject_AnyThread(SubjectKey); }
-	
-private:
-	bool bIsVirtualSubject;
-};
-
-class SLiveLinkClientPanelSubjectRow : public SMultiColumnTableRow<FLiveLinkSubjectUIEntryPtr>
-{
-public:
-	SLATE_BEGIN_ARGS(SLiveLinkClientPanelSubjectRow) {}
-		/** The list item for this row */
-		SLATE_ARGUMENT(FLiveLinkSubjectUIEntryPtr, Entry)
-	SLATE_END_ARGS()
-
-
-	void Construct(const FArguments& Args, const TSharedRef<STableViewBase>& OwnerTableView)
-	{
-		EntryPtr = Args._Entry;
-
-		SMultiColumnTableRow<FLiveLinkSubjectUIEntryPtr>::Construct(
-			FSuperRowType::FArguments()
-			.Padding(1.0f),
-			OwnerTableView
-		);
-	}
-
-	/** Overridden from SMultiColumnTableRow.  Generates a widget for this column of the list view. */
-	virtual TSharedRef<SWidget> GenerateWidgetForColumn(const FName& ColumnName) override
-	{
-		if (ColumnName == SubjectTreeUI::EnabledColumnName)
-		{
-			if (EntryPtr->IsSubject())
-			{
-				return SNew(SCheckBox)
-					.IsChecked(MakeAttributeSP(this, &SLiveLinkClientPanelSubjectRow::GetSubjectEnabled))
-					.OnCheckStateChanged(this, &SLiveLinkClientPanelSubjectRow::OnEnabledChanged);
-			}
-		}
-		else if (ColumnName == SubjectTreeUI::NameColumnName)
-		{
-			return SNew(SHorizontalBox)
-				+ SHorizontalBox::Slot()
-				.AutoWidth()
-				.Padding(6, 0, 0, 0)
-				[
-					SNew(SExpanderArrow, SharedThis(this)).IndentAmount(12)
-				]
-				+ SHorizontalBox::Slot()
-				.FillWidth(1.0f)
-				[
-					SNew(STextBlock)
-					.Text(MakeAttributeSP(this, &SLiveLinkClientPanelSubjectRow::GetItemText))
-				];
-		}
-		else if (ColumnName == SubjectTreeUI::RoleColumnName)
-		{
-			auto RoleAttribute = MakeAttributeSP(this, &SLiveLinkClientPanelSubjectRow::GetItemRole);
-			return SNew(STextBlock)
-				.Text(EntryPtr->IsSubject() ? RoleAttribute : FText::GetEmpty());
-		}
-		else if (ColumnName == SubjectTreeUI::ActionsColumnName)
-		{
-			if (EntryPtr->IsVirtualSubject())
-			{
-				return SNew(SButton)
-					.ButtonStyle(FAppStyle::Get(), "HoverHintOnly")
-					.HAlign(HAlign_Center)
-					.VAlign(VAlign_Center)
-					.OnClicked(this, &SLiveLinkClientPanelSubjectRow::OnRemoveClicked)
-					.ToolTipText(LOCTEXT("RemoveVirtualSubject", "Remove selected live link virtual subject"))
-					.ContentPadding(0.f)
-					.ForegroundColor(FSlateColor::UseForeground())
-					.IsFocusable(false)
-					[
-						SNew(SImage)
-						.Image(FAppStyle::GetBrush("Icons.Delete"))
-						.ColorAndOpacity(FSlateColor::UseForeground())
-					];
-			}
-			else
-			{
-				return SNew(SBox)
-					.HAlign(HAlign_Center)
-					.VAlign(VAlign_Center)
-					[
-						SNew(STextBlock)
-						.Font(FAppStyle::Get().GetFontStyle("FontAwesome.8"))
-						.ColorAndOpacity(this, &SLiveLinkClientPanelSubjectRow::OnGetActivityColor)
-						.Text(FEditorFontGlyphs::Circle)
-					];
-			}
-		}
-
-		return SNullWidget::NullWidget;
-	}
-
-private:
-	ECheckBoxState GetSubjectEnabled() const { return EntryPtr->IsSubjectEnabled() ? ECheckBoxState::Checked : ECheckBoxState::Unchecked; }
-	void OnEnabledChanged(ECheckBoxState NewState) { EntryPtr->SetSubjectEnabled(NewState == ECheckBoxState::Checked); }
-	FText GetItemText() const { return EntryPtr->GetItemText(); }
-	FText GetItemRole() const
-	{
-		TSubclassOf<ULiveLinkRole> Role = EntryPtr->GetItemRole();
-		if (Role.Get())
-		{
-			return Role->GetDefaultObject<ULiveLinkRole>()->GetDisplayName();
-		}
-		return FText::GetEmpty();
-	}
-
-	FReply OnRemoveClicked()
-	{
-		EntryPtr->RemoveFromClient();
-		return FReply::Handled();
-	}
-
-	FSlateColor OnGetActivityColor() const
-	{
-		if (EntryPtr->IsSubjectEnabled())
-		{
-			return EntryPtr->IsSubjectValid() ? GetDefault<ULiveLinkSettings>()->GetValidColor() : GetDefault<ULiveLinkSettings>()->GetInvalidColor();
-		}
-		return FLinearColor(0.f, 0.f, 0.f, 0.f);
-	}
-
-	FLiveLinkSubjectUIEntryPtr EntryPtr;
-};
-
-class SLiveLinkClientPanelSourcesRow : public SMultiColumnTableRow<FLiveLinkSourceUIEntryPtr>
-{
-public:
-	SLATE_BEGIN_ARGS(SLiveLinkClientPanelSourcesRow) {}
-		/** The list item for this row */
-		SLATE_ARGUMENT(FLiveLinkSourceUIEntryPtr, Entry)
-	SLATE_END_ARGS()
-
-	void Construct(const FArguments& Args, const TSharedRef<STableViewBase>& OwnerTableView)
-	{
-		EntryPtr = Args._Entry;
-
-		SMultiColumnTableRow<FLiveLinkSourceUIEntryPtr>::Construct(
-			FSuperRowType::FArguments()
-			.Padding(1.0f),
-			OwnerTableView
-		);
-	}
-
-	/** Overridden from SMultiColumnTableRow.  Generates a widget for this column of the list view. */
-	virtual TSharedRef<SWidget> GenerateWidgetForColumn(const FName& ColumnName) override
-	{
-		if (ColumnName == SourceListUI::TypeColumnName)
-		{
-			return	SNew(STextBlock)
-					.Text(EntryPtr->GetSourceType());
-		}
-		else if (ColumnName == SourceListUI::MachineColumnName)
-		{
-			return	SNew(STextBlock)
-					.Text(MakeAttributeSP(this, &SLiveLinkClientPanelSourcesRow::GetMachineName));
-		}
-		else if (ColumnName == SourceListUI::StatusColumnName)
-		{
-			return	SNew(STextBlock)
-					.Text(MakeAttributeSP(this, &SLiveLinkClientPanelSourcesRow::GetSourceStatus));
-		}
-		else if (ColumnName == SourceListUI::ActionsColumnName)
-		{
-			return	SNew(SButton)
-					.ButtonStyle(FAppStyle::Get(), "HoverHintOnly")
-					.HAlign(HAlign_Center)
-					.VAlign(VAlign_Center)
-					.OnClicked(this, &SLiveLinkClientPanelSourcesRow::OnRemoveClicked)
-					.ToolTipText(LOCTEXT("RemoveSource", "Remove selected live link source"))
-					.ContentPadding(0.f)
-					.ForegroundColor(FSlateColor::UseForeground())
-					.IsFocusable(false)
-					[
-						SNew(SImage)
-						.Image(FAppStyle::GetBrush("Icons.Delete"))
-						.ColorAndOpacity(FSlateColor::UseForeground())
-					];
-		}
-
-		return SNullWidget::NullWidget;
-	}
-
-private:
-	FText GetMachineName() const
-	{
-		return EntryPtr->GetMachineName();
-	}
-
-	FText GetSourceStatus() const
-	{
-		return EntryPtr->GetStatus();
-	}
-
-	FReply OnRemoveClicked()
-	{
-		EntryPtr->RemoveFromClient();
-		return FReply::Handled();
-	}
-
-	FLiveLinkSourceUIEntryPtr EntryPtr;
-};
-
-class SLiveLinkSourceListView : public SLiveLinkListView<SListView<FLiveLinkSourceUIEntryPtr>, FLiveLinkSourceUIEntryPtr>
-{
-};
-
-class SLiveLinkSubjectsTreeView : public SLiveLinkListView<STreeView<FLiveLinkSubjectUIEntryPtr>, FLiveLinkSubjectUIEntryPtr>
-{
-};
-
 SLiveLinkClientPanel::~SLiveLinkClientPanel()
 {
-	if (Client)
-	{
-		Client->OnLiveLinkSourcesChanged().Remove(OnSourcesChangedHandle);
-		OnSourcesChangedHandle.Reset();
-
-		Client->OnLiveLinkSubjectsChanged().Remove(OnSubjectsChangedHandle);
-		OnSubjectsChangedHandle.Reset();
-	}
 	GEditor->UnregisterForUndo(this);
 }
 
@@ -400,54 +35,9 @@ void SLiveLinkClientPanel::Construct(const FArguments& Args, FLiveLinkClient* In
 	check(InClient);
 	Client = InClient;
 
-	bSelectionChangedGuard = false;
 	DetailWidgetIndex = 0;
-
-	OnSourcesChangedHandle = Client->OnLiveLinkSourcesChanged().AddSP(this, &SLiveLinkClientPanel::OnSourcesChangedHandler);
-	OnSubjectsChangedHandle = Client->OnLiveLinkSubjectsChanged().AddSP(this, &SLiveLinkClientPanel::OnSubjectsChangedHandler);
-
-	RefreshSourceData(false);
-
-	CommandList = MakeShareable(new FUICommandList);
-
-	BindCommands();
-
-	// Connection Settings
-	FPropertyEditorModule& PropertyEditorModule = FModuleManager::GetModuleChecked<FPropertyEditorModule>("PropertyEditor");
-
-	FDetailsViewArgs DetailsViewArgs;
-	DetailsViewArgs.bUpdatesFromSelection = false;
-	DetailsViewArgs.bLockable = false;
-	DetailsViewArgs.bShowPropertyMatrixButton = false;
-	DetailsViewArgs.NameAreaSettings = FDetailsViewArgs::HideNameArea;
-	DetailsViewArgs.ViewIdentifier = NAME_None;
-	SettingsDetailsView = PropertyEditorModule.CreateDetailView(DetailsViewArgs);
-	SettingsDetailsView->OnFinishedChangingProperties().AddSP(this, &SLiveLinkClientPanel::OnPropertyChanged);
-
-	SAssignNew(SubjectsTreeView, SLiveLinkSubjectsTreeView)
-		.TreeItemsSource(&SubjectData)
-		.OnGenerateRow(this, &SLiveLinkClientPanel::MakeTreeRowWidget)
-		.OnGetChildren(this, &SLiveLinkClientPanel::GetChildrenForInfo)
-		.OnSelectionChanged(this, &SLiveLinkClientPanel::OnSubjectTreeSelectionChanged)
-		.OnContextMenuOpening(this, &SLiveLinkClientPanel::OnOpenVirtualSubjectContextMenu)
-		.SelectionMode(ESelectionMode::Single)
-		.HeaderRow
-		(
-			SNew(SHeaderRow)
-			+ SHeaderRow::Column(SubjectTreeUI::EnabledColumnName)
-			.DefaultLabel(LOCTEXT("EnabledName", ""))
-			.FixedWidth(22)
-			+ SHeaderRow::Column(SubjectTreeUI::NameColumnName)
-			.DefaultLabel(LOCTEXT("SubjectItemName", "Subject Name"))
-			.FillWidth(0.60f)
-			+ SHeaderRow::Column(SubjectTreeUI::RoleColumnName)
-			.DefaultLabel(LOCTEXT("RoleName", "Role"))
-			.FillWidth(0.40f)
-			+ SHeaderRow::Column(SubjectTreeUI::ActionsColumnName)
-			.ManualWidth(20.f)
-			.DefaultLabel(LOCTEXT("ActionsColumnHeaderName", ""))
-		);
-
+	
+	PanelController = MakeShared<FLiveLinkPanelController>();
 
 	const FName LogName = "Live Link";
 	TSharedPtr<class IMessageLogListing> MessageLogListing;
@@ -460,8 +50,6 @@ void SLiveLinkClientPanel::Construct(const FArguments& Args, FLiveLinkClient* In
 
 	TSharedRef<class SWidget> MessageLogListingWidget = MessageLogListing.IsValid() ? MessageLogModule.CreateLogListingWidget(MessageLogListing.ToSharedRef()) : SNullWidget::NullWidget;
 
-
-	const int WarningPadding = 8;
 	FProperty* PerformanceThrottlingProperty = FindFieldChecked<FProperty>(UEditorPerformanceSettings::StaticClass(), GET_MEMBER_NAME_CHECKED(UEditorPerformanceSettings, bThrottleCPUWhenNotForeground));
 	PerformanceThrottlingProperty->GetDisplayNameText();
 	FFormatNamedArguments Arguments;
@@ -503,28 +91,7 @@ void SLiveLinkClientPanel::Construct(const FArguments& Args, FLiveLinkClient* In
 							.BorderImage(FAppStyle::GetBrush("ToolPanel.GroupBorder"))
 							.Padding(FMargin(4.0f, 4.0f))
 							[
-								SAssignNew(SourceListView, SLiveLinkSourceListView)
-								.ListItemsSource(&SourceData)
-								.SelectionMode(ESelectionMode::Single)
-								.OnGenerateRow(this, &SLiveLinkClientPanel::MakeSourceListViewWidget)
-								.OnContextMenuOpening(this, &SLiveLinkClientPanel::OnSourceConstructContextMenu)
-								.OnSelectionChanged(this, &SLiveLinkClientPanel::OnSourceListSelectionChanged)
-								.HeaderRow
-								(
-									SNew(SHeaderRow)
-									+SHeaderRow::Column(SourceListUI::TypeColumnName)
-									.FillWidth(25.f)
-									.DefaultLabel(LOCTEXT("TypeColumnHeaderName", "Source Type"))
-									+SHeaderRow::Column(SourceListUI::MachineColumnName)
-									.FillWidth(25.f)
-									.DefaultLabel(LOCTEXT("MachineColumnHeaderName", "Source Machine"))
-									+SHeaderRow::Column(SourceListUI::StatusColumnName)
-									.FillWidth(50.f)
-									.DefaultLabel(LOCTEXT("StatusColumnHeaderName", "Status"))
-									+SHeaderRow::Column(SourceListUI::ActionsColumnName)
-									.ManualWidth(20.f)
-									.DefaultLabel(LOCTEXT("ActionsColumnHeaderName", ""))
-								)
+								PanelController->SourcesView->SourcesListView.ToSharedRef()
 							]
 						]
 						+SSplitter::Slot()
@@ -534,7 +101,7 @@ void SLiveLinkClientPanel::Construct(const FArguments& Args, FLiveLinkClient* In
 							.BorderImage(FAppStyle::GetBrush("ToolPanel.GroupBorder"))
 							.Padding(FMargin(4.0f, 4.0f))
 							[
-								SubjectsTreeView->AsShared()
+								PanelController->SubjectsView->SubjectsTreeView.ToSharedRef()
 							]
 						]
 					]
@@ -546,12 +113,12 @@ void SLiveLinkClientPanel::Construct(const FArguments& Args, FLiveLinkClient* In
 						+SWidgetSwitcher::Slot()
 						[
 							//[0] Detail view for Source
-							SettingsDetailsView.ToSharedRef()
+							PanelController->SourcesDetailsView.ToSharedRef()
 						]
 						+SWidgetSwitcher::Slot()
 						[
 							// [1] Detail view for Subject, Frame data & Static data
-							SAssignNew(DataDetailsView, SLiveLinkDataView, Client)
+							PanelController->SubjectsDetailsView.ToSharedRef()
 						]
 					]
 				]
@@ -602,9 +169,6 @@ void SLiveLinkClientPanel::Construct(const FArguments& Args, FLiveLinkClient* In
 			]
 		]
 	];
-
-
-	RebuildSubjectList();
 }
 
 void SLiveLinkClientPanel::AddReferencedObjects(FReferenceCollector& Collector)
@@ -612,314 +176,19 @@ void SLiveLinkClientPanel::AddReferencedObjects(FReferenceCollector& Collector)
 	Collector.AddReferencedObjects(DetailsPanelEditorObjects);
 }
 
-void SLiveLinkClientPanel::BindCommands()
-{
-	CommandList->MapAction(FLiveLinkClientCommands::Get().RemoveSource,
-		FExecuteAction::CreateSP(this, &SLiveLinkClientPanel::HandleRemoveSource),
-		FCanExecuteAction::CreateSP(this, &SLiveLinkClientPanel::CanRemoveSource)
-	);
-
-	CommandList->MapAction(FLiveLinkClientCommands::Get().RemoveAllSources,
-		FExecuteAction::CreateSP(this, &SLiveLinkClientPanel::HandleRemoveAllSources),
-		FCanExecuteAction::CreateSP(this, &SLiveLinkClientPanel::HasSource)
-	);
-
-	CommandList->MapAction(FLiveLinkClientCommands::Get().RemoveSubject,
-	FExecuteAction::CreateSP(this, &SLiveLinkClientPanel::HandleRemoveSubject),
-		FCanExecuteAction::CreateSP(this, &SLiveLinkClientPanel::CanRemoveSubject)
-	);
-}
-
-void SLiveLinkClientPanel::RefreshSourceData(bool bRefreshUI)
-{
-	SourceData.Reset();
-
-	for (FGuid SourceGuid : Client->GetDisplayableSources())
-	{
-		SourceData.Add(MakeShared<FLiveLinkSourceUIEntry>(SourceGuid, Client));
-	}
-	SourceData.Sort([](const FLiveLinkSourceUIEntryPtr& LHS, const FLiveLinkSourceUIEntryPtr& RHS) { return LHS->GetMachineName().CompareTo(RHS->GetMachineName()) < 0; });
-
-	if (bRefreshUI)
-	{
-		SourceListView->RequestListRefresh();
-	}
-}
-
 int32 SLiveLinkClientPanel::GetDetailWidgetIndex() const
 {
-	return DataDetailsView->GetSubjectKey().Source.IsValid() && !DataDetailsView->GetSubjectKey().SubjectName.IsNone() ? 1 : 0;
-}
-
-TSharedRef<ITableRow> SLiveLinkClientPanel::MakeSourceListViewWidget(FLiveLinkSourceUIEntryPtr Entry, const TSharedRef<STableViewBase>& OwnerTable) const
-{
-	return SNew(SLiveLinkClientPanelSourcesRow, OwnerTable)
-		.Entry(Entry);
-}
-
-void SLiveLinkClientPanel::OnSourceListSelectionChanged(FLiveLinkSourceUIEntryPtr Entry, ESelectInfo::Type SelectionType) const
-{
-	if (bSelectionChangedGuard)
-	{
-		return;
-	}
-	TGuardValue<bool> ReentrantGuard(bSelectionChangedGuard, true);
-
-	DataDetailsView->SetSubjectKey(FLiveLinkSubjectKey());
-
-	int32 FoundSubjectEntryIndex = INDEX_NONE;
-	if(Entry.IsValid())
-	{
-		SettingsDetailsView->SetObject(Entry->GetSourceSettings());
-
-		// Find the corresponding subject
-		FoundSubjectEntryIndex = SubjectData.IndexOfByPredicate([Entry](const FLiveLinkSubjectUIEntryPtr& SubjectEntry) { return SubjectEntry->SubjectKey.Source == Entry->GetGuid() && SubjectEntry->IsSource(); });
-	}
-	else
-	{
-		SettingsDetailsView->SetObject(nullptr);
-	}
-
-	// Set the corresponding subject
-	if (FoundSubjectEntryIndex != INDEX_NONE)
-	{
-		SubjectsTreeView->SetSelection(SubjectData[FoundSubjectEntryIndex]);
-	}
-	else
-	{
-		SubjectsTreeView->ClearSelection();
-	}
-}
-
-TSharedPtr<SWidget> SLiveLinkClientPanel::OnSourceConstructContextMenu()
-{
-	const bool bShouldCloseWindowAfterMenuSelection = true;
-	FMenuBuilder MenuBuilder(bShouldCloseWindowAfterMenuSelection, CommandList);
-
-	MenuBuilder.BeginSection(TEXT("Remove"));
-	{
-		if (CanRemoveSource())
-		{
-			MenuBuilder.AddMenuEntry(FLiveLinkClientCommands::Get().RemoveSource);
-		}
-		MenuBuilder.AddMenuEntry(FLiveLinkClientCommands::Get().RemoveAllSources);
-	}
-	MenuBuilder.EndSection();
-
-	return MenuBuilder.MakeWidget();
-}
-
-TSharedRef<ITableRow> SLiveLinkClientPanel::MakeTreeRowWidget(FLiveLinkSubjectUIEntryPtr InInfo, const TSharedRef<STableViewBase>& OwnerTable)
-{
-	return SNew(SLiveLinkClientPanelSubjectRow, OwnerTable)
-		.Entry(InInfo);
-}
-
-void SLiveLinkClientPanel::GetChildrenForInfo(FLiveLinkSubjectUIEntryPtr InInfo, TArray< FLiveLinkSubjectUIEntryPtr >& OutChildren)
-{
-	OutChildren = InInfo->Children;
-}
-
-void SLiveLinkClientPanel::OnSubjectTreeSelectionChanged(FLiveLinkSubjectUIEntryPtr SubjectEntry, ESelectInfo::Type SelectInfo)
-{
-	if (bSelectionChangedGuard)
-	{
-		return;
-	}
-	TGuardValue<bool> ReentrantGuard(bSelectionChangedGuard, true);
-
-	int32 FoundSourceIndex = INDEX_NONE;
-	bool bDetailViewSet = false;
-	if (SubjectEntry.IsValid())
-	{
-		// Find the corresponding Source entry
-		FGuid SourceGuid = SubjectEntry->SubjectKey.Source;
-		FoundSourceIndex = SourceData.IndexOfByPredicate([SourceGuid](FLiveLinkSourceUIEntryPtr SourceEntry) { return SourceEntry->GetGuid() == SourceGuid; });
-
-		if (SubjectEntry->IsSource())
-		{
-			SettingsDetailsView->SetObject(SubjectEntry->GetSettings());
-			DataDetailsView->SetSubjectKey(FLiveLinkSubjectKey());
-		}
-		else
-		{
-			SettingsDetailsView->SetObject(nullptr);
-			DataDetailsView->SetSubjectKey(SubjectEntry->SubjectKey);
-		}
-		bDetailViewSet = true;
-	}
-
-	if (!bDetailViewSet)
-	{
-		SettingsDetailsView->SetObject(nullptr);
-		DataDetailsView->SetSubjectKey(FLiveLinkSubjectKey());
-	}
-
-	// Select the corresponding Source entry
-	if (FoundSourceIndex != INDEX_NONE)
-	{
-		SourceListView->SetSelection(SourceData[FoundSourceIndex]);
-	}
-	else
-	{
-		SourceListView->ClearSelection();
-	}
-}
-
-TSharedPtr<SWidget> SLiveLinkClientPanel::OnOpenVirtualSubjectContextMenu()
-{
-	constexpr bool bShouldCloseWindowAfterMenuSelection = true;
-	FMenuBuilder MenuBuilder(bShouldCloseWindowAfterMenuSelection, CommandList);
-
-	MenuBuilder.BeginSection(TEXT("Remove"));
-	{
-		if (CanRemoveSubject())
-		{
-			MenuBuilder.AddMenuEntry(FLiveLinkClientCommands::Get().RemoveSubject);
-		}
-	}
-	MenuBuilder.EndSection();
-
-	return MenuBuilder.MakeWidget();
-}
-
-void SLiveLinkClientPanel::RebuildSubjectList()
-{
-	TArray<FLiveLinkSubjectKey> SavedSelection;
-	{
-		TArray<FLiveLinkSubjectUIEntryPtr> SelectedItems = SubjectsTreeView->GetSelectedItems();
-		for (const FLiveLinkSubjectUIEntryPtr& SelectedItem : SelectedItems)
-		{
-			SavedSelection.Add(SelectedItem->SubjectKey);
-		}
-	}
-
-	TArray<FLiveLinkSubjectKey> SubjectKeys = Client->GetSubjects(true, true);
-	SubjectData.Reset();
-
-	TMap<FGuid, FLiveLinkSubjectUIEntryPtr> SourceHeaderItems;
-	TArray<FLiveLinkSubjectUIEntryPtr> AllItems;
-	AllItems.Reserve(SubjectKeys.Num());
-
-	for (const FLiveLinkSubjectKey& SubjectKey : SubjectKeys)
-	{
-		FLiveLinkSubjectUIEntryPtr Source;
-		if(FLiveLinkSubjectUIEntryPtr* SourcePtr = SourceHeaderItems.Find(SubjectKey.Source))
-		{
-			Source = *SourcePtr;
-		}
-		else
-		{
-			FLiveLinkSubjectKey SourceKey = SubjectKey;
-			SourceKey.SubjectName = NAME_None;
-			Source = MakeShared<FLiveLinkSubjectUIEntry>(SourceKey, Client);
-			SubjectData.Add(Source);
-			SourceHeaderItems.Add(SubjectKey.Source) = Source;
-
-			SubjectsTreeView->SetItemExpansion(Source, true);
-			AllItems.Add(Source);
-		}
-
-		FLiveLinkSubjectUIEntryPtr SubjectEntry = MakeShared<FLiveLinkSubjectUIEntry>(SubjectKey, Client);
-		Source->Children.Add(SubjectEntry);
-		AllItems.Add(SubjectEntry);
-	}
-
-	auto SortPredicate = [](const FLiveLinkSubjectUIEntryPtr& LHS, const FLiveLinkSubjectUIEntryPtr& RHS) {return LHS->GetItemText().CompareTo(RHS->GetItemText()) < 0; };
-	SubjectData.Sort(SortPredicate);
-	for (FLiveLinkSubjectUIEntryPtr& Subject : SubjectData)
-	{
-		Subject->Children.Sort(SortPredicate);
-	}
-
-	for (const FLiveLinkSubjectUIEntryPtr& Item : AllItems)
-	{
-		for (FLiveLinkSubjectKey& Selection : SavedSelection)
-		{
-			if (Item->SubjectKey == Selection)
-			{
-				SubjectsTreeView->SetItemSelection(Item, true);
-				break;
-			}
-		}
-	}
-
-	SubjectsTreeView->RequestTreeRefresh();
-}
-
-void SLiveLinkClientPanel::OnPropertyChanged(const FPropertyChangedEvent& InEvent)
-{
-	TArray<FLiveLinkSourceUIEntryPtr> Selected;
-	SourceListView->GetSelectedItems(Selected);
-	for (FLiveLinkSourceUIEntryPtr Item : Selected)
-	{
-		Client->OnPropertyChanged(Item->GetGuid(), InEvent);
-	}
-}
-
-void SLiveLinkClientPanel::HandleRemoveSource()
-{
-	TArray<FLiveLinkSourceUIEntryPtr> Selected;
-	SourceListView->GetSelectedItems(Selected);
-	if (Selected.Num() > 0)
-	{
-		Selected[0]->RemoveFromClient();
-	}
-}
-
-bool SLiveLinkClientPanel::CanRemoveSource()
-{
-	return SourceListView->GetNumItemsSelected() > 0;
-}
-
-bool SLiveLinkClientPanel::HasSource() const
-{
-	constexpr bool bIncludeVirtualSources = true;
-	return Client->GetDisplayableSources(bIncludeVirtualSources).Num() > 0;
-}
-
-void SLiveLinkClientPanel::HandleRemoveAllSources()
-{
-	Client->RemoveAllSources();
-}
-
-bool SLiveLinkClientPanel::CanRemoveSubject() const
-{
-	TArray<FLiveLinkSubjectUIEntryPtr> Selected;
-	SubjectsTreeView->GetSelectedItems(Selected);
-	return Selected.Num() > 0 && Selected[0] && Selected[0]->IsVirtualSubject();
-}
-
-void SLiveLinkClientPanel::HandleRemoveSubject()
-{
-	TArray<FLiveLinkSubjectUIEntryPtr> Selected;
-	SubjectsTreeView->GetSelectedItems(Selected);
-	if (Selected.Num() > 0 && Selected[0])
-	{
-		Selected[0]->RemoveFromClient();
-	}
-}
-
-void SLiveLinkClientPanel::OnSourcesChangedHandler()
-{
-	RefreshSourceData(true);
-	RebuildSubjectList();
-}
-
-void SLiveLinkClientPanel::OnSubjectsChangedHandler()
-{
-	RebuildSubjectList();
-	SettingsDetailsView->ForceRefresh();
+	return PanelController->SubjectsDetailsView->GetSubjectKey().Source.IsValid() && !PanelController->SubjectsDetailsView->GetSubjectKey().SubjectName.IsNone() ? 1 : 0;
 }
 
 void SLiveLinkClientPanel::PostUndo(bool bSuccess)
 {
-	SettingsDetailsView->ForceRefresh();
+	PanelController->SourcesDetailsView->ForceRefresh();
 }
 
 void SLiveLinkClientPanel::PostRedo(bool bSuccess)
 {
-	SettingsDetailsView->ForceRefresh();
+	PanelController->SourcesDetailsView->ForceRefresh();
 }
 
 EVisibility SLiveLinkClientPanel::ShowEditorPerformanceThrottlingWarning() const
diff --git a/Engine/Plugins/Animation/LiveLink/Source/LiveLinkEditor/Private/LiveLinkClientPanel.h b/Engine/Plugins/Animation/LiveLink/Source/LiveLinkEditor/Private/LiveLinkClientPanel.h
--- a/Engine/Plugins/Animation/LiveLink/Source/LiveLinkEditor/Private/LiveLinkClientPanel.h
+++ b/Engine/Plugins/Animation/LiveLink/Source/LiveLinkEditor/Private/LiveLinkClientPanel.h
@@ -2,11 +2,12 @@
 
 #pragma once
 
+#include "EditorUndoClient.h"
+#include "LiveLinkPanelController.h"
+#include "Templates/SharedPointer.h"
 #include "UObject/GCObject.h"
 #include "Widgets/SCompoundWidget.h"
 
-#include "EditorUndoClient.h"
-
 namespace ESelectInfo { enum Type : int; }
 namespace ETextCommit { enum Type : int; }
 struct FPropertyChangedEvent;
@@ -23,20 +24,17 @@ class ULiveLinkSourceFactory;
 class SLiveLinkSourceListView;
 class STableViewBase;
 
-typedef TSharedPtr<FLiveLinkSourceUIEntry> FLiveLinkSourceUIEntryPtr;
-typedef TSharedPtr<FLiveLinkSubjectUIEntry> FLiveLinkSubjectUIEntryPtr;
-
-class SLiveLinkClientPanel : public SCompoundWidget, public FGCObject, public FEditorUndoClient
+class LIVELINKEDITOR_API SLiveLinkClientPanel : public SCompoundWidget, public FGCObject, public FEditorUndoClient
 {
 	SLATE_BEGIN_ARGS(SLiveLinkClientPanel){}
 	SLATE_END_ARGS()
 
-	~SLiveLinkClientPanel();
-
+	virtual ~SLiveLinkClientPanel() override;
+	
 	void Construct(const FArguments& Args, FLiveLinkClient* InClient);
 
 	// FGCObject interface
-	void AddReferencedObjects(FReferenceCollector& Collector) override;
+	virtual void AddReferencedObjects(FReferenceCollector& Collector) override;
 	virtual FString GetReferencerName() const override
 	{
 		return TEXT("SLiveLinkClientPanel");
@@ -44,49 +42,11 @@ class SLiveLinkClientPanel : public SCompoundWidget, public FGCObject, public FE
 	// End of FGCObject interface
 
 	// FEditorUndoClient interface
-	virtual void PostUndo(bool bSuccess);
-	virtual void PostRedo(bool bSuccess);
+	virtual void PostUndo(bool bSuccess) override;
+	virtual void PostRedo(bool bSuccess) override;
 	// End FEditorUndoClient interface
 
 private:
-
-	void BindCommands();
-
-	void RefreshSourceData(bool bRefreshUI);
-
-	void RetrieveFactorySourcePanel(FMenuBuilder& MenuBuilder, ULiveLinkSourceFactory* FactoryCDO);
-
-	FReply OnCloseSourceSelectionPanel(ULiveLinkSourceFactory* FactoryCDO, bool bMakeSource);
-
-	void HandleOnAddFromFactory(ULiveLinkSourceFactory* InSourceFactory);
-
-	// Starts adding a new virtual subject to the client
-	void AddVirtualSubject();
-
-	// Handles adding a virtual subject after the user has picked a name
-	void HandleAddVirtualSubject(const FText& NewSubjectName, ETextCommit::Type CommitInfo);
-
-	// Callback when property changes on source settings
-	void OnPropertyChanged(const FPropertyChangedEvent& InEvent);
-
-	// Remove source command handlers
-	void HandleRemoveSource();
-	bool CanRemoveSource();
-
-	// Remove all sources command handlers
-	void HandleRemoveAllSources();
-	bool HasSource() const;
-
-	// Remove subject command handlers
-	bool CanRemoveSubject() const;
-	void HandleRemoveSubject();
-
-	// Registered with the client and called when client's sources change
-	void OnSourcesChangedHandler();
-
-	// Registered with the client and called when client's subjects change
-	void OnSubjectsChangedHandler();
-
 	// Controls whether the editor performance throttling warning should be visible
 	EVisibility ShowEditorPerformanceThrottlingWarning() const;
 
@@ -98,61 +58,21 @@ private:
 
 	// Return the occurrence count and last time occurred text
 	FText GetSelectedMessageOccurrenceText() const;
-
 private:
 	int32 GetDetailWidgetIndex() const;
-
-	TSharedRef<ITableRow> MakeSourceListViewWidget(FLiveLinkSourceUIEntryPtr Entry, const TSharedRef<STableViewBase>& OwnerTable) const;
-	TSharedPtr<SWidget> OnSourceConstructContextMenu();
-
-	// Helper functions for building the subject tree UI
-	TSharedRef<ITableRow> MakeTreeRowWidget(FLiveLinkSubjectUIEntryPtr InInfo, const TSharedRef<STableViewBase>& OwnerTable);
-	void GetChildrenForInfo(FLiveLinkSubjectUIEntryPtr InInfo, TArray< FLiveLinkSubjectUIEntryPtr >& OutChildren);
-	void RebuildSubjectList();
-
-	// Handler for the source list selection changing
-	void OnSourceListSelectionChanged(FLiveLinkSourceUIEntryPtr Entry, ESelectInfo::Type SelectionType) const;
-
-	// Handler for the subject tree selection changing
-	void OnSubjectTreeSelectionChanged(FLiveLinkSubjectUIEntryPtr BoneInfo, ESelectInfo::Type SelectInfo);
-
+	
 	// Handler for the subject tree context menu opening
 	TSharedPtr<SWidget> OnOpenVirtualSubjectContextMenu();
 
-	// Source list widget
-	TSharedPtr<SLiveLinkSourceListView> SourceListView;
-
-	// Source list items
-	TArray<FLiveLinkSourceUIEntryPtr> SourceData;
+	// Handles callback connections between the sources, subjects and details views
+	TSharedPtr<FLiveLinkPanelController> PanelController;
 
-	// Subject tree widget
-	TSharedPtr<STreeView<FLiveLinkSubjectUIEntryPtr>> SubjectsTreeView;
-
-	// Subject tree items
-	TArray<FLiveLinkSubjectUIEntryPtr> SubjectData;
-
-	TSharedPtr<FUICommandList> CommandList;
-
-	FLiveLinkClient* Client;
-
-	// Reference to connection settings struct details panel
-	TSharedPtr<class IDetailsView> SettingsDetailsView;
-
-	// Reference to the data value struct details panel
-	TSharedPtr<class SLiveLinkDataView> DataDetailsView;
-
-	// Handle to delegate when client sources list has changed */
-	FDelegateHandle OnSourcesChangedHandle;
-
-	// Handle to delegate when a client subjects list has changed */
-	FDelegateHandle OnSubjectsChangedHandle;
+	// Pointer to the livelink client
+	FLiveLinkClient* Client = nullptr;
 
 	// Map to cover 
 	TMap<TObjectPtr<UClass>, TObjectPtr<UObject>> DetailsPanelEditorObjects;
 
 	// Details index
-	int32 DetailWidgetIndex;
-
-	// Guard from reentrant selection
-	mutable bool bSelectionChangedGuard;
+	int32 DetailWidgetIndex = -1;
 };
diff --git a/Engine/Plugins/Animation/LiveLink/Source/LiveLinkEditor/Private/LiveLinkClientPanelToolbar.cpp b/Engine/Plugins/Animation/LiveLink/Source/LiveLinkEditor/Private/LiveLinkClientPanelToolbar.cpp
--- a/Engine/Plugins/Animation/LiveLink/Source/LiveLinkEditor/Private/LiveLinkClientPanelToolbar.cpp
+++ b/Engine/Plugins/Animation/LiveLink/Source/LiveLinkEditor/Private/LiveLinkClientPanelToolbar.cpp
@@ -408,6 +408,8 @@ void SLiveLinkClientPanelToolbar::Construct(const FArguments& Args, FLiveLinkCli
 {
 	Client = InClient;
 
+	ParentWindowOverride = Args._ParentWindow;
+
 	TArray<UClass*> Results;
 	GetDerivedClasses(ULiveLinkSourceFactory::StaticClass(), Results, true);
 	for (UClass* SourceFactory : Results)
@@ -765,7 +767,7 @@ TSharedRef<SWidget> SLiveLinkClientPanelToolbar::OnPresetGeneratePresetsMenu()
 	return MenuBuilder.MakeWidget();
 }
 
-static bool OpenSaveDialog(const FString& InDefaultPath, const FString& InNewNameSuggestion, FString& OutPackageName)
+static bool OpenSaveDialog(const FString& InDefaultPath, const FString& InNewNameSuggestion, TSharedPtr<SWindow> InParentWindowOverride, FString& OutPackageName)
 {
 	FSaveAssetDialogConfig SaveAssetDialogConfig;
 	{
@@ -774,6 +776,7 @@ static bool OpenSaveDialog(const FString& InDefaultPath, const FString& InNewNam
 		SaveAssetDialogConfig.AssetClassNames.Add(ULiveLinkPreset::StaticClass()->GetClassPathName());
 		SaveAssetDialogConfig.ExistingAssetPolicy = ESaveAssetDialogExistingAssetPolicy::AllowButWarn;
 		SaveAssetDialogConfig.DialogTitleOverride = LOCTEXT("SaveLiveLinkPresetDialogTitle", "Save LiveLink Preset");
+		SaveAssetDialogConfig.WindowOverride = InParentWindowOverride;
 	}
 
 	FContentBrowserModule& ContentBrowserModule = FModuleManager::LoadModuleChecked<FContentBrowserModule>("ContentBrowser");
@@ -788,7 +791,7 @@ static bool OpenSaveDialog(const FString& InDefaultPath, const FString& InNewNam
 	return false;
 }
 
-bool GetSavePresetPackageName(FString& OutName)
+bool GetSavePresetPackageName(TSharedPtr<SWindow> InParentWindowOverride, FString& OutName)
 {
 	ULiveLinkUserSettings* ConfigSettings = GetMutableDefault<ULiveLinkUserSettings>();
 
@@ -825,7 +828,7 @@ bool GetSavePresetPackageName(FString& OutName)
 	bool bFilenameValid = false;
 	while (!bFilenameValid)
 	{
-		if (!OpenSaveDialog(DialogStartPath, DialogStartName, UserPackageName))
+		if (!OpenSaveDialog(DialogStartPath, DialogStartName, InParentWindowOverride, UserPackageName))
 		{
 			return false;
 		}
@@ -845,7 +848,7 @@ bool GetSavePresetPackageName(FString& OutName)
 void SLiveLinkClientPanelToolbar::OnSaveAsPreset()
 {
 	FString PackageName;
-	if (!GetSavePresetPackageName(PackageName))
+	if (!GetSavePresetPackageName(ParentWindowOverride, PackageName))
 	{
 		return;
 	}
diff --git a/Engine/Plugins/Animation/LiveLink/Source/LiveLinkEditor/Private/LiveLinkClientPanelToolbar.h b/Engine/Plugins/Animation/LiveLink/Source/LiveLinkEditor/Private/LiveLinkClientPanelToolbar.h
deleted file mode 100644
diff --git a/Engine/Plugins/Animation/LiveLink/Source/LiveLinkEditor/Private/LiveLinkClientPanelViews.cpp b/Engine/Plugins/Animation/LiveLink/Source/LiveLinkEditor/Private/LiveLinkClientPanelViews.cpp
new file mode 100644
--- /dev/null
+++ b/Engine/Plugins/Animation/LiveLink/Source/LiveLinkEditor/Private/LiveLinkClientPanelViews.cpp
@@ -0,0 +1,551 @@
+// Copyright Epic Games, Inc. All Rights Reserved.
+
+#include "LiveLinkClientPanelViews.h"
+
+#include "EditorFontGlyphs.h"
+#include "Framework/Commands/UICommandList.h"
+#include "IDetailsView.h"
+#include "Internationalization/Text.h"
+#include "LiveLinkClient.h"
+#include "LiveLinkClientCommands.h"
+#include "LiveLinkSettings.h"
+#include "LiveLinkTypes.h"
+#include "Framework/MultiBox/MultiBoxBuilder.h"
+#include "Modules/ModuleManager.h"
+#include "PropertyEditorModule.h"
+#include "SLiveLinkDataView.h"
+#include "Widgets/DeclarativeSyntaxSupport.h"
+#include "Widgets/Input/SButton.h"
+#include "Widgets/Input/SCheckBox.h"
+
+
+#define LOCTEXT_NAMESPACE "LiveLinkClientPanel.PanelViews"
+
+// Static Source UI FNames
+namespace SourceListUI
+{
+	static const FName TypeColumnName(TEXT("Type"));
+	static const FName MachineColumnName(TEXT("Machine"));
+	static const FName StatusColumnName(TEXT("Status"));
+	static const FName ActionsColumnName(TEXT("Action"));
+};
+
+// Static Subject UI FNames
+namespace SubjectTreeUI
+{
+	static const FName EnabledColumnName(TEXT("Enabled"));
+	static const FName NameColumnName(TEXT("Name"));
+	static const FName RoleColumnName(TEXT("Role"));
+	static const FName ActionsColumnName(TEXT("Action"));
+};
+
+namespace UE::LiveLink
+{
+TSharedPtr<IDetailsView> CreateSourcesDetailsView(const TSharedPtr<FLiveLinkSourcesView>& InSourcesView)
+{
+	FPropertyEditorModule& PropertyEditorModule = FModuleManager::GetModuleChecked<FPropertyEditorModule>("PropertyEditor");
+
+	FDetailsViewArgs DetailsViewArgs;
+	DetailsViewArgs.bUpdatesFromSelection = false;
+	DetailsViewArgs.bLockable = false;
+	DetailsViewArgs.bShowPropertyMatrixButton = false;
+	DetailsViewArgs.NameAreaSettings = FDetailsViewArgs::HideNameArea;
+	DetailsViewArgs.ViewIdentifier = NAME_None;
+	TSharedPtr<IDetailsView> SettingsDetailsView = PropertyEditorModule.CreateDetailView(DetailsViewArgs);
+	// todo: use controller here instead of view widget
+	SettingsDetailsView->OnFinishedChangingProperties().AddRaw(InSourcesView.Get(), &FLiveLinkSourcesView::OnPropertyChanged);
+
+	return SettingsDetailsView;
+}
+
+TSharedPtr<SLiveLinkDataView> CreateSubjectsDetailsView(FLiveLinkClient* InLiveLinkClient)
+{
+	return SNew(SLiveLinkDataView, InLiveLinkClient);
+}
+} // namespace UE::LiveLink
+
+FGuid FLiveLinkSourceUIEntry::GetGuid() const
+{
+	return EntryGuid;
+}
+FText FLiveLinkSourceUIEntry::GetSourceType() const
+{
+	return Client->GetSourceType(EntryGuid);
+}
+FText FLiveLinkSourceUIEntry::GetMachineName() const
+{
+	return Client->GetSourceMachineName(EntryGuid);
+}
+FText FLiveLinkSourceUIEntry::GetStatus() const
+{
+	return Client->GetSourceStatus(EntryGuid);
+}
+ULiveLinkSourceSettings* FLiveLinkSourceUIEntry::GetSourceSettings() const
+{
+	return Client->GetSourceSettings(EntryGuid);
+}
+void FLiveLinkSourceUIEntry::RemoveFromClient() const
+{
+	Client->RemoveSource(EntryGuid);
+}
+FText FLiveLinkSourceUIEntry::GetDisplayName() const
+{
+	return GetSourceType();
+}
+
+FLiveLinkSubjectUIEntry::FLiveLinkSubjectUIEntry(const FLiveLinkSubjectKey& InSubjectKey, FLiveLinkClient* InClient)
+	: SubjectKey(InSubjectKey)
+	, Client(InClient)
+{
+	if (InClient)
+	{
+		bIsVirtualSubject = InClient->IsVirtualSubject(InSubjectKey);
+	}
+}
+
+bool FLiveLinkSubjectUIEntry::IsSubject() const
+{
+	return !SubjectKey.SubjectName.IsNone();
+}
+
+bool FLiveLinkSubjectUIEntry::IsSource() const
+{
+	return SubjectKey.SubjectName.IsNone();
+}
+
+bool FLiveLinkSubjectUIEntry::IsVirtualSubject() const
+{
+	return IsSubject() && bIsVirtualSubject;
+}
+
+UObject* FLiveLinkSubjectUIEntry::GetSettings() const
+{
+	if (IsSource())
+	{
+		return Client->GetSourceSettings(SubjectKey.Source);
+	}
+	else
+	{
+		return Client->GetSubjectSettings(SubjectKey);
+	}
+}
+
+bool FLiveLinkSubjectUIEntry::IsSubjectEnabled() const
+{
+	return IsSubject() ? Client->IsSubjectEnabled(SubjectKey, false) : false;
+}
+
+bool FLiveLinkSubjectUIEntry::IsSubjectValid() const
+{
+	return IsSubject() ? Client->IsSubjectValid(SubjectKey) : false;
+}
+
+void FLiveLinkSubjectUIEntry::SetSubjectEnabled(bool bIsEnabled)
+{
+	if (IsSubject())
+	{
+		Client->SetSubjectEnabled(SubjectKey, bIsEnabled);
+	}
+}
+
+FText FLiveLinkSubjectUIEntry::GetItemText() const
+{
+	if (IsSource())
+	{
+		return Client->GetSourceType(SubjectKey.Source);
+	}
+	else
+	{
+		return FText::FromName(SubjectKey.SubjectName);
+	}
+}
+
+TSubclassOf<ULiveLinkRole> FLiveLinkSubjectUIEntry::GetItemRole() const
+{
+	return IsSubject() ? Client->GetSubjectRole(SubjectKey) : TSubclassOf<ULiveLinkRole>();
+}
+
+void FLiveLinkSubjectUIEntry::RemoveFromClient() const
+{
+	Client->RemoveSubject_AnyThread(SubjectKey);
+}
+
+
+class SLiveLinkClientPanelSubjectRow : public SMultiColumnTableRow<FLiveLinkSubjectUIEntryPtr>
+{
+public:
+	SLATE_BEGIN_ARGS(SLiveLinkClientPanelSubjectRow) {}
+	/** The list item for this row */
+	SLATE_ARGUMENT(FLiveLinkSubjectUIEntryPtr, Entry)
+	SLATE_END_ARGS()
+
+
+	void Construct(const FArguments& Args, const TSharedRef<STableViewBase>& OwnerTableView)
+	{
+		EntryPtr = Args._Entry;
+
+		SMultiColumnTableRow<FLiveLinkSubjectUIEntryPtr>::Construct(
+			FSuperRowType::FArguments()
+			.Padding(1.0f),
+			OwnerTableView
+		);
+	}
+
+	/** Overridden from SMultiColumnTableRow.  Generates a widget for this column of the list view. */
+	virtual TSharedRef<SWidget> GenerateWidgetForColumn(const FName& ColumnName) override
+	{
+		if (ColumnName == SubjectTreeUI::EnabledColumnName)
+		{
+			if (EntryPtr->IsSubject())
+			{
+				return SNew(SCheckBox)
+					.IsChecked(MakeAttributeSP(this, &SLiveLinkClientPanelSubjectRow::GetSubjectEnabled))
+					.OnCheckStateChanged(this, &SLiveLinkClientPanelSubjectRow::OnEnabledChanged);
+			}
+		}
+		else if (ColumnName == SubjectTreeUI::NameColumnName)
+		{
+			return SNew(SHorizontalBox)
+				+ SHorizontalBox::Slot()
+				.AutoWidth()
+				.Padding(6, 0, 0, 0)
+				[
+					SNew(SExpanderArrow, SharedThis(this)).IndentAmount(12)
+				]
+				+ SHorizontalBox::Slot()
+				.FillWidth(1.0f)
+				[
+					SNew(STextBlock)
+					.Text(MakeAttributeSP(this, &SLiveLinkClientPanelSubjectRow::GetItemText))
+				];
+		}
+		else if (ColumnName == SubjectTreeUI::RoleColumnName)
+		{
+			auto RoleAttribute = MakeAttributeSP(this, &SLiveLinkClientPanelSubjectRow::GetItemRole);
+			return SNew(STextBlock)
+				.Text(EntryPtr->IsSubject() ? RoleAttribute : FText::GetEmpty());
+		}
+		else if (ColumnName == SubjectTreeUI::ActionsColumnName)
+		{
+			if (EntryPtr->IsVirtualSubject())
+			{
+				return SNew(SButton)
+					.ButtonStyle(FAppStyle::Get(), "HoverHintOnly")
+					.HAlign(HAlign_Center)
+					.VAlign(VAlign_Center)
+					.OnClicked(this, &SLiveLinkClientPanelSubjectRow::OnRemoveClicked)
+					.ToolTipText(LOCTEXT("RemoveVirtualSubject", "Remove selected live link virtual subject"))
+					.ContentPadding(0.f)
+					.ForegroundColor(FSlateColor::UseForeground())
+					.IsFocusable(false)
+					[
+						SNew(SImage)
+						.Image(FAppStyle::GetBrush("Icons.Delete"))
+						.ColorAndOpacity(FSlateColor::UseForeground())
+					];
+			}
+			else
+			{
+				return SNew(SBox)
+					.HAlign(HAlign_Center)
+					.VAlign(VAlign_Center)
+					[
+						SNew(STextBlock)
+						.Font(FAppStyle::Get().GetFontStyle("FontAwesome.8"))
+						.ColorAndOpacity(this, &SLiveLinkClientPanelSubjectRow::OnGetActivityColor)
+						.Text(FEditorFontGlyphs::Circle)
+					];
+			}
+		}
+
+		return SNullWidget::NullWidget;
+	}
+
+private:
+	ECheckBoxState GetSubjectEnabled() const { return EntryPtr->IsSubjectEnabled() ? ECheckBoxState::Checked : ECheckBoxState::Unchecked; }
+	void OnEnabledChanged(ECheckBoxState NewState) { EntryPtr->SetSubjectEnabled(NewState == ECheckBoxState::Checked); }
+	FText GetItemText() const { return EntryPtr->GetItemText(); }
+	FText GetItemRole() const
+	{
+		TSubclassOf<ULiveLinkRole> Role = EntryPtr->GetItemRole();
+		if (Role.Get())
+		{
+			return Role->GetDefaultObject<ULiveLinkRole>()->GetDisplayName();
+		}
+		return FText::GetEmpty();
+	}
+
+	FReply OnRemoveClicked()
+	{
+		EntryPtr->RemoveFromClient();
+		return FReply::Handled();
+	}
+
+	FSlateColor OnGetActivityColor() const
+	{
+		if (EntryPtr->IsSubjectEnabled())
+		{
+			return EntryPtr->IsSubjectValid() ? GetDefault<ULiveLinkSettings>()->GetValidColor() : GetDefault<ULiveLinkSettings>()->GetInvalidColor();
+		}
+		return FLinearColor(0.f, 0.f, 0.f, 0.f);
+	}
+
+	FLiveLinkSubjectUIEntryPtr EntryPtr;
+};
+
+class SLiveLinkClientPanelSourcesRow : public SMultiColumnTableRow<FLiveLinkSourceUIEntryPtr>
+{
+public:
+	SLATE_BEGIN_ARGS(SLiveLinkClientPanelSourcesRow) {}
+	/** The list item for this row */
+		SLATE_ARGUMENT(FLiveLinkSourceUIEntryPtr, Entry)
+	SLATE_END_ARGS()
+
+	void Construct(const FArguments& Args, const TSharedRef<STableViewBase>& OwnerTableView)
+	{
+		EntryPtr = Args._Entry;
+
+		SMultiColumnTableRow<FLiveLinkSourceUIEntryPtr>::Construct(
+			FSuperRowType::FArguments()
+			.Padding(1.0f),
+			OwnerTableView
+		);
+	}
+
+	/** Overridden from SMultiColumnTableRow.  Generates a widget for this column of the list view. */
+	virtual TSharedRef<SWidget> GenerateWidgetForColumn(const FName& ColumnName) override
+	{
+		if (ColumnName == SourceListUI::TypeColumnName)
+		{
+			return SNew(STextBlock)
+				.Text(EntryPtr->GetSourceType());
+		}
+		else if (ColumnName == SourceListUI::MachineColumnName)
+		{
+			return SNew(STextBlock)
+				.Text(MakeAttributeSP(this, &SLiveLinkClientPanelSourcesRow::GetMachineName));
+		}
+		else if (ColumnName == SourceListUI::StatusColumnName)
+		{
+			return SNew(STextBlock)
+				.Text(MakeAttributeSP(this, &SLiveLinkClientPanelSourcesRow::GetSourceStatus));
+		}
+		else if (ColumnName == SourceListUI::ActionsColumnName)
+		{
+			return SNew(SButton)
+				.ButtonStyle(FAppStyle::Get(), "HoverHintOnly")
+				.HAlign(HAlign_Center)
+				.VAlign(VAlign_Center)
+				.OnClicked(this, &SLiveLinkClientPanelSourcesRow::OnRemoveClicked)
+				.ToolTipText(LOCTEXT("RemoveSource", "Remove selected live link source"))
+				.ContentPadding(0.f)
+				.ForegroundColor(FSlateColor::UseForeground())
+				.IsFocusable(false)
+				[
+					SNew(SImage)
+					.Image(FAppStyle::GetBrush("Icons.Delete"))
+					.ColorAndOpacity(FSlateColor::UseForeground())
+				];
+		}
+
+		return SNullWidget::NullWidget;
+	}
+
+private:
+	FText GetMachineName() const
+	{
+		return EntryPtr->GetMachineName();
+	}
+
+	FText GetSourceStatus() const
+	{
+		return EntryPtr->GetStatus();
+	}
+
+	FReply OnRemoveClicked()
+	{
+		EntryPtr->RemoveFromClient();
+		return FReply::Handled();
+	}
+
+	FLiveLinkSourceUIEntryPtr EntryPtr;
+};
+
+FLiveLinkSourcesView::FLiveLinkSourcesView(FLiveLinkClient* InLiveLinkClient, TSharedPtr<FUICommandList> InCommandList, FOnSourceSelectionChanged InOnSourceSelectionChanged)
+	: Client(InLiveLinkClient)
+	, OnSourceSelectionChangedDelegate(MoveTemp(InOnSourceSelectionChanged))
+{
+	CreateSourcesListView(InCommandList);
+}
+
+TSharedRef<ITableRow> FLiveLinkSourcesView::MakeSourceListViewWidget(FLiveLinkSourceUIEntryPtr Entry, const TSharedRef<STableViewBase>& OwnerTable) const
+{
+	return SNew(SLiveLinkClientPanelSourcesRow, OwnerTable)
+		.Entry(Entry);
+}
+
+void FLiveLinkSourcesView::OnSourceListSelectionChanged(FLiveLinkSourceUIEntryPtr Entry, ESelectInfo::Type SelectionType) const
+{
+	OnSourceSelectionChangedDelegate.Execute(Entry, SelectionType);
+}
+
+void FLiveLinkSourcesView::CreateSourcesListView(const TSharedPtr<FUICommandList>& InCommandList)
+{
+	SAssignNew(SourcesListView, SLiveLinkSourceListView)
+		.ListItemsSource(&SourceData)
+		.SelectionMode(ESelectionMode::Single)
+		.OnGenerateRow_Raw(this, &FLiveLinkSourcesView::MakeSourceListViewWidget)
+		.OnContextMenuOpening_Raw(this, &FLiveLinkSourcesView::OnSourceConstructContextMenu, InCommandList)
+		.OnSelectionChanged_Raw(this, &FLiveLinkSourcesView::OnSourceListSelectionChanged)
+		.HeaderRow
+		(
+			SNew(SHeaderRow)
+			+ SHeaderRow::Column(SourceListUI::TypeColumnName)
+			.FillWidth(25.f)
+			.DefaultLabel(LOCTEXT("TypeColumnHeaderName", "Source Type"))
+			+ SHeaderRow::Column(SourceListUI::MachineColumnName)
+			.FillWidth(25.f)
+			.DefaultLabel(LOCTEXT("MachineColumnHeaderName", "Source Machine"))
+			+ SHeaderRow::Column(SourceListUI::StatusColumnName)
+			.FillWidth(50.f)
+			.DefaultLabel(LOCTEXT("StatusColumnHeaderName", "Status"))
+			+ SHeaderRow::Column(SourceListUI::ActionsColumnName)
+			.ManualWidth(20.f)
+			.DefaultLabel(LOCTEXT("ActionsColumnHeaderName", ""))
+		);
+}
+
+TSharedPtr<SWidget> FLiveLinkSourcesView::OnSourceConstructContextMenu(TSharedPtr<FUICommandList> InCommandList)
+{
+	const bool bShouldCloseWindowAfterMenuSelection = true;
+	FMenuBuilder MenuBuilder(bShouldCloseWindowAfterMenuSelection, InCommandList);
+
+	MenuBuilder.BeginSection(TEXT("Remove"));
+	{
+		if (CanRemoveSource())
+		{
+			MenuBuilder.AddMenuEntry(FLiveLinkClientCommands::Get().RemoveSource);
+		}
+		MenuBuilder.AddMenuEntry(FLiveLinkClientCommands::Get().RemoveAllSources);
+	}
+	MenuBuilder.EndSection();
+
+	return MenuBuilder.MakeWidget();
+}
+
+void FLiveLinkSourcesView::RefreshSourceData(bool bRefreshUI)
+{
+	SourceData.Reset();
+
+	for (FGuid SourceGuid : Client->GetDisplayableSources())
+	{
+		SourceData.Add(MakeShared<FLiveLinkSourceUIEntry>(SourceGuid, Client));
+	}
+	SourceData.Sort([](const FLiveLinkSourceUIEntryPtr& LHS, const FLiveLinkSourceUIEntryPtr& RHS) { return LHS->GetMachineName().CompareTo(RHS->GetMachineName()) < 0; });
+
+	if (bRefreshUI)
+	{
+		SourcesListView->RequestListRefresh();
+	}
+}
+
+void FLiveLinkSourcesView::HandleRemoveSource()
+{
+	TArray<FLiveLinkSourceUIEntryPtr> Selected;
+	SourcesListView->GetSelectedItems(Selected);
+	if (Selected.Num() > 0)
+	{
+		Selected[0]->RemoveFromClient();
+	}
+}
+
+bool FLiveLinkSourcesView::CanRemoveSource()
+{
+	return SourcesListView->GetNumItemsSelected() > 0;
+}
+
+FLiveLinkSubjectsView::FLiveLinkSubjectsView(FOnSubjectSelectionChanged InOnSubjectSelectionChanged, const TSharedPtr<FUICommandList>& InCommandList)
+	: SubjectSelectionChangedDelegate(InOnSubjectSelectionChanged)
+{
+	CreateSubjectsTreeView(InCommandList);
+}
+
+void FLiveLinkSubjectsView::OnSubjectSelectionChanged(FLiveLinkSubjectUIEntryPtr SubjectEntry, ESelectInfo::Type SelectInfo)
+{
+	SubjectSelectionChangedDelegate.Execute(SubjectEntry, SelectInfo);
+}
+
+void FLiveLinkSourcesView::OnPropertyChanged(const FPropertyChangedEvent& InEvent)
+{
+	TArray<FLiveLinkSourceUIEntryPtr> Selected;
+	SourcesListView->GetSelectedItems(Selected);
+	for (FLiveLinkSourceUIEntryPtr Item : Selected)
+	{
+		Client->OnPropertyChanged(Item->GetGuid(), InEvent);
+	}
+}
+
+TSharedRef<ITableRow> FLiveLinkSubjectsView::MakeTreeRowWidget(FLiveLinkSubjectUIEntryPtr InInfo, const TSharedRef<STableViewBase>& OwnerTable)
+{
+	return SNew(SLiveLinkClientPanelSubjectRow, OwnerTable)
+		.Entry(InInfo);
+}
+
+void FLiveLinkSubjectsView::GetChildrenForInfo(FLiveLinkSubjectUIEntryPtr InInfo, TArray< FLiveLinkSubjectUIEntryPtr >& OutChildren)
+{
+	OutChildren = InInfo->Children;
+}
+
+TSharedPtr<SWidget> FLiveLinkSubjectsView::OnOpenVirtualSubjectContextMenu(TSharedPtr<FUICommandList> InCommandList)
+{
+	constexpr bool bShouldCloseWindowAfterMenuSelection = true;
+	FMenuBuilder MenuBuilder(bShouldCloseWindowAfterMenuSelection, InCommandList);
+
+	MenuBuilder.BeginSection(TEXT("Remove"));
+	{
+		if (CanRemoveSubject())
+		{
+			MenuBuilder.AddMenuEntry(FLiveLinkClientCommands::Get().RemoveSubject);
+		}
+	}
+	MenuBuilder.EndSection();
+
+	return MenuBuilder.MakeWidget();
+}
+
+bool FLiveLinkSubjectsView::CanRemoveSubject() const
+{
+	TArray<FLiveLinkSubjectUIEntryPtr> Selected;
+	SubjectsTreeView->GetSelectedItems(Selected);
+	return Selected.Num() > 0 && Selected[0] && Selected[0]->IsVirtualSubject();
+}
+
+void FLiveLinkSubjectsView::CreateSubjectsTreeView(const TSharedPtr<FUICommandList>& InCommandList)
+{
+	SAssignNew(SubjectsTreeView, SLiveLinkSubjectsTreeView)
+		.TreeItemsSource(&SubjectData)
+		.OnGenerateRow_Raw(this, &FLiveLinkSubjectsView::MakeTreeRowWidget)
+		.OnGetChildren_Raw(this, &FLiveLinkSubjectsView::GetChildrenForInfo)
+		.OnSelectionChanged_Raw(this, &FLiveLinkSubjectsView::OnSubjectSelectionChanged)
+		.OnContextMenuOpening_Raw(this, &FLiveLinkSubjectsView::OnOpenVirtualSubjectContextMenu, InCommandList)
+		.SelectionMode(ESelectionMode::Single)
+		.HeaderRow
+		(
+			SNew(SHeaderRow)
+			+ SHeaderRow::Column(SubjectTreeUI::EnabledColumnName)
+			.DefaultLabel(LOCTEXT("EnabledName", ""))
+			.FixedWidth(22)
+			+ SHeaderRow::Column(SubjectTreeUI::NameColumnName)
+			.DefaultLabel(LOCTEXT("SubjectItemName", "Subject Name"))
+			.FillWidth(0.60f)
+			+ SHeaderRow::Column(SubjectTreeUI::RoleColumnName)
+			.DefaultLabel(LOCTEXT("RoleName", "Role"))
+			.FillWidth(0.40f)
+			+ SHeaderRow::Column(SubjectTreeUI::ActionsColumnName)
+			.ManualWidth(20.f)
+			.DefaultLabel(LOCTEXT("ActionsColumnHeaderName", ""))
+		);
+}
+
+
+#undef LOCTEXT_NAMESPACE /**LiveLinkClientPanel*/
\ No newline at end of file
diff --git a/Engine/Plugins/Animation/LiveLink/Source/LiveLinkEditor/Private/LiveLinkPanelController.cpp b/Engine/Plugins/Animation/LiveLink/Source/LiveLinkEditor/Private/LiveLinkPanelController.cpp
new file mode 100644
--- /dev/null
+++ b/Engine/Plugins/Animation/LiveLink/Source/LiveLinkEditor/Private/LiveLinkPanelController.cpp
@@ -0,0 +1,284 @@
+// Copyright Epic Games, Inc. All Rights Reserved.
+
+#include "LiveLinkPanelController.h"
+#include "Async/TaskGraphInterfaces.h"
+#include "Framework/Commands/UICommandList.h"
+#include "Features/IModularFeatures.h"
+#include "IDetailsView.h"
+#include "LiveLinkClient.h"
+#include "LiveLinkClientCommands.h"
+#include "LiveLinkClientPanelViews.h"
+#include "SLiveLinkDataView.h"
+
+#ifndef WITH_LIVELINK_HUB
+#define WITH_LIVELINK_HUB 0
+#endif
+
+FLiveLinkPanelController::FLiveLinkPanelController()
+{
+	Client = (FLiveLinkClient*)&IModularFeatures::Get().GetModularFeature<ILiveLinkClient>(ILiveLinkClient::ModularFeatureName);
+
+	OnSourcesChangedHandle = Client->OnLiveLinkSourcesChanged().AddRaw(this, &FLiveLinkPanelController::OnSourcesChangedHandler);
+	OnSubjectsChangedHandle = Client->OnLiveLinkSubjectsChanged().AddRaw(this, &FLiveLinkPanelController::OnSubjectsChangedHandler);
+
+	CommandList = MakeShared<FUICommandList>();
+	BindCommands();
+
+	SourcesView = MakeShared<FLiveLinkSourcesView>(Client, CommandList, FLiveLinkSourcesView::FOnSourceSelectionChanged::CreateRaw(this, &FLiveLinkPanelController::OnSourceSelectionChanged));
+	SubjectsView = MakeShared<FLiveLinkSubjectsView>(FLiveLinkSubjectsView::FOnSubjectSelectionChanged::CreateRaw(this, &FLiveLinkPanelController::OnSubjectSelectionChanged), CommandList);
+	SourcesDetailsView = UE::LiveLink::CreateSourcesDetailsView(SourcesView);
+	SubjectsDetailsView = UE::LiveLink::CreateSubjectsDetailsView(Client);
+
+	RebuildSourceList();
+	RebuildSubjectList();
+}
+
+FLiveLinkPanelController::~FLiveLinkPanelController()
+{
+	if (Client)
+	{
+		Client->OnLiveLinkSourcesChanged().Remove(OnSourcesChangedHandle);
+		OnSourcesChangedHandle.Reset();
+
+		Client->OnLiveLinkSubjectsChanged().Remove(OnSubjectsChangedHandle);
+		OnSubjectsChangedHandle.Reset();
+	}
+}
+
+void FLiveLinkPanelController::OnSourceSelectionChanged(FLiveLinkSourceUIEntryPtr Entry, ESelectInfo::Type SelectionType) const
+{
+	if (bSelectionChangedGuard)
+	{
+		return;
+	}
+	
+	TGuardValue<bool> ReentrantGuard(bSelectionChangedGuard, true);
+	
+	SubjectsDetailsView->SetSubjectKey(FLiveLinkSubjectKey());
+
+	int32 FoundSubjectEntryIndex = INDEX_NONE;
+	if (Entry.IsValid())
+	{
+		SourcesDetailsView->SetObject(Entry->GetSourceSettings());
+
+		// Find the corresponding subject
+		FoundSubjectEntryIndex = SubjectsView->SubjectData.IndexOfByPredicate([Entry](const FLiveLinkSubjectUIEntryPtr& SubjectEntry) { return SubjectEntry->SubjectKey.Source == Entry->GetGuid() && SubjectEntry->IsSource(); });
+	}
+	else
+	{
+		SourcesDetailsView->SetObject(nullptr);
+	}
+
+	// Set the corresponding subject
+	if (FoundSubjectEntryIndex != INDEX_NONE)
+	{
+		SubjectsView->SubjectsTreeView->SetSelection(SubjectsView->SubjectData[FoundSubjectEntryIndex]);
+	}
+	else
+	{
+		SubjectsView->SubjectsTreeView->ClearSelection();
+	}
+}
+
+void FLiveLinkPanelController::OnSubjectSelectionChanged(FLiveLinkSubjectUIEntryPtr SubjectEntry, ESelectInfo::Type SelectInfo)
+{
+	if (bSelectionChangedGuard)
+	{
+		return;
+	}
+	TGuardValue<bool> ReentrantGuard(bSelectionChangedGuard, true);
+
+	int32 FoundSourceIndex = INDEX_NONE;
+	bool bDetailViewSet = false;
+	if (SubjectEntry.IsValid())
+	{
+		// Find the corresponding Source entry
+		FGuid SourceGuid = SubjectEntry->SubjectKey.Source;
+		FoundSourceIndex = SourcesView->SourceData.IndexOfByPredicate([SourceGuid](FLiveLinkSourceUIEntryPtr SourceEntry) { return SourceEntry->GetGuid() == SourceGuid; });
+
+		if (SubjectEntry->IsSource())
+		{
+			SourcesDetailsView->SetObject(SubjectEntry->GetSettings());
+#if !WITH_LIVELINK_HUB // Sources and subjects live in different tabs in the hub
+			SubjectsDetailsView->SetSubjectKey(FLiveLinkSubjectKey());
+#endif
+		}
+		else
+		{
+#if !WITH_LIVELINK_HUB // Sources and subjects live in different tabs in the hub
+			SourcesDetailsView->SetObject(nullptr);
+#endif
+			SubjectsDetailsView->SetSubjectKey(SubjectEntry->SubjectKey);
+		}
+		bDetailViewSet = true;
+	}
+
+	if (!bDetailViewSet)
+	{
+		SourcesDetailsView->SetObject(nullptr);
+		SubjectsDetailsView->SetSubjectKey(FLiveLinkSubjectKey());
+	}
+
+	// Select the corresponding Source entry
+	if (FoundSourceIndex != INDEX_NONE)
+	{
+		SourcesView->SourcesListView->SetSelection(SourcesView->SourceData[FoundSourceIndex]);
+
+#if WITH_LIVELINK_HUB
+		// Update source details to the selected subject's source.
+		SourcesDetailsView->SetObject(SourcesView->SourceData[FoundSourceIndex]->GetSourceSettings());
+#endif
+	}
+	else
+	{
+		SourcesView->SourcesListView->ClearSelection();
+	}
+}
+
+void FLiveLinkPanelController::BindCommands()
+{
+	CommandList->MapAction(FLiveLinkClientCommands::Get().RemoveSource,
+		FExecuteAction::CreateRaw(this, &FLiveLinkPanelController::HandleRemoveSource),
+		FCanExecuteAction::CreateRaw(this, &FLiveLinkPanelController::CanRemoveSource)
+	);
+
+	CommandList->MapAction(FLiveLinkClientCommands::Get().RemoveAllSources,
+		FExecuteAction::CreateRaw(this, &FLiveLinkPanelController::HandleRemoveAllSources),
+		FCanExecuteAction::CreateRaw(this, &FLiveLinkPanelController::HasSource)
+	);
+
+	CommandList->MapAction(FLiveLinkClientCommands::Get().RemoveSubject,
+		FExecuteAction::CreateRaw(this, &FLiveLinkPanelController::HandleRemoveSubject),
+		FCanExecuteAction::CreateRaw(this, &FLiveLinkPanelController::CanRemoveSubject)
+	);
+}
+
+void FLiveLinkPanelController::OnSourcesChangedHandler()
+{
+	// Since  this can be called from any thread, make sure we only update slate on the game thread.
+	FFunctionGraphTask::CreateAndDispatchWhenReady([this]()
+	{
+		SourcesView->RefreshSourceData(true);
+		RebuildSubjectList();
+	}, TStatId(), nullptr, ENamedThreads::GameThread);
+}
+
+void FLiveLinkPanelController::OnSubjectsChangedHandler()
+{
+	// Since  this can be called from any thread, make sure we only update slate on the game thread.
+	FFunctionGraphTask::CreateAndDispatchWhenReady([this]()
+	{
+		RebuildSubjectList();
+		SourcesDetailsView->ForceRefresh();
+	}, TStatId(), nullptr, ENamedThreads::GameThread);
+}
+
+void FLiveLinkPanelController::RebuildSourceList()
+{
+	SourcesView->RefreshSourceData(true);
+}
+
+void FLiveLinkPanelController::RebuildSubjectList()
+{
+	TArray<FLiveLinkSubjectKey> SavedSelection;
+	{
+		TArray<FLiveLinkSubjectUIEntryPtr> SelectedItems = SubjectsView->SubjectsTreeView->GetSelectedItems();
+		for (const FLiveLinkSubjectUIEntryPtr& SelectedItem : SelectedItems)
+		{
+			SavedSelection.Add(SelectedItem->SubjectKey);
+		}
+	}
+
+	TArray<FLiveLinkSubjectKey> SubjectKeys = Client->GetSubjects(true, true);
+	SubjectsView->SubjectData.Reset();
+
+	TMap<FGuid, FLiveLinkSubjectUIEntryPtr> SourceHeaderItems;
+	TArray<FLiveLinkSubjectUIEntryPtr> AllItems;
+	AllItems.Reserve(SubjectKeys.Num());
+
+	for (const FLiveLinkSubjectKey& SubjectKey : SubjectKeys)
+	{
+		FLiveLinkSubjectUIEntryPtr Source;
+		if(FLiveLinkSubjectUIEntryPtr* SourcePtr = SourceHeaderItems.Find(SubjectKey.Source))
+		{
+			Source = *SourcePtr;
+		}
+		else
+		{
+			FLiveLinkSubjectKey SourceKey = SubjectKey;
+			SourceKey.SubjectName = NAME_None;
+			Source = MakeShared<FLiveLinkSubjectUIEntry>(SourceKey, Client);
+			SubjectsView->SubjectData.Add(Source);
+			SourceHeaderItems.Add(SubjectKey.Source) = Source;
+
+			SubjectsView->SubjectsTreeView->SetItemExpansion(Source, true);
+			AllItems.Add(Source);
+		}
+
+		FLiveLinkSubjectUIEntryPtr SubjectEntry = MakeShared<FLiveLinkSubjectUIEntry>(SubjectKey, Client);
+		Source->Children.Add(SubjectEntry);
+		AllItems.Add(SubjectEntry);
+	}
+
+	auto SortPredicate = [](const FLiveLinkSubjectUIEntryPtr& LHS, const FLiveLinkSubjectUIEntryPtr& RHS) {return LHS->GetItemText().CompareTo(RHS->GetItemText()) < 0; };
+	SubjectsView->SubjectData.Sort(SortPredicate);
+	for (FLiveLinkSubjectUIEntryPtr& Subject : SubjectsView->SubjectData)
+	{
+		Subject->Children.Sort(SortPredicate);
+	}
+
+	for (const FLiveLinkSubjectUIEntryPtr& Item : AllItems)
+	{
+		for (FLiveLinkSubjectKey& Selection : SavedSelection)
+		{
+			if (Item->SubjectKey == Selection)
+			{
+				SubjectsView->SubjectsTreeView->SetItemSelection(Item, true);
+				break;
+			}
+		}
+	}
+
+	SubjectsView->SubjectsTreeView->RequestTreeRefresh();
+}
+
+bool FLiveLinkPanelController::HasSource() const
+{
+	constexpr bool bIncludeVirtualSources = true;
+	return Client->GetDisplayableSources(bIncludeVirtualSources).Num() > 0;
+}
+
+bool FLiveLinkPanelController::CanRemoveSource() const
+{
+	return SourcesView->SourcesListView->GetNumItemsSelected() > 0;
+}
+
+void FLiveLinkPanelController::HandleRemoveSource()
+{
+	TArray<FLiveLinkSourceUIEntryPtr> Selected;
+	SourcesView->SourcesListView->GetSelectedItems(Selected);
+	if (Selected.Num() > 0)
+	{
+		Selected[0]->RemoveFromClient();
+	}
+}
+
+void FLiveLinkPanelController::HandleRemoveAllSources()
+{
+	Client->RemoveAllSources();
+}
+
+bool FLiveLinkPanelController::CanRemoveSubject() const
+{
+	return SubjectsView->CanRemoveSubject();
+}
+
+void FLiveLinkPanelController::HandleRemoveSubject()
+{
+	TArray<FLiveLinkSubjectUIEntryPtr> Selected;
+	SubjectsView->SubjectsTreeView->GetSelectedItems(Selected);
+	if (Selected.Num() > 0 && Selected[0])
+	{
+		Selected[0]->RemoveFromClient();
+	}
+}
\ No newline at end of file
diff --git a/Engine/Plugins/Animation/LiveLink/Source/LiveLinkEditor/Private/SLiveLinkDataView.h b/Engine/Plugins/Animation/LiveLink/Source/LiveLinkEditor/Private/SLiveLinkDataView.h
deleted file mode 100644
diff --git a/Engine/Source/Runtime/LiveLinkInterface/Public/ILiveLinkClient.h b/Engine/Source/Runtime/LiveLinkInterface/Public/ILiveLinkClient.h
--- a/Engine/Source/Runtime/LiveLinkInterface/Public/ILiveLinkClient.h
+++ b/Engine/Source/Runtime/LiveLinkInterface/Public/ILiveLinkClient.h
@@ -158,6 +158,9 @@ public:
 
 	/** Get the type of a source */
 	virtual FText GetSourceType(FGuid SourceGuid) const = 0;
+	
+	/** Get the status of a source */
+	virtual FText GetSourceStatus(FGuid EntryGuid) const = 0;
 
 	/** Push static data for a specific subject for a certain role. This will clear all buffered frames */
 	virtual void PushSubjectStaticData_AnyThread(const FLiveLinkSubjectKey& SubjectKey, TSubclassOf<ULiveLinkRole> Role, FLiveLinkStaticDataStruct&& StaticData) = 0;
diff --git a/Engine/Source/Runtime/LiveLinkMessageBusFramework/Public/LiveLinkProviderImpl.h b/Engine/Source/Runtime/LiveLinkMessageBusFramework/Public/LiveLinkProviderImpl.h
--- a/Engine/Source/Runtime/LiveLinkMessageBusFramework/Public/LiveLinkProviderImpl.h
+++ b/Engine/Source/Runtime/LiveLinkMessageBusFramework/Public/LiveLinkProviderImpl.h
@@ -29,7 +29,7 @@ struct FTrackedSubject
 };
 PRAGMA_ENABLE_DEPRECATION_WARNINGS
 
-struct FLiveLinkProvider : public ILiveLinkProvider
+struct LIVELINKMESSAGEBUSFRAMEWORK_API FLiveLinkProvider : public ILiveLinkProvider
 {
 private:
 	const FString ProviderName;
@@ -61,9 +61,6 @@ private:
 	void HandleHeartbeat(const FLiveLinkHeartbeatMessage& Message,
 						 const TSharedRef<class IMessageContext, ESPMode::ThreadSafe>& Context);
 	// End message bus message handlers
-
-	// Validate our current connections
-	void ValidateConnections();
 	
 	FTrackedSubject& GetTrackedSubject(const FName& SubjectName);
 
@@ -85,9 +82,13 @@ private:
 
 	void SendClearSubjectToConnections(FName SubjectName);
 
+protected:
+	// Get the addresses of all connected instances.
 	void GetConnectedAddresses(TArray<FMessageAddress>& Addresses);
 
-protected:
+	// Validate our current connections, removing those that have timed out.
+	void ValidateConnections();
+
 	template<typename MessageType>
 	void SendMessage(MessageType* Message)
 	{
