Add first version of livelink hub

diff --git a/Engine/Plugins/Animation/LiveLink/Source/LiveLinkHub/LiveLinkHub.Build.cs b/Engine/Plugins/Animation/LiveLink/Source/LiveLinkHub/LiveLinkHub.Build.cs
--- a/Engine/Plugins/Animation/LiveLink/Source/LiveLinkHub/LiveLinkHub.Build.cs
+++ b/Engine/Plugins/Animation/LiveLink/Source/LiveLinkHub/LiveLinkHub.Build.cs
@@ -16,6 +16,24 @@ namespace UnrealBuildTool.Rules
 			PrivateDependencyModuleNames.AddRange(
 			new string[]
 			{
+				"ApplicationCore",
+				"AssetTools",
+				"ContentBrowser",
+				"ContentBrowserAssetDataSource",
+				"ContentBrowserData",
+				"Engine",
+				"InputCore",
+				"LiveLink",
+				"LiveLinkEditor",
+				"LiveLinkInterface",
+				"LiveLinkMessageBusFramework",
+				"Slate",
+				"SlateCore",
+				"StandaloneRenderer",
+				"StructUtils",
+				"TimeManagement",
+				"ToolWidgets",
+				"UnrealEd",
 			});
 		}
 	}
diff --git a/Engine/Plugins/Animation/LiveLink/Source/LiveLinkHub/Private/LiveLinkHub.cpp b/Engine/Plugins/Animation/LiveLink/Source/LiveLinkHub/Private/LiveLinkHub.cpp
new file mode 100644
--- /dev/null
+++ b/Engine/Plugins/Animation/LiveLink/Source/LiveLinkHub/Private/LiveLinkHub.cpp
@@ -0,0 +1,133 @@
+// Copyright Epic Games, Inc. All Rights Reserved.
+
+#include "LiveLinkHub.h"
+
+#include "Features/IModularFeatures.h"
+#include "LiveLinkHubClient.h"
+#include "LiveLinkProvider.h"
+#include "LiveLinkSubject.h"
+#include "LiveLinkSubjectSettings.h"
+#include "Misc/ConfigCacheIni.h"
+#include "Recording/LiveLinkHubPlaybackController.h"
+#include "Recording/LiveLinkHubRecordingController.h"
+#include "Recording/LiveLinkHubRecordingListController.h"
+#include "UI/Window/LiveLinkHubWindowController.h"
+
+#define LOCTEXT_NAMESPACE "LiveLinkHub"
+
+void FLiveLinkHub::Initialize()
+{
+	LiveLinkProvider = ILiveLinkProvider::CreateLiveLinkProvider(TEXT("LiveLink Hub"));
+	LiveLinkHubClient = MakeShared<FLiveLinkHubClient>(AsShared());
+
+	IModularFeatures::Get().RegisterModularFeature(ILiveLinkClient::ModularFeatureName, LiveLinkHubClient.Get());
+
+	RecordingController = MakeShared<FLiveLinkHubRecordingController>();
+	PlaybackController = MakeShared<FLiveLinkHubPlaybackController>();
+	RecordingListController = MakeShared<FLiveLinkHubRecordingListController>(AsShared());
+
+	FString LiveLinkHubLayoutIni = GConfig->GetConfigFilename(TEXT("LiveLinkHubLayout"));
+	WindowController = MakeShared<FLiveLinkHubWindowController>(FLiveLinkHubWindowInitParams{ LiveLinkHubLayoutIni });
+	WindowController->RestoreLayout();
+
+	LiveLinkHubClient->OnStaticDataReceived_AnyThread().AddSP(this, &FLiveLinkHub::OnStaticDataReceived_AnyThread);
+	LiveLinkHubClient->OnFrameDataReceived_AnyThread().AddSP(this, &FLiveLinkHub::OnFrameDataReceived_AnyThread);
+	LiveLinkHubClient->OnLiveLinkSubjectAdded().AddSP(this, &FLiveLinkHub::OnSubjectAdded);
+	LiveLinkHubClient->OnLiveLinkSubjectRemoved().AddSP(this, &FLiveLinkHub::OnSubjectRemoved);
+
+	PlaybackController->Start();
+	RecordingController->Initialize(PlaybackController);
+}
+
+FLiveLinkHub::~FLiveLinkHub()
+{
+	LiveLinkHubClient->OnLiveLinkSubjectRemoved().RemoveAll(this);
+	LiveLinkHubClient->OnLiveLinkSubjectAdded().RemoveAll(this);
+	LiveLinkHubClient->OnFrameDataReceived_AnyThread().RemoveAll(this);
+	LiveLinkHubClient->OnStaticDataReceived_AnyThread().RemoveAll(this);
+}
+
+bool FLiveLinkHub::IsInPlayback() const
+{
+	return PlaybackController->IsInPlayback();
+}
+
+bool FLiveLinkHub::IsRecording() const
+{
+	return RecordingController->IsRecording();
+}
+
+void FLiveLinkHub::Tick()
+{
+	LiveLinkHubClient->Tick();
+}
+
+TSharedRef<SWindow> FLiveLinkHub::GetRootWindow() const
+{
+	return WindowController->GetRootWindow().ToSharedRef();
+}
+
+TSharedPtr<ILiveLinkProvider> FLiveLinkHub::GetLiveLinkProvider() const
+{
+	return LiveLinkProvider;
+}
+
+TSharedPtr<FLiveLinkHubRecordingController> FLiveLinkHub::GetRecordingController() const
+{
+	return RecordingController;
+}
+
+TSharedPtr<FLiveLinkHubRecordingListController> FLiveLinkHub::GetRecordingListController() const
+{
+	return RecordingListController;
+}
+
+TSharedPtr<FLiveLinkHubPlaybackController> FLiveLinkHub::GetPlaybackController() const
+{
+	return PlaybackController;
+}
+
+void FLiveLinkHub::OnStaticDataReceived_AnyThread(const FLiveLinkSubjectKey& InSubjectKey, TSubclassOf<ULiveLinkRole> InRole, const FLiveLinkStaticDataStruct& InStaticDataStruct)
+{
+	if (RecordingController->IsRecording())
+	{
+		RecordingController->RecordStaticData(InSubjectKey, InRole, InStaticDataStruct);
+	}
+
+	FLiveLinkStaticDataStruct StaticDataCopy;
+	StaticDataCopy.InitializeWith(InStaticDataStruct);
+	LiveLinkProvider->UpdateSubjectStaticData(InSubjectKey.SubjectName, InRole, MoveTemp(StaticDataCopy));
+}
+
+void FLiveLinkHub::OnFrameDataReceived_AnyThread(const FLiveLinkSubjectKey& InSubjectKey, const FLiveLinkFrameDataStruct& InFrameDataStruct)
+{
+	if (RecordingController->IsRecording())
+	{
+		RecordingController->RecordFrameData(InSubjectKey, InFrameDataStruct);
+	}
+
+	FLiveLinkFrameDataStruct FrameDataCopy;
+	FrameDataCopy.InitializeWith(InFrameDataStruct);
+	LiveLinkProvider->UpdateSubjectFrameData(InSubjectKey.SubjectName, MoveTemp(FrameDataCopy));
+}
+
+void FLiveLinkHub::OnSubjectAdded(FLiveLinkSubjectKey InSubjectKey)
+{
+	// Send an update to connected clients as well.
+	ULiveLinkSubjectSettings* SubjectSettings = Cast<ULiveLinkSubjectSettings>(LiveLinkHubClient->GetSubjectSettings(InSubjectKey));
+
+	if (const FLiveLinkStaticDataStruct* StaticData = LiveLinkHubClient->GetSubjectStaticData(InSubjectKey))
+	{
+		FLiveLinkStaticDataStruct StaticDataCopy;
+		StaticDataCopy.InitializeWith(*StaticData);
+		LiveLinkProvider->UpdateSubjectStaticData(InSubjectKey.SubjectName, SubjectSettings->Role, MoveTemp(StaticDataCopy));
+	}
+}
+
+void FLiveLinkHub::OnSubjectRemoved(FLiveLinkSubjectKey InSubjectKey)
+{
+	// Send an update to connected clients as well.
+	LiveLinkProvider->RemoveSubject(InSubjectKey.SubjectName);
+}
+
+#undef LOCTEXT_NAMESPACE /*LiveLinkHub*/
\ No newline at end of file
diff --git a/Engine/Plugins/Animation/LiveLink/Source/LiveLinkHub/Private/LiveLinkHub.h b/Engine/Plugins/Animation/LiveLink/Source/LiveLinkHub/Private/LiveLinkHub.h
new file mode 100644
--- /dev/null
+++ b/Engine/Plugins/Animation/LiveLink/Source/LiveLinkHub/Private/LiveLinkHub.h
@@ -0,0 +1,86 @@
+// Copyright Epic Games, Inc. All Rights Reserved.
+
+#pragma once
+
+#include "LiveLinkTypes.h"
+#include "Templates/SubclassOf.h"
+
+class FLiveLinkHubClient;
+class FLiveLinkHubPlaybackController;
+class FLiveLinkHubRecordingController;
+class FLiveLinkHubRecordingListController;
+class FLiveLinkHubWindowController;
+struct FLiveLinkSubjectKey;
+struct ILiveLinkProvider;
+class SWindow;
+class ULiveLinkRole;
+
+/**
+ * Main interface for the live link hub.
+ */
+class ILiveLinkHub
+{
+public:
+	virtual ~ILiveLinkHub() {}
+	// todo: replace with GetStatus?
+
+	/** Whether the hub is currently playing a recording. */
+	virtual bool IsInPlayback() const = 0;
+	/** Whether the hub is currently recording livelink data. */
+	virtual bool IsRecording() const = 0;
+};
+
+/**
+ * Implementation of the live link hub.
+ * Contains the apps' different components and is responsible for handling communication between them.
+ */
+class FLiveLinkHub : public ILiveLinkHub, public TSharedFromThis<FLiveLinkHub>
+{
+public:
+	virtual ~FLiveLinkHub();
+
+	//~ Begin ILiveLinkHub interface
+	virtual bool IsInPlayback() const override;
+	virtual bool IsRecording() const override;
+	//~ End ILiveLinkHub interface
+
+public:
+	/** Launch the slate application and initialize its components. */
+	void Initialize();
+	/** Tick the hub. */
+	void Tick();
+	/** Get the root window that hosts the hub's slate application. */
+	TSharedRef<SWindow> GetRootWindow() const;
+	/** Get the livelink provider used to rebroadcast livelink data to connected UE clients. */
+	TSharedPtr<ILiveLinkProvider> GetLiveLinkProvider() const;
+	/** Get the controller that manages recording livelink data. */
+	TSharedPtr<FLiveLinkHubRecordingController> GetRecordingController() const;
+	/** Get the recording list controller, that handles displaying livelink recording assets. */
+	TSharedPtr<FLiveLinkHubRecordingListController> GetRecordingListController() const;
+	/** Get the controller that manages playing back livelink data. */
+	TSharedPtr<FLiveLinkHubPlaybackController> GetPlaybackController() const;
+	
+private:
+	//~ LiveLink Client delegates
+	void OnStaticDataReceived_AnyThread(const FLiveLinkSubjectKey& InSubjectKey, TSubclassOf<ULiveLinkRole> InRole, const FLiveLinkStaticDataStruct& InStaticDataStruct);
+	void OnFrameDataReceived_AnyThread(const FLiveLinkSubjectKey& InSubjectKey, const FLiveLinkFrameDataStruct& InFrameDataStruct);
+	void OnSubjectAdded(FLiveLinkSubjectKey InSubjectKey);
+	void OnSubjectRemoved(FLiveLinkSubjectKey InSubjectKey);
+	//~ LiveLink Client delegates
+
+private:
+	/** Recording controller. */
+	TSharedPtr<FLiveLinkHubRecordingController> RecordingController;
+	/** Recordings list controller. */
+	TSharedPtr<FLiveLinkHubRecordingListController> RecordingListController;
+	/** Playback controller. */
+	TSharedPtr<FLiveLinkHubPlaybackController> PlaybackController;
+	/** Window controller */
+	TSharedPtr<FLiveLinkHubWindowController> WindowController;
+	/** LiveLinkHub's livelink client. */
+	TSharedPtr<FLiveLinkHubClient> LiveLinkHubClient;
+	/** LiveLinkProvider used to transfer data to connected UE clients. */
+	TSharedPtr<ILiveLinkProvider> LiveLinkProvider;
+
+	friend class FLiveLinkHubModule;
+};
diff --git a/Engine/Plugins/Animation/LiveLink/Source/LiveLinkHub/Private/LiveLinkHubApplication.h b/Engine/Plugins/Animation/LiveLink/Source/LiveLinkHub/Private/LiveLinkHubApplication.h
new file mode 100644
--- /dev/null
+++ b/Engine/Plugins/Animation/LiveLink/Source/LiveLinkHub/Private/LiveLinkHubApplication.h
@@ -0,0 +1,60 @@
+// Copyright Epic Games, Inc. All Rights Reserved.
+
+#pragma once
+
+#include "Commandlets/Commandlet.h"
+#include "CoreMinimal.h"
+#include "Engine/Engine.h"
+#include "Editor/EditorPerformanceSettings.h"
+#include "Misc/CoreDelegates.h"
+#include "LiveLinkHub.h"
+#include "Misc/App.h"
+#include "Runtime/Launch/Resources/Version.h"
+#include "UObject/UObjectGlobals.h"
+
+
+DEFINE_LOG_CATEGORY_STATIC(LogLiveLinkHubApplication, Log, All);
+
+void LiveLinkHubLoop(const TSharedPtr<FLiveLinkHub>& LiveLinkHub)
+{
+	// Disable throttling for the hub
+	GetMutableDefault<UEditorPerformanceSettings>()->bThrottleCPUWhenNotForeground = false;
+
+	{
+		UE_LOG(LogLiveLinkHubApplication, Display, TEXT("LiveLinkHub Initialized (Version: %d.%d)"), ENGINE_MAJOR_VERSION, ENGINE_MINOR_VERSION);
+
+		double LastTime = FPlatformTime::Seconds();
+
+		constexpr double IdealFrameRate = 60.0;
+		constexpr float IdealFrameTime = 1.0f / IdealFrameRate;
+
+		while (!IsEngineExitRequested())
+		{
+			const double CurrentTime = FPlatformTime::Seconds();
+			const double DeltaTime = CurrentTime - LastTime;
+
+			FApp::SetDeltaTime(DeltaTime);
+			GEngine->UpdateTimeAndHandleMaxTickRate();
+
+			CommandletHelpers::TickEngine(nullptr, DeltaTime);
+
+			// This is normally ticked by OnSamplingInput.
+			LiveLinkHub->Tick();
+
+			// Run garbage collection for the UObjects for the rest of the frame or at least to 2 ms
+			IncrementalPurgeGarbage(true, FMath::Max<float>(0.002f, IdealFrameTime - (FPlatformTime::Seconds() - LastTime)));
+
+#if PLATFORM_MAC
+			// Pumps the message from the main loop. (On Mac, we have a full application to get a proper console window to output logs)
+			FPlatformApplicationMisc::PumpMessages(true);
+#endif
+
+			FCoreDelegates::OnEndFrame.Broadcast();
+
+			// Throttle main thread main fps by sleeping if we still have time
+			FPlatformProcess::Sleep(FMath::Max<float>(0.0f, IdealFrameTime - (FPlatformTime::Seconds() - LastTime)));
+
+			LastTime = CurrentTime;
+		}
+	}
+}
diff --git a/Engine/Plugins/Animation/LiveLink/Source/LiveLinkHub/Private/LiveLinkHubClient.cpp b/Engine/Plugins/Animation/LiveLink/Source/LiveLinkHub/Private/LiveLinkHubClient.cpp
new file mode 100644
--- /dev/null
+++ b/Engine/Plugins/Animation/LiveLink/Source/LiveLinkHub/Private/LiveLinkHubClient.cpp
@@ -0,0 +1,417 @@
+// Copyright Epic Games, Inc. All Rights Reserved.
+
+#include "LiveLinkHubClient.h"
+
+#include "Async/TaskGraphInterfaces.h"
+#include "Engine/Engine.h"
+#include "LiveLinkClient.h"
+#include "LiveLinkHub.h"
+#include "LiveLinkHubLog.h"
+#include "LiveLinkRoleTrait.h"
+#include "LiveLinkSourceCollection.h"
+#include "LiveLinkSettings.h"
+#include "LiveLinkSubject.h"
+#include "LiveLinkSubjectSettings.h"
+#include "LiveLinkTimedDataInput.h"
+#include "Modules/ModuleManager.h"
+#include "Recording/LiveLinkPlaybackSource.h"
+#include "Recording/LiveLinkPlaybackSubject.h"
+#include "Stats/Stats2.h"
+#include "UObject/Package.h"
+#include "UObject/UObjectGlobals.h"
+
+DECLARE_STATS_GROUP(TEXT("Live Link Hub"), STATGROUP_LiveLinkHub, STATCAT_Advanced);
+DECLARE_CYCLE_STAT(TEXT("LiveLinkHub - Push StaticData"), STAT_LiveLinkHub_PushStaticData, STATGROUP_LiveLinkHub);
+DECLARE_CYCLE_STAT(TEXT("LiveLinkHub - Push FrameData"), STAT_LiveLinkHub_PushFrameData, STATGROUP_LiveLinkHub);
+
+#define LOCTEXT_NAMESPACE "LiveLinkHub.LiveLinkHubClient"
+
+FLiveLinkHubClient::FLiveLinkHubClient(TSharedPtr<ILiveLinkHub> InLiveLinkHub)
+	: FLiveLinkClient()
+	, LiveLinkHub(MoveTemp(InLiveLinkHub))
+{
+}
+
+FLiveLinkHubClient::~FLiveLinkHubClient()
+{
+}
+
+const FLiveLinkStaticDataStruct* FLiveLinkHubClient::GetSubjectStaticData(const FLiveLinkSubjectKey& InSubjectKey)
+{
+	FScopeLock Lock(&CollectionAccessCriticalSection);
+	if (const FLiveLinkCollectionSubjectItem* SubjectItem = Collection->FindSubject(InSubjectKey))
+	{
+		if (FLiveLinkSubject* LiveLinkSubject = SubjectItem->GetLiveSubject())
+		{
+			return &LiveLinkSubject->GetStaticData();
+		}
+	}
+
+	return nullptr;
+}
+
+bool FLiveLinkHubClient::CreateSource(const FLiveLinkSourcePreset& InSourcePreset)
+{
+	// Create fake sources if we're in playback mode so we don't process livelink data while we're playing a recording.
+	if (LiveLinkHub.Pin()->IsInPlayback())
+	{
+		return CreatePlaybackSource(InSourcePreset);
+	}
+	else
+	{
+		return FLiveLinkClient::CreateSource(InSourcePreset);
+	}
+}
+
+bool FLiveLinkHubClient::CreatePlaybackSource(const FLiveLinkSourcePreset& InSourcePreset)
+{
+	check(Collection);
+
+	if (InSourcePreset.Settings == nullptr)
+	{
+		UE_LOG(LogLiveLinkHub, Warning, TEXT("Create Source Failure: The settings are not defined."));
+		return false;
+	}
+
+	if (InSourcePreset.Guid == FLiveLinkSourceCollection::DefaultVirtualSubjectGuid)
+	{
+		UE_LOG(LogLiveLinkHub, Warning, TEXT("Create Source Failure: Can't create default subject source. It will be created automatically."));
+		return false;
+	}
+
+	if (!InSourcePreset.Guid.IsValid())
+	{
+		UE_LOG(LogLiveLinkHub, Warning, TEXT("Create Source Failure: The guid is invalid."));
+		return false;
+	}
+
+	{
+		FScopeLock Lock(&CollectionAccessCriticalSection);
+		if (Collection->FindSource(InSourcePreset.Guid) != nullptr)
+		{
+			UE_LOG(LogLiveLinkHub, Warning, TEXT("Create Source Failure: The guid already exist."));
+			return false;
+		}
+	}
+
+	ULiveLinkSourceSettings* Setting = nullptr;
+	TSharedPtr<ILiveLinkSource> Source;
+	FLiveLinkCollectionSourceItem Data;
+	Data.Guid = InSourcePreset.Guid;
+
+	// subject source have a special settings class. We can differentiate them using this
+	if (InSourcePreset.Settings->GetClass()->GetName() == TEXT("LiveLinkVirtualSubjectSourceSettings"))
+	{
+		Source = MakeShared<FLiveLinkPlaybackSource>();
+		Data.bIsVirtualSource = true;
+	}
+	else
+	{
+		Source = MakeShared<FLiveLinkPlaybackSource>();
+		Data.TimedData = MakeShared<FLiveLinkTimedDataInput>(this, InSourcePreset.Guid);
+	}
+
+	Data.Source = Source;
+
+	// In case a source has changed its source settings class, instead of duplicating, create the right one and copy previous properties
+	UClass* SourceSettingsClass = Source->GetSettingsClass().Get();
+	if (SourceSettingsClass && SourceSettingsClass != InSourcePreset.Settings->GetClass())
+	{
+		Setting = NewObject<ULiveLinkSourceSettings>(GetTransientPackage(), SourceSettingsClass);
+		UEngine::CopyPropertiesForUnrelatedObjects(InSourcePreset.Settings, Setting);
+		Data.Setting = Setting;
+	}
+	else
+	{
+		Setting = Data.Setting = DuplicateObject<ULiveLinkSourceSettings>(InSourcePreset.Settings, GetTransientPackage());
+	}
+
+	{
+		FScopeLock Lock(&CollectionAccessCriticalSection);
+		Collection->AddSource(MoveTemp(Data));
+	}
+	Source->ReceiveClient(this, InSourcePreset.Guid);
+	Source->InitializeSettings(Setting);
+
+	return true;
+}
+
+bool FLiveLinkHubClient::CreateSubject(const FLiveLinkSubjectPreset& InSubjectPreset)
+{
+	// Create fake subjects if we're in playback mode
+	if (LiveLinkHub.Pin()->IsInPlayback())
+	{
+		return CreatePlaybackSubject(InSubjectPreset);
+	}
+	else
+	{
+		return FLiveLinkClient::CreateSubject(InSubjectPreset);
+	}
+}
+
+bool FLiveLinkHubClient::CreatePlaybackSubject(const FLiveLinkSubjectPreset& InSubjectPreset)
+{
+	check(Collection);
+
+	if (InSubjectPreset.Role.Get() == nullptr || InSubjectPreset.Role.Get() == ULiveLinkRole::StaticClass())
+	{
+		UE_LOG(LogLiveLinkHub, Warning, TEXT("Create Subject Failure: The role is not defined."));
+		return false;
+	}
+
+	if (InSubjectPreset.Key.Source == FLiveLinkSourceCollection::DefaultVirtualSubjectGuid && InSubjectPreset.VirtualSubject == nullptr)
+	{
+		UE_LOG(LogLiveLinkHub, Warning, TEXT("Create Source Failure: Can't create an empty subject."));
+		return false;
+	}
+
+	if (InSubjectPreset.Key.SubjectName.IsNone())
+	{
+		UE_LOG(LogLiveLinkHub, Warning, TEXT("Create Subject Failure: The subject name is invalid."));
+		return false;
+	}
+
+	FLiveLinkCollectionSourceItem* SourceItem = Collection->FindSource(InSubjectPreset.Key.Source);
+	if (SourceItem == nullptr || SourceItem->bPendingKill)
+	{
+		UE_LOG(LogLiveLinkHub, Warning, TEXT("Create Subject Failure: The source doesn't exist."));
+		return false;
+	}
+
+	FLiveLinkCollectionSubjectItem* SubjectItem = Collection->FindSubject(InSubjectPreset.Key);
+	if (SubjectItem != nullptr)
+	{
+		if (SubjectItem->bPendingKill)
+		{
+			FScopeLock Lock(&CollectionAccessCriticalSection);
+			Collection->RemoveSubject(InSubjectPreset.Key);
+		}
+		else
+		{
+			UE_LOG(LogLiveLinkHub, Warning, TEXT("Create Subject Failure: The subject already exist."));
+			return false;
+		}
+	}
+
+	if (InSubjectPreset.VirtualSubject)
+	{
+		bool bEnabled = false;
+		ULiveLinkVirtualSubject* VSubject = DuplicateObject<ULiveLinkVirtualSubject>(InSubjectPreset.VirtualSubject, GetTransientPackage());
+		FLiveLinkCollectionSubjectItem VSubjectData(InSubjectPreset.Key, VSubject, bEnabled);
+		VSubject->Initialize(VSubjectData.Key, VSubject->GetRole(), this);
+
+		FScopeLock Lock(&CollectionAccessCriticalSection);
+		Collection->AddSubject(MoveTemp(VSubjectData));
+		Collection->SetSubjectEnabled(InSubjectPreset.Key, InSubjectPreset.bEnabled);
+	}
+	else
+	{
+		ULiveLinkSubjectSettings* SubjectSettings = nullptr;
+		if (InSubjectPreset.Settings)
+		{
+			SubjectSettings = DuplicateObject<ULiveLinkSubjectSettings>(InSubjectPreset.Settings, GetTransientPackage());
+		}
+		else
+		{
+			SubjectSettings = NewObject<ULiveLinkSubjectSettings>();
+		}
+
+		bool bEnabled = false;
+		FLiveLinkCollectionSubjectItem CollectionSubjectItem(InSubjectPreset.Key, MakeUnique<FLiveLinkPlaybackSubject>(SourceItem->TimedData), SubjectSettings, bEnabled);
+		CollectionSubjectItem.GetLiveSubject()->Initialize(InSubjectPreset.Key, InSubjectPreset.Role.Get(), this);
+
+		FScopeLock Lock(&CollectionAccessCriticalSection);
+		Collection->AddSubject(MoveTemp(CollectionSubjectItem));
+		Collection->SetSubjectEnabled(InSubjectPreset.Key, InSubjectPreset.bEnabled);
+	}
+	return true;
+}
+
+void FLiveLinkHubClient::PushSubjectStaticData_AnyThread(const FLiveLinkSubjectKey& SubjectKey, TSubclassOf<ULiveLinkRole> Role, FLiveLinkStaticDataStruct&& InStaticData)
+{
+	// Copied and adapted from FLiveLinkClient::PushSubjectStaticData_AnyThread
+	SCOPE_CYCLE_COUNTER(STAT_LiveLinkHub_PushStaticData);
+
+	check(Collection);
+
+	if (!FLiveLinkRoleTrait::Validate(Role, InStaticData))
+	{
+		if (Role == nullptr)
+		{
+			UE_LOG(LogLiveLinkHub, Error, TEXT("Trying to add unsupported static data type with subject '%s'."), *SubjectKey.SubjectName.ToString());
+		}
+		else
+		{
+			UE_LOG(LogLiveLinkHub, Error, TEXT("Trying to add unsupported static data type to role '%s' with subject '%s'."), *Role->GetName(), *SubjectKey.SubjectName.ToString());
+		}
+		return;
+	}
+
+	bool bShouldLogIfInvalidStaticData = true;
+	if (!Role.GetDefaultObject()->IsStaticDataValid(InStaticData, bShouldLogIfInvalidStaticData))
+	{
+		if (bShouldLogIfInvalidStaticData)
+		{
+			UE_LOG(LogLiveLinkHub, Error, TEXT("Trying to add static data that is not formatted properly to role '%s' with subject '%s'."), *Role->GetName(), *SubjectKey.SubjectName.ToString());
+		}
+		return;
+	}
+
+	const FLiveLinkCollectionSourceItem* SourceItem = nullptr;
+	{
+		FScopeLock Lock(&CollectionAccessCriticalSection);
+		SourceItem = Collection->FindSource(SubjectKey.Source);
+	}
+
+	if (SourceItem == nullptr || SourceItem->bPendingKill)
+	{
+		return;
+	}
+
+	if (SourceItem->IsVirtualSource())
+	{
+		UE_LOG(LogLiveLinkHub, Error, TEXT("Trying to add frame data to the subject '%s'."), *Role->GetName(), *SubjectKey.SubjectName.ToString());
+		return;
+	}
+
+	FLiveLinkSubject* LiveLinkSubject = nullptr;
+	{
+		FScopeLock Lock(&CollectionAccessCriticalSection);
+		if (FLiveLinkCollectionSubjectItem* SubjectItem = Collection->FindSubject(SubjectKey))
+		{
+			LiveLinkSubject = SubjectItem->GetLiveSubject();
+
+			if (LiveLinkSubject->GetRole() != Role)
+			{
+				UE_LOG(LogLiveLinkHub, Warning, TEXT("Subject '%s' of role '%s' is changing its role to '%s'. Current subject will be removed and a new one will be created"), *SubjectKey.SubjectName.ToString(), *LiveLinkSubject->GetRole().GetDefaultObject()->GetDisplayName().ToString(), *Role.GetDefaultObject()->GetDisplayName().ToString());
+
+				Collection->RemoveSubject(SubjectKey);
+				LiveLinkSubject = nullptr;
+			}
+			else
+			{
+				LiveLinkSubject->ClearFrames();
+			}
+		}
+
+	}
+	if (LiveLinkSubject == nullptr)
+	{
+		const ULiveLinkSettings* LiveLinkSettings = GetDefault<ULiveLinkSettings>();
+		const FLiveLinkRoleProjectSetting DefaultSetting = LiveLinkSettings->GetDefaultSettingForRole(Role.Get());
+
+		// Setting class should always be valid
+		ULiveLinkSubjectSettings* SubjectSettings = nullptr;
+		{
+			UClass* SettingClass = DefaultSetting.SettingClass.Get();
+			if (SettingClass == nullptr)
+			{
+				SettingClass = ULiveLinkSubjectSettings::StaticClass();
+			}
+
+			SubjectSettings = NewObject<ULiveLinkSubjectSettings>(GetTransientPackage(), SettingClass);
+			SubjectSettings->Role = Role;
+
+			UClass* FrameInterpolationProcessorClass = DefaultSetting.FrameInterpolationProcessor.Get();
+			if (FrameInterpolationProcessorClass != nullptr)
+			{
+				UClass* InterpolationRole = FrameInterpolationProcessorClass->GetDefaultObject<ULiveLinkFrameInterpolationProcessor>()->GetRole();
+				if (Role->IsChildOf(InterpolationRole))
+				{
+					SubjectSettings->InterpolationProcessor = NewObject<ULiveLinkFrameInterpolationProcessor>(SubjectSettings, FrameInterpolationProcessorClass);
+				}
+				else
+				{
+					UE_LOG(LogLiveLinkHub, Warning, TEXT("The interpolator '%s' is not valid for the Role '%s'"), *FrameInterpolationProcessorClass->GetName(), *Role->GetName());
+				}
+			}
+			else
+			{
+				// If no settings were found for a specific role, check if the default interpolator is compatible with the role
+				UClass* FallbackInterpolationProcessorClass = LiveLinkSettings->FrameInterpolationProcessor.Get();
+				if (FallbackInterpolationProcessorClass != nullptr)
+				{
+					UClass* InterpolationRole = FallbackInterpolationProcessorClass->GetDefaultObject<ULiveLinkFrameInterpolationProcessor>()->GetRole();
+					if (Role->IsChildOf(InterpolationRole))
+					{
+						SubjectSettings->InterpolationProcessor = NewObject<ULiveLinkFrameInterpolationProcessor>(SubjectSettings, FallbackInterpolationProcessorClass);
+					}
+				}
+			}
+
+			for (TSubclassOf<ULiveLinkFramePreProcessor> PreProcessor : DefaultSetting.FramePreProcessors)
+			{
+				if (PreProcessor != nullptr)
+				{
+					UClass* PreProcessorRole = PreProcessor->GetDefaultObject<ULiveLinkFramePreProcessor>()->GetRole();
+					if (Role->IsChildOf(PreProcessorRole))
+					{
+						SubjectSettings->PreProcessors.Add(NewObject<ULiveLinkFramePreProcessor>(SubjectSettings, PreProcessor.Get()));
+					}
+					else
+					{
+						UE_LOG(LogLiveLinkHub, Warning, TEXT("The pre processor '%s' is not valid for the Role '%s'"), *PreProcessor->GetName(), *Role->GetName());
+					}
+				}
+			}
+		}
+
+
+		{
+			FScopeLock Lock(&CollectionAccessCriticalSection);
+
+			bool bEnabled = Collection->FindEnabledSubject(SubjectKey.SubjectName) == nullptr;
+			FLiveLinkCollectionSubjectItem CollectionSubjectItem(SubjectKey, MakeUnique<FLiveLinkSubject>(SourceItem->TimedData), SubjectSettings, bEnabled);
+			CollectionSubjectItem.GetLiveSubject()->Initialize(SubjectKey, Role.Get(), this);
+
+			LiveLinkSubject = CollectionSubjectItem.GetLiveSubject();
+
+			Collection->AddSubject(MoveTemp(CollectionSubjectItem));
+			
+			// Clear async flag on the newly created uobject to allow it to be garbage collected.
+			SubjectSettings->AtomicallyClearInternalFlags(EInternalObjectFlags::Async);
+		}
+	}
+
+	OnStaticDataReceivedDelegate_AnyThread.Broadcast(SubjectKey, Role, InStaticData);
+
+	if (LiveLinkSubject)
+	{
+		// Dispatch it on the game thread since FLiveLinkSubject::SetStaticData asserts when called outside the game thread.
+		FFunctionGraphTask::CreateAndDispatchWhenReady([this, LiveLinkSubject, SubjectKey, Role, StaticData = MoveTemp(InStaticData)]() mutable
+			{
+				LiveLinkSubject->SetStaticData(Role, MoveTemp(StaticData));
+			}, TStatId(), nullptr, ENamedThreads::GameThread);
+	}
+
+}
+
+void FLiveLinkHubClient::PushSubjectFrameData_AnyThread(const FLiveLinkSubjectKey& SubjectKey, FLiveLinkFrameDataStruct&& FrameData)
+{
+	SCOPE_CYCLE_COUNTER(STAT_LiveLinkHub_PushFrameData);
+
+	OnFrameDataReceivedDelegate_AnyThread.Broadcast(SubjectKey, FrameData);
+}
+
+FText FLiveLinkHubClient::GetSourceStatus(FGuid InEntryGuid) const
+{
+	if (LiveLinkHub.Pin()->IsInPlayback())
+	{
+		return LOCTEXT("PlaybackText", "Playback");
+	}
+
+	return FLiveLinkClient::GetSourceStatus(InEntryGuid);
+}
+
+bool FLiveLinkHubClient::IsSubjectValid(const FLiveLinkSubjectKey& InSubjectKey) const
+{
+	FScopeLock Lock(&CollectionAccessCriticalSection);
+	if (const FLiveLinkCollectionSubjectItem* SubjectItem = Collection->FindSubject(InSubjectKey))
+	{
+		// todo LiveLinkHub : Get subject validity from the UE client since we don't store snapshots on the hub side yet
+		return true;
+	}
+	return false;
+}
+
+#undef LOCTEXT_NAMESPACE
\ No newline at end of file
diff --git a/Engine/Plugins/Animation/LiveLink/Source/LiveLinkHub/Private/LiveLinkHubClient.h b/Engine/Plugins/Animation/LiveLink/Source/LiveLinkHub/Private/LiveLinkHubClient.h
new file mode 100644
--- /dev/null
+++ b/Engine/Plugins/Animation/LiveLink/Source/LiveLinkHub/Private/LiveLinkHubClient.h
@@ -0,0 +1,65 @@
+// Copyright Epic Games, Inc. All Rights Reserved.
+
+#pragma once
+
+#include "LiveLinkClient.h"
+
+#include "Delegates/DelegateCombinations.h"
+#include "LiveLinkTypes.h"
+#include "Templates/SubclassOf.h"
+
+class FLiveLinkHubPlaybackController;
+class FLiveLinkHubRecordingController;
+class ILiveLinkHub;
+struct ILiveLinkProvider;
+
+DECLARE_TS_MULTICAST_DELEGATE_TwoParams(FOnFrameDataReceived_AnyThread, const FLiveLinkSubjectKey&, const FLiveLinkFrameDataStruct&);
+DECLARE_TS_MULTICAST_DELEGATE_ThreeParams(FOnStaticDataReceived_AnyThread, const FLiveLinkSubjectKey&, TSubclassOf<ULiveLinkRole>, const FLiveLinkStaticDataStruct&);
+
+class FLiveLinkHubClient : public FLiveLinkClient
+{
+public:
+	FLiveLinkHubClient(TSharedPtr<ILiveLinkHub> InLiveLinkHub);
+	virtual ~FLiveLinkHubClient();
+
+	/** Utility method to grab a subject's static data. Used by the RecordingController when static data is missing from the recording. */
+	const FLiveLinkStaticDataStruct* GetSubjectStaticData(const FLiveLinkSubjectKey& InSubjectKey);
+	
+	/** Get the delegate called when frame data is received. */
+	FOnFrameDataReceived_AnyThread& OnFrameDataReceived_AnyThread()
+	{
+		return OnFrameDataReceivedDelegate_AnyThread;
+	}
+
+	/** Get the delegate called when static data is received. */
+	FOnStaticDataReceived_AnyThread& OnStaticDataReceived_AnyThread()
+	{
+		return OnStaticDataReceivedDelegate_AnyThread;
+	}
+
+public:
+	//~ Begin ILiveLinkClient interface
+	virtual bool CreateSource(const FLiveLinkSourcePreset& InSourcePreset) override;
+	virtual bool CreateSubject(const FLiveLinkSubjectPreset& InSubjectPreset) override;
+	virtual void PushSubjectStaticData_AnyThread(const FLiveLinkSubjectKey& SubjectKey, TSubclassOf<ULiveLinkRole> Role, FLiveLinkStaticDataStruct&& InStaticData) override;
+	virtual void PushSubjectFrameData_AnyThread(const FLiveLinkSubjectKey& SubjectKey, FLiveLinkFrameDataStruct&& FrameData) override;
+	virtual FText GetSourceStatus(FGuid InEntryGuid) const override;
+	virtual bool IsSubjectValid(const FLiveLinkSubjectKey& InSubjectKey) const override;
+	//~ End ILiveLinkClient interface
+
+private:
+	/** Create a LiveLinkPlaybackSource which acts as a dummy source when doing playback. */
+	bool CreatePlaybackSource(const FLiveLinkSourcePreset& InSourcePreset);
+	/** Create a LiveLinkPlaybackSubject which acts as a dummy subject when doing playback. */
+	bool CreatePlaybackSubject(const FLiveLinkSubjectPreset& InSubjectPreset);
+	/** Lock to stop multiple threads accessing the Subjects from the collection at the same time */
+	mutable FCriticalSection CollectionAccessCriticalSection;
+
+private:
+	/** Weak pointer to the live link hub. */
+	TWeakPtr<ILiveLinkHub> LiveLinkHub;
+    /** Delegate called when frame data is received. */
+    FOnFrameDataReceived_AnyThread OnFrameDataReceivedDelegate_AnyThread;
+    /** Delegate called when static data is received. */
+    FOnStaticDataReceived_AnyThread OnStaticDataReceivedDelegate_AnyThread;
+};
diff --git a/Engine/Plugins/Animation/LiveLink/Source/LiveLinkHub/Private/LiveLinkHubLog.cpp b/Engine/Plugins/Animation/LiveLink/Source/LiveLinkHub/Private/LiveLinkHubLog.cpp
new file mode 100644
--- /dev/null
+++ b/Engine/Plugins/Animation/LiveLink/Source/LiveLinkHub/Private/LiveLinkHubLog.cpp
@@ -0,0 +1,5 @@
+// Copyright Epic Games, Inc. All Rights Reserved.
+
+#include "LiveLinkHubLog.h"
+
+DEFINE_LOG_CATEGORY(LogLiveLinkHub);
diff --git a/Engine/Plugins/Animation/LiveLink/Source/LiveLinkHub/Private/LiveLinkHubLog.h b/Engine/Plugins/Animation/LiveLink/Source/LiveLinkHub/Private/LiveLinkHubLog.h
new file mode 100644
--- /dev/null
+++ b/Engine/Plugins/Animation/LiveLink/Source/LiveLinkHub/Private/LiveLinkHubLog.h
@@ -0,0 +1,7 @@
+// Copyright Epic Games, Inc. All Rights Reserved.
+
+#pragma once
+
+#include "Logging/LogMacros.h"
+
+DECLARE_LOG_CATEGORY_EXTERN(LogLiveLinkHub, Log, All);
diff --git a/Engine/Plugins/Animation/LiveLink/Source/LiveLinkHub/Private/LiveLinkHubModule.cpp b/Engine/Plugins/Animation/LiveLink/Source/LiveLinkHub/Private/LiveLinkHubModule.cpp
--- a/Engine/Plugins/Animation/LiveLink/Source/LiveLinkHub/Private/LiveLinkHubModule.cpp
+++ b/Engine/Plugins/Animation/LiveLink/Source/LiveLinkHub/Private/LiveLinkHubModule.cpp
@@ -1,10 +1,47 @@
 // Copyright Epic Games, Inc. All Rights Reserved.
 
 #include "LiveLinkHubModule.h"
+
+#include "LiveLinkHubApplication.h"
 #include "Modules/ModuleManager.h"
+#include "Recording/LiveLinkHubPlaybackController.h"
+#include "Recording/LiveLinkHubRecordingController.h"
 
 void FLiveLinkHubModule::StartLiveLinkHub()
 {
+	// Nothing will get executed after this, so put everything before.
+	LiveLinkHub = MakeShared<FLiveLinkHub>();
+	LiveLinkHub->Initialize();
+
+	LiveLinkHubLoop(LiveLinkHub);
+
+	// If we got to this point, the app is shutdown so we should release our references.
+	LiveLinkHub.Reset();
+}
+
+TSharedPtr<FLiveLinkHub> FLiveLinkHubModule::GetLiveLinkHub() const
+{
+	return LiveLinkHub;
+}
+
+TSharedPtr<ILiveLinkProvider> FLiveLinkHubModule::GetLiveLinkProvider() const
+{
+	return LiveLinkHub->LiveLinkProvider;
+}
+
+TSharedPtr<FLiveLinkHubRecordingController> FLiveLinkHubModule::GetRecordingController() const
+{
+	return LiveLinkHub->RecordingController;
+}
+
+TSharedPtr<FLiveLinkHubRecordingListController> FLiveLinkHubModule::GetRecordingListController() const
+{
+	return LiveLinkHub->RecordingListController;
+}
+
+TSharedPtr<FLiveLinkHubPlaybackController> FLiveLinkHubModule::GetPlaybackController() const
+{
+	return LiveLinkHub->PlaybackController;
 }
 
 IMPLEMENT_MODULE(FLiveLinkHubModule, LiveLinkHub);
diff --git a/Engine/Plugins/Animation/LiveLink/Source/LiveLinkHub/Private/LiveLinkHubModule.h b/Engine/Plugins/Animation/LiveLink/Source/LiveLinkHub/Private/LiveLinkHubModule.h
--- a/Engine/Plugins/Animation/LiveLink/Source/LiveLinkHub/Private/LiveLinkHubModule.h
+++ b/Engine/Plugins/Animation/LiveLink/Source/LiveLinkHub/Private/LiveLinkHubModule.h
@@ -3,6 +3,13 @@
 #pragma once
 
 #include "ILiveLinkHubModule.h"
+#include "Templates/SharedPointer.h"
+
+class FLiveLinkHub;
+struct ILiveLinkProvider;
+class FLiveLinkHubRecordingController;
+class FLiveLinkHubPlaybackController;
+class FLiveLinkHubRecordingListController;
 
 class FLiveLinkHubModule : public ILiveLinkHubModule
 {
@@ -10,4 +17,18 @@ public:
 	//~ Begin ILiveLinkHubModule interface
 	virtual void StartLiveLinkHub() override;
 	//~ End ILiveLinkHubModule interface
+
+	/** Get the livelink hub object. */
+	TSharedPtr<FLiveLinkHub> GetLiveLinkHub() const;
+	/** Get the livelink provider responsible for forwarding livelink data to connected UE clients. */
+	TSharedPtr<ILiveLinkProvider> GetLiveLinkProvider() const;
+	/** Get the recording controller. */
+	TSharedPtr<FLiveLinkHubRecordingController> GetRecordingController() const;
+	/** Get the recording list controller. */
+	TSharedPtr<FLiveLinkHubRecordingListController> GetRecordingListController() const;
+	/** Get the playback controller. */
+	TSharedPtr<FLiveLinkHubPlaybackController> GetPlaybackController() const;
+private:
+	/** LiveLinkHub object responsible for initializing the different controllers. */
+	TSharedPtr<FLiveLinkHub> LiveLinkHub;
 };
diff --git a/Engine/Plugins/Animation/LiveLink/Source/LiveLinkHub/Private/Recording/LiveLinkHubPlaybackController.h b/Engine/Plugins/Animation/LiveLink/Source/LiveLinkHub/Private/Recording/LiveLinkHubPlaybackController.h
new file mode 100644
--- /dev/null
+++ b/Engine/Plugins/Animation/LiveLink/Source/LiveLinkHub/Private/Recording/LiveLinkHubPlaybackController.h
@@ -0,0 +1,29 @@
+// Copyright Epic Games, Inc. All Rights Reserved.
+
+#pragma once
+
+/** Stub class for the playback controller. Will be replaced with the real implementation in a future CL. */
+class FLiveLinkHubPlaybackController
+{
+public:
+	FLiveLinkHubPlaybackController() = default;
+
+	bool IsInPlayback() const
+	{
+		return false;
+	}
+
+	bool IsLooping() const
+	{
+		return false;
+	}
+
+	void SetLooping(bool bInLoop)
+	{
+	}
+	
+public:
+	void Start()
+	{
+	}
+};
diff --git a/Engine/Plugins/Animation/LiveLink/Source/LiveLinkHub/Private/Recording/LiveLinkHubRecordingController.h b/Engine/Plugins/Animation/LiveLink/Source/LiveLinkHub/Private/Recording/LiveLinkHubRecordingController.h
new file mode 100644
--- /dev/null
+++ b/Engine/Plugins/Animation/LiveLink/Source/LiveLinkHub/Private/Recording/LiveLinkHubRecordingController.h
@@ -0,0 +1,46 @@
+// Copyright Epic Games, Inc. All Rights Reserved.
+
+#pragma once
+
+#include "LiveLinkTypes.h"
+#include "Widgets/SNullWidget.h"
+
+/** Stub class for the recording controller. Will be replaced with the real implementation in a future CL. */
+class FLiveLinkHubRecordingController
+{
+public:
+
+	void Initialize(const TSharedPtr<class FLiveLinkHubPlaybackController>& InPlaybackController)
+	{
+	}
+
+	~FLiveLinkHubRecordingController()
+	{
+	}
+
+	TSharedRef<SWidget> MakeRecordToolbarEntry()
+	{
+		return SNullWidget::NullWidget;
+	}
+
+	void StartRecording()
+	{
+	}
+	
+	void StopRecording()
+	{
+	}
+
+	bool IsRecording() const
+	{
+		return false;
+	}
+	
+	void RecordStaticData(const FLiveLinkSubjectKey& SubjectKey, TSubclassOf<ULiveLinkRole> Role, const FLiveLinkStaticDataStruct& StaticData)
+	{
+	}
+	
+	void RecordFrameData(const FLiveLinkSubjectKey& SubjectKey, const FLiveLinkFrameDataStruct& FrameData)
+	{
+	}
+};
diff --git a/Engine/Plugins/Animation/LiveLink/Source/LiveLinkHub/Private/Recording/LiveLinkHubRecordingListController.h b/Engine/Plugins/Animation/LiveLink/Source/LiveLinkHub/Private/Recording/LiveLinkHubRecordingListController.h
new file mode 100644
--- /dev/null
+++ b/Engine/Plugins/Animation/LiveLink/Source/LiveLinkHub/Private/Recording/LiveLinkHubRecordingListController.h
@@ -0,0 +1,20 @@
+// Copyright Epic Games, Inc. All Rights Reserved.
+
+#pragma once
+
+#include "Widgets/SNullWidget.h"
+
+/** Stub class for the recording list controller. Will be replaced with the real implementation in a future CL. */
+class FLiveLinkHubRecordingListController
+{
+public:
+	FLiveLinkHubRecordingListController(const TSharedPtr<FLiveLinkHub>& InLiveLinkHub)
+	{
+	}
+
+	/** Create the list's widget. */
+	TSharedRef<SWidget> MakeRecordingList()
+	{
+		return SNullWidget::NullWidget;
+	}
+};
diff --git a/Engine/Plugins/Animation/LiveLink/Source/LiveLinkHub/Private/Recording/LiveLinkPlaybackSource.h b/Engine/Plugins/Animation/LiveLink/Source/LiveLinkHub/Private/Recording/LiveLinkPlaybackSource.h
new file mode 100644
--- /dev/null
+++ b/Engine/Plugins/Animation/LiveLink/Source/LiveLinkHub/Private/Recording/LiveLinkPlaybackSource.h
@@ -0,0 +1,78 @@
+// Copyright Epic Games, Inc. All Rights Reserved.
+
+#pragma once
+
+#include "ILiveLinkSource.h"
+#include "LiveLinkSourceSettings.h"
+
+#include "LiveLinkPlaybackSource.generated.h"
+
+class ILiveLinkClient;
+
+/**
+ * Completely empty "source" displayed in the UI when playing back a recording.
+ */
+class FLiveLinkPlaybackSource : public ILiveLinkSource
+{
+public:
+	FLiveLinkPlaybackSource() = default;
+	virtual ~FLiveLinkPlaybackSource() = default;
+
+	//~ Begin ILiveLinkSource interface
+	virtual bool CanBeDisplayedInUI() const override
+	{
+		return true;
+	}
+	
+	virtual void ReceiveClient(ILiveLinkClient* InClient, FGuid InSourceGuid) override
+	{
+	}
+	
+	virtual void InitializeSettings(ULiveLinkSourceSettings* Settings) override
+	{
+		SourceName = TEXT("Playback source.");
+	}
+	
+	virtual bool IsSourceStillValid() const override
+	{
+		return true; 
+	}
+	
+	virtual bool RequestSourceShutdown() override
+	{
+		return true; 
+	}
+
+	virtual FText GetSourceType() const override
+	{
+		return FText::FromName(SourceName);
+	}
+	
+	virtual FText GetSourceMachineName() const override
+	{
+		return FText();
+	}
+	
+	virtual FText GetSourceStatus() const override
+	{
+		return NSLOCTEXT("LiveLinkPlaybackSource", "PlaybackSourceStatus", "Playback");
+	}
+	//~ End ILiveLinkSource interface
+
+protected:
+	/** Source name. */
+	FName SourceName;
+};
+
+/** PlaybackSourceSettings to be able to differentiate from live sources and keep a name associated to the source */
+UCLASS()
+class ULiveLinkPlaybackSourceSettings : public ULiveLinkSourceSettings
+{
+	GENERATED_BODY()
+
+public:
+
+	/** Source name. */
+	UPROPERTY()
+	FName SourceName;
+};
diff --git a/Engine/Plugins/Animation/LiveLink/Source/LiveLinkHub/Private/Recording/LiveLinkPlaybackSubject.h b/Engine/Plugins/Animation/LiveLink/Source/LiveLinkHub/Private/Recording/LiveLinkPlaybackSubject.h
new file mode 100644
--- /dev/null
+++ b/Engine/Plugins/Animation/LiveLink/Source/LiveLinkHub/Private/Recording/LiveLinkPlaybackSubject.h
@@ -0,0 +1,100 @@
+// Copyright Epic Games, Inc. All Rights Reserved.
+
+#pragma once
+
+#include "Containers/Array.h"
+#include "CoreTypes.h"
+#include "LiveLinkFrameTranslator.h"
+#include "LiveLinkSubject.h"
+#include "LiveLinkTypes.h"
+#include "Templates/SubclassOf.h"
+#include "UObject/NameTypes.h"
+#include "UObject/Object.h"
+#include "UObject/ObjectMacros.h"
+#include "UObject/ObjectPtr.h"
+#include "UObject/UObjectGlobals.h"
+
+class ILiveLinkClient;
+class ULiveLinkRole;
+
+// A Virtual subject is made up of one or more real subjects from a source
+class FLiveLinkPlaybackSubject : public FLiveLinkSubject
+{
+	//~ Begin ILiveLinkSubject Interface
+public:
+	FLiveLinkPlaybackSubject(TSharedPtr<FLiveLinkTimedDataInput> InTimedDataGroup)
+		: FLiveLinkSubject(MoveTemp(InTimedDataGroup))
+	{
+	}
+
+	virtual ~FLiveLinkPlaybackSubject()
+	{
+	}
+
+	virtual void Initialize(FLiveLinkSubjectKey InSubjectKey, TSubclassOf<ULiveLinkRole> InRole, ILiveLinkClient* InLiveLinkClient) override
+	{
+		SubjectKey = InSubjectKey;
+		Role = InRole;
+	}
+	
+	virtual void Update() override
+	{
+	}
+	
+	virtual bool EvaluateFrame(TSubclassOf<ULiveLinkRole> InDesiredRole, FLiveLinkSubjectFrameData& OutFrame) override
+	{
+		return true;
+	}
+	
+	virtual void ClearFrames() override
+	{
+	}
+	
+	virtual FLiveLinkSubjectKey GetSubjectKey() const override
+	{
+		return SubjectKey;
+	}
+	
+	virtual TSubclassOf<ULiveLinkRole> GetRole() const override
+	{
+		return Role;
+	}
+	
+	virtual bool HasValidFrameSnapshot() const override
+	{
+		return true;
+	}
+	
+	virtual FLiveLinkStaticDataStruct& GetStaticData() override
+	{
+		return CurrentFrameSnapshot.StaticData;
+	}
+	
+	virtual const FLiveLinkStaticDataStruct& GetStaticData() const override
+	{
+		return CurrentFrameSnapshot.StaticData;
+	}
+	
+	virtual const TArray<ULiveLinkFrameTranslator::FWorkerSharedPtr> GetFrameTranslators() const override
+	{
+		return {};
+	}
+	
+	virtual TArray<FLiveLinkTime> GetFrameTimes() const override
+	{
+		return {};
+	}
+protected:
+	virtual const FLiveLinkSubjectFrameData& GetFrameSnapshot() const override
+	{
+		return CurrentFrameSnapshot;
+	}
+	//~ End ILiveLinkSubject Interface
+	
+private:
+	/** Not used, but needed to implement the FLiveLinkSubject interface. */
+	FLiveLinkSubjectFrameData CurrentFrameSnapshot;
+
+	/** Subject key for the playback subject. */
+	FLiveLinkSubjectKey SubjectKey;
+};
diff --git a/Engine/Plugins/Animation/LiveLink/Source/LiveLinkHub/Private/UI/Widgets/LiveLinkHubMainTabController.cpp b/Engine/Plugins/Animation/LiveLink/Source/LiveLinkHub/Private/UI/Widgets/LiveLinkHubMainTabController.cpp
new file mode 100644
--- /dev/null
+++ b/Engine/Plugins/Animation/LiveLink/Source/LiveLinkHub/Private/UI/Widgets/LiveLinkHubMainTabController.cpp
@@ -0,0 +1,43 @@
+// Copyright Epic Games, Inc. All Rights Reserved.
+
+#include "LiveLinkHubMainTabController.h"
+
+#include "UI/Window/LiveLinkHubWindowController.h"
+#include "SLiveLinkHubMainTabView.h"
+#include "Widgets/Docking/SDockTab.h"
+
+#define LOCTEXT_NAMESPACE "LiveLinkHubUI"
+
+static const FName LiveLinkHubTabID = "LiveLinkHubMainTab";
+
+void FLiveLinkHubMainTabController::Init(const FLiveLinkHubComponentInitParams& Params)
+{
+	FGlobalTabmanager::Get()->RegisterTabSpawner(LiveLinkHubTabID, FOnSpawnTab::CreateRaw(this, &FLiveLinkHubMainTabController::SpawnMainTab, Params.WindowController->GetRootWindow()))
+		.SetDisplayName(LOCTEXT("LiveLinkHubTabLabel", "LiveLink Hub"))
+		.SetTooltipText(LOCTEXT("MainTabTooltipText", "LiveLink Hub Main Tab"));
+	
+	Params.MainStack->AddTab(LiveLinkHubTabID, ETabState::OpenedTab);
+}
+
+void FLiveLinkHubMainTabController::OpenTab()
+{
+	FGlobalTabmanager::Get()->TryInvokeTab(LiveLinkHubTabID);
+}
+
+TSharedRef<SDockTab> FLiveLinkHubMainTabController::SpawnMainTab(const FSpawnTabArgs& Args, TSharedPtr<SWindow> RootWindow)
+{
+	const TSharedRef<SDockTab> DockTab = SNew(SDockTab)
+		.Label(LOCTEXT("LiveLinkHubTitle", "LiveLink Hub"))
+		.TabRole(MajorTab)
+		.CanEverClose(false);
+
+	DockTab->SetContent(
+		SAssignNew(MainTabView, SLiveLinkHubMainTabView, LiveLinkHubTabID)
+			.ConstructUnderMajorTab(DockTab)
+			.ConstructUnderWindow(RootWindow)
+	);
+
+	return DockTab;
+}
+
+#undef LOCTEXT_NAMESPACE
diff --git a/Engine/Plugins/Animation/LiveLink/Source/LiveLinkHub/Private/UI/Widgets/LiveLinkHubMainTabController.h b/Engine/Plugins/Animation/LiveLink/Source/LiveLinkHub/Private/UI/Widgets/LiveLinkHubMainTabController.h
new file mode 100644
--- /dev/null
+++ b/Engine/Plugins/Animation/LiveLink/Source/LiveLinkHub/Private/UI/Widgets/LiveLinkHubMainTabController.h
@@ -0,0 +1,33 @@
+// Copyright Epic Games, Inc. All Rights Reserved.
+
+#pragma once
+
+#include "UI/Window/ILiveLinkHubComponent.h"
+
+
+class ILiveLinkClient;
+class FSpawnTabArgs;
+class SLiveLinkHubMainTabView;
+class SWindow;
+
+class FLiveLinkHubMainTabController : public ILiveLinkHubComponent, public TSharedFromThis<FLiveLinkHubMainTabController>
+{
+public:
+
+	//~ Begin ILiveLinkHubComponent Interface
+	virtual void Init(const FLiveLinkHubComponentInitParams& Params) override;
+	//~ End ILiveLinkHubComponent Interface
+
+	/** Spawn the main tab. */
+	void OpenTab();
+
+private:
+	/** Manages the sub-tabs */
+	TSharedPtr<SLiveLinkHubMainTabView> MainTabView;
+
+	/** Spawn the main tab. */
+	TSharedRef<SDockTab> SpawnMainTab(const FSpawnTabArgs& Args, TSharedPtr<SWindow> RootWindow);
+
+	/** LiveLink Client */
+	ILiveLinkClient* LiveLinkClient = nullptr;
+};
diff --git a/Engine/Plugins/Animation/LiveLink/Source/LiveLinkHub/Private/UI/Widgets/SLiveLinkHubMainTabView.cpp b/Engine/Plugins/Animation/LiveLink/Source/LiveLinkHub/Private/UI/Widgets/SLiveLinkHubMainTabView.cpp
new file mode 100644
--- /dev/null
+++ b/Engine/Plugins/Animation/LiveLink/Source/LiveLinkHub/Private/UI/Widgets/SLiveLinkHubMainTabView.cpp
@@ -0,0 +1,232 @@
+// Copyright Epic Games, Inc. All Rights Reserved.
+
+#include "SLiveLinkHubMainTabView.h"
+
+#include "Features/IModularFeatures.h"
+#include "IDetailsView.h"
+#include "ILiveLinkClient.h"
+#include "LiveLinkClientPanelViews.h"
+#include "LiveLinkClientPanelToolbar.h"
+#include "LiveLinkHub.h"
+#include "LiveLinkHubModule.h"
+#include "LiveLinkPanelController.h"
+#include "Modules/ModuleManager.h"
+#include "Recording/LiveLinkHubPlaybackController.h"
+#include "Recording/LiveLinkHubRecordingController.h"
+#include "Recording/LiveLinkHubRecordingListController.h"
+#include "SLiveLinkDataView.h"
+#include "Widgets/Docking/SDockTab.h"
+#include "Widgets/Input/SCheckBox.h"
+#include "Widgets/Text/STextBlock.h"
+
+#define LOCTEXT_NAMESPACE "LiveLinkHub.SLiveLinkHubMainTabView"
+
+const FName SLiveLinkHubMainTabView::StatusBarTabId("StatusBarTabId");
+const FName SLiveLinkHubMainTabView::SourcesTabId("SourcesTabId");
+const FName SLiveLinkHubMainTabView::SourceDetailsTabId("SourceDetailsTabId");
+const FName SLiveLinkHubMainTabView::SubjectsTabId("SubjectsTabId");
+const FName SLiveLinkHubMainTabView::SubjectsDetailsTabId("SubjectsDetailsTabId");
+const FName SLiveLinkHubMainTabView::PlaybackTabId("PlaybackTabId");
+const FName SLiveLinkHubMainTabView::ClientsTabId("ClientsTabId");
+const FName SLiveLinkHubMainTabView::ClientDetailsTabId("ClientDetailsTabId");
+
+const FText SLiveLinkHubMainTabView::SourcesTabName = LOCTEXT("SourcesTabLabel", "Sources");
+const FText SLiveLinkHubMainTabView::SourceDetailsTabName = LOCTEXT("SourceDetailsTabLabel", "Source Details");
+const FText SLiveLinkHubMainTabView::SubjectsTabName = LOCTEXT("SubjectsTabLabel", "Subjects");
+const FText SLiveLinkHubMainTabView::SubjectsDetailsTabName = LOCTEXT("SubjectsDetailsTabLabel", "Subjects Details");
+const FText SLiveLinkHubMainTabView::PlaybackTabName = LOCTEXT("PlaybackTabLabel", "Playback");
+const FText SLiveLinkHubMainTabView::ClientsTabName = LOCTEXT("ClientsTabLabel", "Clients");
+const FText SLiveLinkHubMainTabView::ClientDetailsTabName = LOCTEXT("ClientDetailsTabLabel", "Client Details");
+
+void SLiveLinkHubMainTabView::Construct(const FArguments& InArgs, FName InStatusBarID)
+{	
+	PanelController = MakeShared<FLiveLinkPanelController>();
+
+	SLiveLinkHubTabViewWithManagerBase::Construct(
+		SLiveLinkHubTabViewWithManagerBase::FArguments()
+		.ConstructUnderWindow(InArgs._ConstructUnderWindow)
+		.ConstructUnderMajorTab(InArgs._ConstructUnderMajorTab)
+		.CreateTabs(FCreateTabs::CreateLambda([this, &InArgs](const TSharedRef<FTabManager>& InTabManager, const TSharedRef<FTabManager::FLayout>& InLayout)
+		{
+			CreateTabs(InTabManager, InLayout, InArgs);
+		}))
+		.LayoutName("LiveLinkHubSourcesTabView_v1.0")
+	);
+}
+
+void SLiveLinkHubMainTabView::CreateTabs(const TSharedRef<FTabManager>& InTabManager, const TSharedRef<FTabManager::FLayout>& InLayout, const FArguments& InArgs)
+{
+	InTabManager->RegisterTabSpawner(SourcesTabId, FOnSpawnTab::CreateSP(this, &SLiveLinkHubMainTabView::SpawnSourcesTab))
+		.SetDisplayName(SourcesTabName);
+	InTabManager->RegisterTabSpawner(SourceDetailsTabId, FOnSpawnTab::CreateSP(this, &SLiveLinkHubMainTabView::SpawnSourcesDetailsTab))
+		.SetDisplayName(SourceDetailsTabName);
+	InTabManager->RegisterTabSpawner(SubjectsTabId, FOnSpawnTab::CreateSP(this, &SLiveLinkHubMainTabView::SpawnSubjectsTab))
+		.SetDisplayName(SubjectsTabName);
+	InTabManager->RegisterTabSpawner(SubjectsDetailsTabId, FOnSpawnTab::CreateSP(this, &SLiveLinkHubMainTabView::SpawnSubjectsDetailsTab))
+		.SetDisplayName(SubjectsDetailsTabName);
+	InTabManager->RegisterTabSpawner(PlaybackTabId, FOnSpawnTab::CreateSP(this, &SLiveLinkHubMainTabView::SpawnPlaybackTab))
+		.SetDisplayName(PlaybackTabName);
+	InTabManager->RegisterTabSpawner(ClientsTabId, FOnSpawnTab::CreateSP(this, &SLiveLinkHubMainTabView::SpawnClientsTab))
+		.SetDisplayName(ClientsTabName);
+	InTabManager->RegisterTabSpawner(ClientDetailsTabId, FOnSpawnTab::CreateSP(this, &SLiveLinkHubMainTabView::SpawnClientDetailsTab))
+		.SetDisplayName(ClientDetailsTabName);
+	
+	InLayout->AddArea
+		(
+			FTabManager::NewPrimaryArea()
+				->SetOrientation(Orient_Vertical)
+				->Split
+				(
+					FTabManager::NewSplitter()
+					->SetSizeCoefficient(1.f)
+					->SetOrientation(Orient_Horizontal)
+					->Split
+					(
+						FTabManager::NewSplitter()
+						->SetSizeCoefficient(0.25f)
+						->SetOrientation(Orient_Vertical)
+						->Split
+						(
+							FTabManager::NewStack()
+							->SetSizeCoefficient(0.5f)
+							->AddTab(SourcesTabId, ETabState::OpenedTab)
+						)
+						->Split
+						(
+							FTabManager::NewStack()
+							->SetSizeCoefficient(0.5f)
+							->AddTab(SourceDetailsTabId, ETabState::OpenedTab)
+						)
+					)
+					->Split
+					(
+						FTabManager::NewSplitter()
+						->SetSizeCoefficient(0.25f)
+						->SetOrientation(Orient_Vertical)
+						->Split
+						(
+							FTabManager::NewStack()
+							->SetSizeCoefficient(0.5f)
+							->AddTab(SubjectsTabId, ETabState::OpenedTab)
+							->AddTab(PlaybackTabId, ETabState::OpenedTab)
+							->SetForegroundTab(SubjectsTabId)
+						)
+						->Split
+						(
+							FTabManager::NewStack()
+							->SetSizeCoefficient(0.5f)
+							->AddTab(SubjectsDetailsTabId, ETabState::OpenedTab)
+						)
+					)
+					->Split
+					(
+						FTabManager::NewSplitter()
+						->SetSizeCoefficient(0.25f)
+						->SetOrientation(Orient_Vertical)
+						->Split
+						(
+							FTabManager::NewStack()
+							->SetSizeCoefficient(0.5f)
+							->AddTab(ClientsTabId, ETabState::OpenedTab)
+						)
+						->Split
+						(
+							FTabManager::NewStack()
+							->SetSizeCoefficient(0.5f)
+							->AddTab(ClientDetailsTabId, ETabState::OpenedTab)
+						)
+					)
+				)
+		);
+}
+
+TSharedRef<SDockTab> SLiveLinkHubMainTabView::SpawnSourcesTab(const FSpawnTabArgs& InTabArgs)
+{
+	FLiveLinkClient* Client = (FLiveLinkClient*)&IModularFeatures::Get().GetModularFeature<ILiveLinkClient>(ILiveLinkClient::ModularFeatureName);
+	FLiveLinkHubModule& LiveLinkHubModule = FModuleManager::Get().GetModuleChecked<FLiveLinkHubModule>("LiveLinkHub");
+	
+	return SNew(SDockTab)
+		.TabRole(PanelTab)
+		[
+			SNew(SVerticalBox)
+			+SVerticalBox::Slot()
+			.AutoHeight()
+			[
+				SNew(SLiveLinkClientPanelToolbar, Client)
+				.ParentWindow(LiveLinkHubModule.GetLiveLinkHub()->GetRootWindow())
+			]
+			+SVerticalBox::Slot()
+			.FillHeight(1.f)
+			.Padding(FMargin(0.0f, 4.0f, 0.0f, 0.0f))
+			[
+				PanelController->SourcesView->SourcesListView.ToSharedRef()
+			]
+		];
+}
+
+
+TSharedRef<SDockTab> SLiveLinkHubMainTabView::SpawnSourcesDetailsTab(const FSpawnTabArgs& InTabArgs)
+{
+	return SNew(SDockTab)
+		.Label(SourceDetailsTabName)
+		.TabRole(PanelTab)
+		[
+			PanelController->SourcesDetailsView.ToSharedRef()
+		];
+}
+
+TSharedRef<SDockTab> SLiveLinkHubMainTabView::SpawnSubjectsTab(const FSpawnTabArgs& InTabArgs)
+{
+	return SNew(SDockTab)
+		.Label(SubjectsTabName)
+		.TabRole(PanelTab)
+		[
+			PanelController->SubjectsView->SubjectsTreeView.ToSharedRef()
+		];
+}
+
+TSharedRef<SDockTab> SLiveLinkHubMainTabView::SpawnSubjectsDetailsTab(const FSpawnTabArgs& InTabArgs)
+{
+	return SNew(SDockTab)
+		.Label(SubjectsDetailsTabName)
+		.TabRole(PanelTab)
+		[
+			PanelController->SubjectsDetailsView.ToSharedRef()
+		];
+}
+
+TSharedRef<SDockTab> SLiveLinkHubMainTabView::SpawnPlaybackTab(const FSpawnTabArgs& InTabArgs)
+{
+	FLiveLinkHubModule& LiveLinkHubModule = FModuleManager::Get().GetModuleChecked<FLiveLinkHubModule>("LiveLinkHub");
+
+	return SNew(SDockTab)
+		.Label(PlaybackTabName)
+		.TabRole(PanelTab)
+		[
+			LiveLinkHubModule.GetRecordingListController()->MakeRecordingList()
+		];
+}
+
+TSharedRef<SDockTab> SLiveLinkHubMainTabView::SpawnClientsTab(const FSpawnTabArgs& InTabArgs)
+{
+	return SNew(SDockTab)
+		.Label(ClientsTabName)
+		.TabRole(PanelTab)
+		[
+			SNew(STextBlock)
+			.Text(LOCTEXT("TemporaryText", "Client list here"))
+		];
+}
+
+TSharedRef<SDockTab> SLiveLinkHubMainTabView::SpawnClientDetailsTab(const FSpawnTabArgs& InTabArgs)
+{
+	return SNew(SDockTab)
+		.Label(ClientDetailsTabName)
+		.TabRole(PanelTab)
+		[
+			SNew(STextBlock)
+			.Text(LOCTEXT("TemporaryText", "Client details here"))
+		];
+}
+
+#undef LOCTEXT_NAMESPACE
diff --git a/Engine/Plugins/Animation/LiveLink/Source/LiveLinkHub/Private/UI/Widgets/SLiveLinkHubMainTabView.h b/Engine/Plugins/Animation/LiveLink/Source/LiveLinkHub/Private/UI/Widgets/SLiveLinkHubMainTabView.h
new file mode 100644
--- /dev/null
+++ b/Engine/Plugins/Animation/LiveLink/Source/LiveLinkHub/Private/UI/Widgets/SLiveLinkHubMainTabView.h
@@ -0,0 +1,57 @@
+// Copyright Epic Games, Inc. All Rights Reserved.
+
+#pragma once
+
+#include "SLiveLinkHubTabViewWithManagerBase.h"
+
+class FLiveLinkHubRecordingController;
+class ILiveLinkClient;
+
+/** Manages the UI logic of the Clients tab */
+class SLiveLinkHubMainTabView : public SLiveLinkHubTabViewWithManagerBase
+{
+public:
+	//~ Tab IDs
+	static const FName StatusBarTabId;
+	static const FName SourcesTabId;
+	static const FName SourceDetailsTabId;
+	static const FName SubjectsTabId;
+	static const FName SubjectsDetailsTabId;
+	static const FName PlaybackTabId;
+	static const FName ClientsTabId;
+	static const FName ClientDetailsTabId;
+
+	//~ Tab Names
+	static const FText SourcesTabName;
+	static const FText SourceDetailsTabName;
+	static const FText SubjectsTabName;
+	static const FText SubjectsDetailsTabName;
+	static const FText PlaybackTabName;
+	static const FText ClientsTabName;
+	static const FText ClientDetailsTabName;
+
+	SLATE_BEGIN_ARGS(SLiveLinkHubMainTabView)
+	{}
+		SLATE_ARGUMENT(TSharedPtr<SDockTab>, ConstructUnderMajorTab)
+		SLATE_ARGUMENT(TSharedPtr<SWindow>, ConstructUnderWindow)
+	SLATE_END_ARGS()
+
+	void Construct(const FArguments& InArgs, FName InStatusBarID);
+	
+private:
+	/** Create all livelink hub tabs. */
+	void CreateTabs(const TSharedRef<FTabManager>& InTabManager, const TSharedRef<FTabManager::FLayout>& InLayout, const FArguments& InArgs);
+
+	//~ Respective tab handlers.
+	TSharedRef<SDockTab> SpawnSourcesTab(const FSpawnTabArgs& InTabArgs);
+	TSharedRef<SDockTab> SpawnSourcesDetailsTab(const FSpawnTabArgs& InTabArgs);
+	TSharedRef<SDockTab> SpawnSubjectsTab(const FSpawnTabArgs& InTabArgs);
+	TSharedRef<SDockTab> SpawnSubjectsDetailsTab(const FSpawnTabArgs& InTabArgs);
+	TSharedRef<SDockTab> SpawnPlaybackTab(const FSpawnTabArgs& InTabArgs);
+	TSharedRef<SDockTab> SpawnClientsTab(const FSpawnTabArgs& InTabArgs);
+	TSharedRef<SDockTab> SpawnClientDetailsTab(const FSpawnTabArgs& InTabArgs);
+
+private:
+	/** Holds the livelink panel controller responsible for creating sources and subjects tabs. */
+	TSharedPtr<class FLiveLinkPanelController> PanelController;
+};
diff --git a/Engine/Plugins/Animation/LiveLink/Source/LiveLinkHub/Private/UI/Widgets/SLiveLinkHubTabViewBase.cpp b/Engine/Plugins/Animation/LiveLink/Source/LiveLinkHub/Private/UI/Widgets/SLiveLinkHubTabViewBase.cpp
new file mode 100644
--- /dev/null
+++ b/Engine/Plugins/Animation/LiveLink/Source/LiveLinkHub/Private/UI/Widgets/SLiveLinkHubTabViewBase.cpp
@@ -0,0 +1,109 @@
+// Copyright Epic Games, Inc. All Rights Reserved.
+
+#include "SLiveLinkHubTabViewBase.h"
+
+#include "LiveLinkHub.h"
+#include "LiveLinkHubModule.h"
+#include "Framework/Commands/UIAction.h"
+#include "Framework/MultiBox/MultiBoxBuilder.h"
+#include "Misc/App.h"
+#include "Modules/ModuleManager.h"
+#include "Recording/LiveLinkHubRecordingController.h"
+#include "Widgets/Input/SButton.h"
+#include "Widgets/Input/SComboButton.h"
+#include "Widgets/Layout/SBorder.h"
+#include "Widgets/Layout/SSeparator.h"
+#include "Widgets/SBoxPanel.h"
+#include "Widgets/STimecode.h"
+#include "Widgets/Text/STextBlock.h"
+
+#define LOCTEXT_NAMESPACE "LiveLinkHub"
+
+TSharedRef<SWidget> GetModeSwitcherContent()
+{
+	FMenuBuilder MenuBuilder(true, NULL);
+	{
+		FUIAction CreatorModeAction(FExecuteAction::CreateLambda([](){}));
+		MenuBuilder.AddMenuEntry(LOCTEXT("CreatorModeLabel", "Creator Mode"), LOCTEXT("CreatorModeHint", "-Placeholder text-"), FSlateIcon(), CreatorModeAction);
+		
+		FUIAction StudioModeAction(FExecuteAction::CreateLambda([](){}));
+		MenuBuilder.AddMenuEntry(LOCTEXT("StudioModeLabel", "Studio Mode"), LOCTEXT("StudioModeHint", "-Placeholder text-"), FSlateIcon(), StudioModeAction);
+	}
+
+	return MenuBuilder.MakeWidget();
+}
+
+void SLiveLinkHubTabViewBase::Construct(const FArguments& InArgs)
+{
+	TSharedPtr<FLiveLinkHub> LiveLinkHub = FModuleManager::Get().GetModuleChecked<FLiveLinkHubModule>("LiveLinkHub").GetLiveLinkHub();
+
+	ChildSlot
+	[
+		SNew(SBorder)
+		.BorderImage(FAppStyle::Get().GetBrush("ToolPanel.GroupBorder"))
+		.Padding(FMargin(1.0f, 2.0f))
+		[
+			SNew(SVerticalBox)
+			+SVerticalBox::Slot()
+			.FillHeight(0.05f)
+			.HAlign(HAlign_Fill)
+			.VAlign(VAlign_Center)
+			[
+				SNew(SHorizontalBox)
+				+ SHorizontalBox::Slot()
+				.Padding(2.f)
+				.HAlign(HAlign_Left)
+				[
+					// Should be replaced by the mode switcher widget
+					SNew(SComboButton)
+					.ContentPadding(3.f)
+					.OnGetMenuContent_Static(&::GetModeSwitcherContent)
+					.ButtonContent()
+					[
+						// Should switch depending on the currently selected mode.
+						SNew(STextBlock).Text(LOCTEXT("SelectModeLabl", "Creator Mode"))
+					]
+				]
+				+ SHorizontalBox::Slot()
+				.HAlign(HAlign_Right)
+				.VAlign(VAlign_Center)
+				[
+					SNew(SHorizontalBox)
+					+ SHorizontalBox::Slot()
+					.Padding(2.f)
+					.VAlign(VAlign_Center)
+					.AutoWidth()
+					[
+						LiveLinkHub->GetRecordingController()->MakeRecordToolbarEntry()
+					]
+					+ SHorizontalBox::Slot()
+					.Padding(2.f)
+					.VAlign(VAlign_Center)
+					.AutoWidth()
+					[
+						SNew(SSeparator)
+						.Orientation(Orient_Vertical)
+					]
+					+ SHorizontalBox::Slot()
+					.Padding(2.f)
+                    .VAlign(VAlign_Center)
+					[
+						SNew(STimecode)
+						.Timecode(MakeAttributeLambda([]
+						{
+							return FApp::GetTimecode();
+						}))
+					]
+				]
+			]
+			+SVerticalBox::Slot()
+			.FillHeight(1.0f)
+			.Padding(1.0f, 2.0f)
+			[
+				InArgs._Content.Widget
+			]
+		]
+	];
+}
+
+#undef LOCTEXT_NAMESPACE // LiveLinkHub
diff --git a/Engine/Plugins/Animation/LiveLink/Source/LiveLinkHub/Private/UI/Widgets/SLiveLinkHubTabViewBase.h b/Engine/Plugins/Animation/LiveLink/Source/LiveLinkHub/Private/UI/Widgets/SLiveLinkHubTabViewBase.h
new file mode 100644
--- /dev/null
+++ b/Engine/Plugins/Animation/LiveLink/Source/LiveLinkHub/Private/UI/Widgets/SLiveLinkHubTabViewBase.h
@@ -0,0 +1,26 @@
+// Copyright Epic Games, Inc. All Rights Reserved.
+
+#pragma once
+
+#include "Widgets/SCompoundWidget.h"
+
+class FLiveLinkHub;
+
+/**
+ * Utility base class for all content that is displayed in an entire major tab.
+ */
+class SLiveLinkHubTabViewBase : public SCompoundWidget
+{
+public:
+
+	SLATE_BEGIN_ARGS(SLiveLinkHubTabViewBase) {}
+		/** The tab's content */
+		SLATE_NAMED_SLOT(FArguments, Content)
+	SLATE_END_ARGS()
+
+	/**
+	 * @param InArgs
+	 * @param InStatusBarId Unique ID needed for the status bar
+	 */
+	void Construct(const FArguments& InArgs);
+};
diff --git a/Engine/Plugins/Animation/LiveLink/Source/LiveLinkHub/Private/UI/Widgets/SLiveLinkHubTabViewWithManagerBase.cpp b/Engine/Plugins/Animation/LiveLink/Source/LiveLinkHub/Private/UI/Widgets/SLiveLinkHubTabViewWithManagerBase.cpp
new file mode 100644
--- /dev/null
+++ b/Engine/Plugins/Animation/LiveLink/Source/LiveLinkHub/Private/UI/Widgets/SLiveLinkHubTabViewWithManagerBase.cpp
@@ -0,0 +1,73 @@
+// Copyright Epic Games, Inc. All Rights Reserved.
+
+#include "SLiveLinkHubTabViewWithManagerBase.h"
+
+#include "Framework/Application/SlateApplication.h"
+#include "Framework/Commands/UICommandList.h"
+#include "Framework/Docking/LayoutService.h"
+#include "Framework/MultiBox/MultiBoxBuilder.h"
+#include "UI/Widgets/SLiveLinkHubTabViewBase.h"
+
+#define LOCTEXT_NAMESPACE "LiveLinkHub.SLiveLinkHubTabViewWithManagerBase"
+
+void SLiveLinkHubTabViewWithManagerBase::Construct(const FArguments& InArgs)
+{
+	check(InArgs._ConstructUnderWindow && InArgs._ConstructUnderMajorTab);
+	SLiveLinkHubTabViewBase::Construct(
+		SLiveLinkHubTabViewBase::FArguments()
+		.Content()
+		[
+			InArgs._OverlayTabs.IsBound() ? InArgs._OverlayTabs.Execute(CreateTabs(InArgs)) : CreateTabs(InArgs)
+		]
+	);
+}
+
+TSharedRef<SWidget> SLiveLinkHubTabViewWithManagerBase::CreateTabs(const FArguments& InArgs)
+{
+	TabManager = FGlobalTabmanager::Get()->NewTabManager(InArgs._ConstructUnderMajorTab.ToSharedRef());
+	TabManager->SetMainTab(InArgs._ConstructUnderMajorTab.ToSharedRef());
+	
+	TSharedRef<FTabManager::FLayout> Layout = FTabManager::NewLayout(InArgs._LayoutName);
+	InArgs._CreateTabs.ExecuteIfBound(TabManager.ToSharedRef(), Layout);
+	Layout = FLayoutSaveRestore::LoadFromConfig(TEXT("LiveLinkHubLayout"), Layout);
+
+	TabManager->SetOnPersistLayout(
+		FTabManager::FOnPersistLayout::CreateStatic(
+			[](const TSharedRef<FTabManager::FLayout>& InLayout)
+			{
+				if (InLayout->GetPrimaryArea().Pin().IsValid())
+				{
+					FLayoutSaveRestore::SaveToConfig(TEXT("LiveLinkHubLayout"), InLayout);
+				}
+			}
+		)
+	);
+
+	TSharedRef<SWidget> Result = TabManager->RestoreFrom(Layout, InArgs._ConstructUnderWindow).ToSharedRef();
+	FMenuBarBuilder MenuBarBuilder = FMenuBarBuilder(TSharedPtr<FUICommandList>());
+	FillInDefaultMenuItems(MenuBarBuilder);
+	InArgs._CreateMenuBar.ExecuteIfBound(MenuBarBuilder);
+	
+	const TSharedRef<SWidget> MenuWidget = MenuBarBuilder.MakeWidget();
+	TabManager->SetAllowWindowMenuBar(true);
+	TabManager->SetMenuMultiBox(MenuBarBuilder.GetMultiBox(), MenuWidget);
+	
+	return Result;
+}
+
+void SLiveLinkHubTabViewWithManagerBase::FillInDefaultMenuItems(FMenuBarBuilder MenuBarBuilder)
+{
+	MenuBarBuilder.AddPullDownMenu(
+		LOCTEXT("WindowMenuLabel", "Window"),
+		FText::GetEmpty(),
+		FNewMenuDelegate::CreateSP(this, &SLiveLinkHubTabViewWithManagerBase::FillWindowMenu),
+		"Window"
+	);
+}
+
+void SLiveLinkHubTabViewWithManagerBase::FillWindowMenu(FMenuBuilder& MenuBuilder)
+{
+	TabManager->PopulateLocalTabSpawnerMenu(MenuBuilder);
+}
+
+#undef LOCTEXT_NAMESPACE
diff --git a/Engine/Plugins/Animation/LiveLink/Source/LiveLinkHub/Private/UI/Widgets/SLiveLinkHubTabViewWithManagerBase.h b/Engine/Plugins/Animation/LiveLink/Source/LiveLinkHub/Private/UI/Widgets/SLiveLinkHubTabViewWithManagerBase.h
new file mode 100644
--- /dev/null
+++ b/Engine/Plugins/Animation/LiveLink/Source/LiveLinkHub/Private/UI/Widgets/SLiveLinkHubTabViewWithManagerBase.h
@@ -0,0 +1,54 @@
+// Copyright Epic Games, Inc. All Rights Reserved.
+
+#pragma once
+
+#include "Framework/Docking/TabManager.h"
+#include "UI/Widgets/SLiveLinkHubTabViewBase.h"
+
+class FMenuBarBuilder;
+
+/** Base class for tab views that create sub-tabs. */
+class SLiveLinkHubTabViewWithManagerBase : public SLiveLinkHubTabViewBase
+{
+public:
+
+	DECLARE_DELEGATE_TwoParams(FCreateTabs, const TSharedRef<FTabManager>& /*TabManager*/, const TSharedRef<FTabManager::FLayout>& /*Layout*/);
+	DECLARE_DELEGATE_OneParam(FCreateWindowMenu, FMenuBarBuilder&);
+	DECLARE_DELEGATE_RetVal_OneParam(TSharedRef<SWidget>, FOverlayTabs, const TSharedRef<SWidget>& /* Tabs */);
+	
+	SLATE_BEGIN_ARGS(SLiveLinkHubTabViewWithManagerBase) {}
+		/** Which major tab to construct the sub-tabs under. */
+		SLATE_ARGUMENT(TSharedPtr<SDockTab>, ConstructUnderMajorTab)
+		/** The window in which the sub-tabs will be created. */
+		SLATE_ARGUMENT(TSharedPtr<SWindow>, ConstructUnderWindow)
+		/** Callback for creating the sub-tabs */
+		SLATE_EVENT(FCreateTabs, CreateTabs)
+		/** Called to fill the major tab's menu bar */
+		SLATE_EVENT(FCreateWindowMenu, CreateMenuBar)
+		/** Optional function to overlay all the tabs (for example to overlay a message to execute a console command first) */
+		SLATE_EVENT(FOverlayTabs, OverlayTabs)
+		/** Name to give the layout. Important for saving config. */
+		SLATE_ARGUMENT(FName, LayoutName)
+	SLATE_END_ARGS()
+
+	/**
+	 * @param InArgs
+	 * @param InStatusBarId Unique ID needed for the status bar
+	 */
+	void Construct(const FArguments& InArgs);
+
+protected:
+	/** Get the tab manager. */
+	TSharedPtr<FTabManager> GetTabManager() const { return TabManager; }
+	
+private:
+
+	/** Holds the tab manager that manages the front-end's tabs. */
+	TSharedPtr<FTabManager> TabManager;
+	/** Handles restoring the layout and instantiating tabs. */
+	TSharedRef<SWidget> CreateTabs(const FArguments& InArgs);
+	/** Creates the top bar menu items. */
+	void FillInDefaultMenuItems(FMenuBarBuilder MenuBarBuilder);
+	/** Populates the tab spawner's menu. */
+	void FillWindowMenu(FMenuBuilder& MenuBuilder);
+};
diff --git a/Engine/Plugins/Animation/LiveLink/Source/LiveLinkHub/Private/UI/Window/ILiveLinkHubComponent.h b/Engine/Plugins/Animation/LiveLink/Source/LiveLinkHub/Private/UI/Window/ILiveLinkHubComponent.h
new file mode 100644
--- /dev/null
+++ b/Engine/Plugins/Animation/LiveLink/Source/LiveLinkHub/Private/UI/Window/ILiveLinkHubComponent.h
@@ -0,0 +1,32 @@
+// Copyright Epic Games, Inc. All Rights Reserved.
+
+#pragma once
+
+#include "Framework/Docking/TabManager.h"
+
+class ILiveLinkClient;
+class FLiveLinkHubWindowController;
+
+struct FLiveLinkHubComponentInitParams
+{
+	/** Manages the hub's slate application */
+	TSharedRef<FLiveLinkHubWindowController> WindowController;
+
+	/** The root area of the main window layout. Add tabs to this. */
+	TSharedRef<FTabManager::FStack> MainStack;
+
+	FLiveLinkHubComponentInitParams(TSharedRef<FLiveLinkHubWindowController> InWindowController, TSharedRef<FTabManager::FStack> InMainStack)
+		: WindowController(MoveTemp(InWindowController))
+		, MainStack(MoveTemp(InMainStack))
+	{}
+};
+
+/** Provides the base interface for elements in the LiveLink Hub UI. */
+class ILiveLinkHubComponent
+{
+public:
+	virtual ~ILiveLinkHubComponent() = default;
+
+	/** Initialises the component, e.g. registering tab spawners, etc. */
+	virtual void Init(const FLiveLinkHubComponentInitParams& Params) {}
+};
diff --git a/Engine/Plugins/Animation/LiveLink/Source/LiveLinkHub/Private/UI/Window/LiveLinkHubWindowController.cpp b/Engine/Plugins/Animation/LiveLink/Source/LiveLinkHub/Private/UI/Window/LiveLinkHubWindowController.cpp
new file mode 100644
--- /dev/null
+++ b/Engine/Plugins/Animation/LiveLink/Source/LiveLinkHub/Private/UI/Window/LiveLinkHubWindowController.cpp
@@ -0,0 +1,121 @@
+// Copyright Epic Games, Inc. All Rights Reserved.
+
+#include "LiveLinkHubWindowController.h"
+
+#include "Framework/Application/SlateApplication.h"
+#include "Framework/Docking/LayoutService.h"
+#include "Framework/Notifications/NotificationManager.h"
+#include "HAL/PlatformApplicationMisc.h"
+#include "Misc/ConfigCacheIni.h"
+#include "StandaloneRenderer.h"
+#include "UI/Widgets/LiveLinkHubMainTabController.h"
+#include "UI/Window/ModalWindowManager.h"
+
+#define LOCTEXT_NAMESPACE "LiveLinkHubWindowController"
+
+FLiveLinkHubWindowController::FLiveLinkHubWindowController(const FLiveLinkHubWindowInitParams& Params)
+	: LiveLinkHubLayoutIni(Params.LiveLinkHubLayoutIni)
+	, MainTabController(MakeShared<FLiveLinkHubMainTabController>())
+{
+	LiveLinkHubComponents.Add(MainTabController);
+	
+	ModalWindowManager = InitializeSlateApplication();
+}
+
+FLiveLinkHubWindowController::~FLiveLinkHubWindowController()
+{
+	FGlobalTabmanager::Get()->SaveAllVisualState();
+	FSlateApplication::Shutdown();
+}
+
+TSharedRef<SWindow> FLiveLinkHubWindowController::CreateWindow()
+{
+	FDisplayMetrics DisplayMetrics;
+	FSlateApplication::Get().GetDisplayMetrics(DisplayMetrics);
+	const float DPIScaleFactor = FPlatformApplicationMisc::GetDPIScaleFactorAtPoint(DisplayMetrics.PrimaryDisplayWorkAreaRect.Left, DisplayMetrics.PrimaryDisplayWorkAreaRect.Top);
+
+	constexpr bool bEmbedTitleAreaContent = true;
+	const FVector2D ClientSize(1200.0f * DPIScaleFactor, 800.0f * DPIScaleFactor);
+	TSharedRef<SWindow> RootWindowRef = SNew(SWindow)
+		.Title(LOCTEXT("WindowTitle", "LiveLink Hub"))
+		.CreateTitleBar(!bEmbedTitleAreaContent)
+		.SupportsMaximize(true)
+		.SupportsMinimize(true)
+		.IsInitiallyMaximized(false)
+		.IsInitiallyMinimized(false)
+		.SizingRule(ESizingRule::UserSized)
+		.AutoCenter(EAutoCenter::PreferredWorkArea)
+		.ClientSize(ClientSize)
+		.AdjustInitialSizeAndPositionForDPIScale(false);
+
+	RootWindow = RootWindowRef;
+		
+	constexpr bool bShowRootWindowImmediately = false;
+	FSlateApplication::Get().AddWindow(RootWindowRef, bShowRootWindowImmediately);
+	FGlobalTabmanager::Get()->SetRootWindow(RootWindowRef);
+	FGlobalTabmanager::Get()->SetAllowWindowMenuBar(true);
+	FSlateNotificationManager::Get().SetRootWindow(RootWindowRef);	
+
+	RootWindow->SetOnWindowClosed(FOnWindowClosed::CreateRaw(this, &FLiveLinkHubWindowController::OnWindowClosed));
+
+	return RootWindowRef;
+}
+void FLiveLinkHubWindowController::RestoreLayout()
+{
+	const TSharedRef<FTabManager::FLayout> DefaultLayout = FTabManager::NewLayout("LiveLinkHub_v1.0");
+	const TSharedRef<FTabManager::FArea> MainWindowArea = FTabManager::NewPrimaryArea()->SetOrientation(Orient_Horizontal);
+	const TSharedRef<FTabManager::FStack> MainStack = FTabManager::NewStack();
+	InitComponents(MainStack);
+	MainStack->SetHideTabWell(true);
+
+	MainWindowArea->Split(MainStack);
+	DefaultLayout->AddArea(MainWindowArea);
+
+	PersistentLayout = FLayoutSaveRestore::LoadFromConfig(LiveLinkHubLayoutIni, DefaultLayout);
+
+	constexpr bool bEmbedTitleAreaContent = true;
+	const TSharedPtr<SWidget> Content = FGlobalTabmanager::Get()->RestoreFrom(PersistentLayout.ToSharedRef(), RootWindow, bEmbedTitleAreaContent, EOutputCanBeNullptr::Never);
+	RootWindow->SetContent(Content.ToSharedRef());
+
+	RootWindow->ShowWindow();
+	constexpr bool bForceWindowToFront = true;
+	RootWindow->BringToFront(bForceWindowToFront);
+
+	MainTabController->OpenTab();
+}
+
+TSharedPtr<FModalWindowManager> FLiveLinkHubWindowController::InitializeSlateApplication()
+{
+	FSlateApplication::InitializeAsStandaloneApplication(GetStandardStandaloneRenderer());
+
+	const FText ApplicationTitle = LOCTEXT("AppTitle", "LiveLink Hub");
+	FGlobalTabmanager::Get()->SetApplicationTitle(ApplicationTitle);
+
+	return MakeShared<FModalWindowManager>(CreateWindow());
+}
+
+void FLiveLinkHubWindowController::InitComponents(const TSharedRef<FTabManager::FStack>& MainArea)
+{
+	const FLiveLinkHubComponentInitParams Params { SharedThis(this), MainArea };
+	for (const TSharedRef<ILiveLinkHubComponent>& LiveLinkHubComponent : LiveLinkHubComponents)
+	{
+		LiveLinkHubComponent->Init(Params);
+	}
+}
+
+void FLiveLinkHubWindowController::OnWindowClosed(const TSharedRef<SWindow>& Window)
+{
+	SaveLayout();
+	RootWindow.Reset();
+}
+
+void FLiveLinkHubWindowController::SaveLayout() const
+{
+	if (PersistentLayout)
+	{
+		FLayoutSaveRestore::SaveToConfig(LiveLinkHubLayoutIni, PersistentLayout.ToSharedRef());
+	    GConfig->Flush(false, LiveLinkHubLayoutIni);
+	}
+}
+
+#undef LOCTEXT_NAMESPACE 
diff --git a/Engine/Plugins/Animation/LiveLink/Source/LiveLinkHub/Private/UI/Window/LiveLinkHubWindowController.h b/Engine/Plugins/Animation/LiveLink/Source/LiveLinkHub/Private/UI/Window/LiveLinkHubWindowController.h
new file mode 100644
--- /dev/null
+++ b/Engine/Plugins/Animation/LiveLink/Source/LiveLinkHub/Private/UI/Window/LiveLinkHubWindowController.h
@@ -0,0 +1,65 @@
+// Copyright Epic Games, Inc. All Rights Reserved.
+
+#pragma once
+
+#include "Framework/Docking/TabManager.h"
+
+class FLiveLinkHubMainTabController;
+class FLiveLinkHubSessionTabBase;
+class FOutputLogController;
+class ILiveLinkClient;
+class ILiveLinkHubComponent;
+class FModalWindowManager;
+class SWindow;
+
+struct FLiveLinkHubWindowInitParams
+{
+	/** Config path for hub layout ini */
+	FString LiveLinkHubLayoutIni;
+
+	FLiveLinkHubWindowInitParams() = default;
+
+	FLiveLinkHubWindowInitParams(FString LiveLinkHubLayoutIni)
+		: LiveLinkHubLayoutIni(MoveTemp(LiveLinkHubLayoutIni))
+	{}
+};
+
+/** Responsible for creating the Slate window for the hub. */
+class FLiveLinkHubWindowController : public TSharedFromThis<FLiveLinkHubWindowController>
+{
+public:
+	
+	FLiveLinkHubWindowController(const FLiveLinkHubWindowInitParams& Params);
+	~FLiveLinkHubWindowController();
+
+	/** Get the root window. */
+	TSharedPtr<SWindow> GetRootWindow() const { return RootWindow; }
+	/** Restore the window's layout from a config. */
+	void RestoreLayout();
+private:
+	/** Create the main window. */
+	TSharedRef<SWindow> CreateWindow();
+	/** Create the slate application that hosts the livelink hub. */
+	TSharedPtr<FModalWindowManager> InitializeSlateApplication();
+	/** Initialize components. */
+	void InitComponents(const TSharedRef<FTabManager::FStack>& MainArea);
+	/** Window closed handler. */
+	void OnWindowClosed(const TSharedRef<SWindow>& Window);
+	/** Save the layout to a json file. */
+	void SaveLayout() const;
+private:
+	/** Pointer to the livelink client. */
+	ILiveLinkClient* LiveLinkClient = nullptr;
+	/** The ini file to use for saving the layout */
+	FString LiveLinkHubLayoutIni;
+	/** Holds the current layout for saving later. */
+	TSharedPtr<FTabManager::FLayout> PersistentLayout;
+	/** The main window being managed */
+	TSharedPtr<SWindow> RootWindow;
+	/** Manages the main tab. */
+	TSharedRef<FLiveLinkHubMainTabController> MainTabController;
+	/** LiveLinkHub components */
+	TArray<TSharedRef<ILiveLinkHubComponent>> LiveLinkHubComponents;
+	/** Manages modal windows for the application. */
+	TSharedPtr<FModalWindowManager> ModalWindowManager;
+};
diff --git a/Engine/Plugins/Animation/LiveLink/Source/LiveLinkHub/Private/UI/Window/ModalWindowManager.cpp b/Engine/Plugins/Animation/LiveLink/Source/LiveLinkHub/Private/UI/Window/ModalWindowManager.cpp
new file mode 100644
--- /dev/null
+++ b/Engine/Plugins/Animation/LiveLink/Source/LiveLinkHub/Private/UI/Window/ModalWindowManager.cpp
@@ -0,0 +1,81 @@
+﻿// Copyright Epic Games, Inc. All Rights Reserved.
+
+#include "ModalWindowManager.h"
+
+#include "Dialog/SCustomDialog.h"
+
+FModalWindowManager::FModalWindowManager(const TSharedRef<SWindow>& InMainWindow)
+{
+	MainWindow = InMainWindow;
+	InMainWindow->GetOnWindowActivatedEvent().AddRaw(this, &FModalWindowManager::OnMainWindowActivated);
+
+	// Fake modal windows should be closed when the root window is closed
+	InMainWindow->GetOnWindowClosedEvent().AddRaw(this, &FModalWindowManager::OnWindowClosed);
+}
+
+FModalWindowManager::~FModalWindowManager()
+{
+	if (TSharedPtr<SWindow> MainWindowPin = MainWindow.Pin())
+	{
+		MainWindowPin->GetOnWindowClosedEvent().RemoveAll(this);
+		MainWindowPin->GetOnWindowActivatedEvent().RemoveAll(this);
+	}
+}
+
+void FModalWindowManager::ShowFakeModalWindow(const TSharedRef<SCustomDialog>& Dialog)
+{
+	if (ensureMsgf(bRootWindowIsValid, TEXT("Root window has already been destroyed!")))
+	{
+		Dialog->Show();
+		PushWindow(Dialog);
+	}
+}
+
+void FModalWindowManager::PushWindow(const TSharedRef<SWindow>& Window)
+{
+	if (ensureMsgf(bRootWindowIsValid, TEXT("Root window has already been destroyed!")))
+	{
+		check(!WindowStack.Contains(Window));
+		WindowStack.Push(Window);
+	
+		Window->GetOnWindowClosedEvent().AddLambda([this](const TSharedRef<SWindow>& Window)
+		{
+			// Might request window destruction twice
+			if (WindowStack.Contains(Window))
+			{
+				PopWindow();
+			}
+		});
+	}
+}
+
+void FModalWindowManager::PopWindow()
+{
+	WindowStack.Pop();
+}
+
+void FModalWindowManager::OnMainWindowActivated()
+{
+	if (WindowStack.Num() > 0)
+	{
+		if (const TSharedPtr<SWindow> PinnedWindow = WindowStack.Top().Pin())
+		{
+			PinnedWindow->BringToFront();
+			PinnedWindow->FlashWindow();
+		}
+	}
+}
+
+void FModalWindowManager::OnWindowClosed(const TSharedRef<SWindow>& InWindow)
+{
+	bRootWindowIsValid = false;
+
+	for (TWeakPtr<SWindow> ModalWindow : TArray(WindowStack))
+	{
+		if (TSharedPtr<SWindow> Pin = ModalWindow.Pin())
+		{
+			Pin->RequestDestroyWindow();
+		}
+	}
+	WindowStack.Empty();
+}
diff --git a/Engine/Plugins/Animation/LiveLink/Source/LiveLinkHub/Private/UI/Window/ModalWindowManager.h b/Engine/Plugins/Animation/LiveLink/Source/LiveLinkHub/Private/UI/Window/ModalWindowManager.h
new file mode 100644
--- /dev/null
+++ b/Engine/Plugins/Animation/LiveLink/Source/LiveLinkHub/Private/UI/Window/ModalWindowManager.h
@@ -0,0 +1,46 @@
+﻿// Copyright Epic Games, Inc. All Rights Reserved.
+
+#pragma once
+
+#include "Templates/SharedPointer.h"
+
+class SCustomDialog;
+class SWindow;
+
+/**
+ * Manages "fake" modal windows because the hub cannot use "true" modal windows.
+ *
+ * True modal windows block the main thread which is not an option because the hub must remain responsive.
+ */
+class FModalWindowManager
+{
+public:
+
+	FModalWindowManager(const TSharedRef<SWindow>& InMainWindow);
+	~FModalWindowManager();
+
+	/** Show a modal window and add it to our window stack. */
+	void ShowFakeModalWindow(const TSharedRef<SCustomDialog>& Dialog);
+	/** Called when a new fake modal window should be tracked on top of all windows. */
+	void PushWindow(const TSharedRef<SWindow>& Window);
+	/** Called when a fake modal window should no longer be tracked, i.e. when it is closed. */
+	void PopWindow();
+
+private:
+	/** Holds the modal windows. */
+	TArray<TWeakPtr<SWindow>> WindowStack;
+
+	/** Modal windows can only be added when this is true */
+	bool bRootWindowIsValid = true;
+
+	/** Handler called when the main window is activated. */
+	void OnMainWindowActivated();
+
+	/** Handler called when the main window is closed, used for destroying all windows in the stack. */
+	void OnWindowClosed(const TSharedRef<SWindow>& InWindow);
+
+private:
+	/** Weak pointer to the main window of the app. */
+	TWeakPtr<SWindow> MainWindow;
+};
+
