Fix livelink hub CIS issue on Mac

diff --git a/Engine/Plugins/Animation/LiveLink/Source/LiveLinkHub/Private/LiveLinkHubApplication.h b/Engine/Plugins/Animation/LiveLink/Source/LiveLinkHub/Private/LiveLinkHubApplication.h
--- a/Engine/Plugins/Animation/LiveLink/Source/LiveLinkHub/Private/LiveLinkHubApplication.h
+++ b/Engine/Plugins/Animation/LiveLink/Source/LiveLinkHub/Private/LiveLinkHubApplication.h
@@ -12,6 +12,9 @@
 #include "Runtime/Launch/Resources/Version.h"
 #include "UObject/UObjectGlobals.h"
 
+#if PLATFORM_MAC
+#include "HAL/PlatformApplicationMisc.h"
+#endif
 
 DEFINE_LOG_CATEGORY_STATIC(LogLiveLinkHubApplication, Log, All);
 
