Fix LiveLinkHub duplicated loctext macro entry

diff --git a/Engine/Plugins/Animation/LiveLink/Source/LiveLinkHub/Private/UI/Widgets/SLiveLinkHubMainTabView.cpp b/Engine/Plugins/Animation/LiveLink/Source/LiveLinkHub/Private/UI/Widgets/SLiveLinkHubMainTabView.cpp
--- a/Engine/Plugins/Animation/LiveLink/Source/LiveLinkHub/Private/UI/Widgets/SLiveLinkHubMainTabView.cpp
+++ b/Engine/Plugins/Animation/LiveLink/Source/LiveLinkHub/Private/UI/Widgets/SLiveLinkHubMainTabView.cpp
@@ -214,7 +214,7 @@ TSharedRef<SDockTab> SLiveLinkHubMainTabView::SpawnClientsTab(const FSpawnTabArg
 		.TabRole(PanelTab)
 		[
 			SNew(STextBlock)
-			.Text(LOCTEXT("TemporaryText", "Client list here"))
+			.Text(LOCTEXT("TemporaryClientText", "Client list here"))
 		];
 }
 
@@ -225,7 +225,7 @@ TSharedRef<SDockTab> SLiveLinkHubMainTabView::SpawnClientDetailsTab(const FSpawn
 		.TabRole(PanelTab)
 		[
 			SNew(STextBlock)
-			.Text(LOCTEXT("TemporaryText", "Client details here"))
+			.Text(LOCTEXT("TemporaryClientDetailsText", "Client details here"))
 		];
 }
 
