LiveLink: Added ability to genlock the engine to a live link subject.

This is useful when live link is the primary driver of the animation (such as mocap), but needs to be separate from the timecode source.

diff --git a/Engine/Plugins/Animation/LiveLink/Source/LiveLink/Private/LiveLinkClient.cpp b/Engine/Plugins/Animation/LiveLink/Source/LiveLink/Private/LiveLinkClient.cpp
--- a/Engine/Plugins/Animation/LiveLink/Source/LiveLink/Private/LiveLinkClient.cpp
+++ b/Engine/Plugins/Animation/LiveLink/Source/LiveLink/Private/LiveLinkClient.cpp
@@ -695,26 +695,36 @@ void FLiveLinkClient::PushSubjectFrameData_AnyThread(const FLiveLinkSubjectKey&
 	FPendingSubjectFrame SubjectFrame{ InSubjectKey, MoveTemp(InFrameData) };
 	const int32 MaxNumBufferToCached = CVarMaxNewFrameDataPerUpdate.GetValueOnAnyThread();
 	bool bLogError = false;
+
+	bool bCanPushFrame = true;
 	{
 		FScopeLock Lock(&CollectionAccessCriticalSection);
+
 		if (SubjectFrameToPush.Num() > MaxNumBufferToCached) // Something is wrong somewhere. Warn the user and discard the new Frame Data.
 		{
 			bLogError = true;
 			SubjectFrameToPush.RemoveAt(0, SubjectFrameToPush.Num() - MaxNumBufferToCached, false);
+			bCanPushFrame = false;
 		}
-		else
-		{
+	}
 
+	if (bCanPushFrame)
+	{
+		{
+			FScopeLock BroadcastLock(&SubjectFrameReceivedHandleseCriticalSection);
+			if (const FSubjectFramesReceivedHandles* Handles = SubjectFrameReceivedHandles.Find(InSubjectKey))
 			{
-				FScopeLock BroadcastLock(&SubjectFrameReceivedHandleseCriticalSection);
-				if (const FSubjectFramesReceivedHandles* Handles = SubjectFrameReceivedHandles.Find(InSubjectKey))
-				{
-					Handles->OnFrameDataReceived.Broadcast(SubjectFrame.FrameData);
-				}
+				Handles->OnFrameDataReceived.Broadcast(SubjectFrame.FrameData);
 			}
-			
-			SubjectFrameToPush.Add(MoveTemp(SubjectFrame));
 		}
+			
+		// Since the lock was released between setting bCanPushFrame and adding to the array, it is possible that
+		// we exceed MaxNumBufferToCached. But this should be rare and also harmless.
+		// The lock was released so that OnFrameDataReceived doesn't need to be called with the lock on,
+		// which can hang the game thread when it calls EvaluateFrame if the broadcast takes longer than usual.
+
+		FScopeLock Lock(&CollectionAccessCriticalSection);
+		SubjectFrameToPush.Add(MoveTemp(SubjectFrame));
 	}
 
 	if (bLogError)
diff --git a/Engine/Plugins/Animation/LiveLink/Source/LiveLink/Private/LiveLinkCustomTimeStep.cpp b/Engine/Plugins/Animation/LiveLink/Source/LiveLink/Private/LiveLinkCustomTimeStep.cpp
new file mode 100644
--- /dev/null
+++ b/Engine/Plugins/Animation/LiveLink/Source/LiveLink/Private/LiveLinkCustomTimeStep.cpp
@@ -0,0 +1,296 @@
+// Copyright Epic Games, Inc. All Rights Reserved.
+
+#include "LiveLinkCustomTimeStep.h"
+
+#include "Features/IModularFeatures.h"
+#include "ILiveLinkClient.h"
+#include "Misc/App.h"
+
+
+
+bool ULiveLinkCustomTimeStep::Initialize(UEngine* InEngine)
+{
+	State = ECustomTimeStepSynchronizationState::Synchronizing;
+	EventLiveLink->Reset();
+
+	IModularFeatures& ModularFeatures = IModularFeatures::Get();
+	ModularFeatures.OnModularFeatureRegistered().AddUObject(this, &ULiveLinkCustomTimeStep::OnLiveLinkClientRegistered);
+	ModularFeatures.OnModularFeatureUnregistered().AddUObject(this, &ULiveLinkCustomTimeStep::OnLiveLinkClientUnregistered);
+
+	InitLiveLinkClient();
+
+	return true;
+}
+
+void ULiveLinkCustomTimeStep::Shutdown(UEngine* InEngine)
+{
+	UninitLiveLinkClient();
+}
+
+bool ULiveLinkCustomTimeStep::UpdateTimeStep(UEngine* InEngine)
+{
+	if (State != ECustomTimeStepSynchronizationState::Synchronized)
+	{
+		EventWaitForSync->Trigger(); // Prevents start/restart quasi-deadlock.
+
+		return true; // means that the Engine's TimeStep should be performed.
+	}
+
+	// Update anythread variable shadows
+	{
+		bLockStepModeAnyThread = bLockStepMode;
+		FrameRateDividerAnyThread = FrameRateDivider;
+		TimeoutInSecondsAnyThread = TimeoutInSeconds;
+	}
+
+	UpdateApplicationLastTime(); // Copies "CurrentTime" (used during the previous frame) in "LastTime"
+	
+	const double StartPlatformTime = FPlatformTime::Seconds();
+
+	if (!WaitForSync())
+	{
+		State = ECustomTimeStepSynchronizationState::Synchronizing;
+		return true;
+	}
+
+	const double EndPlatformTime = FPlatformTime::Seconds();
+
+	const double ElapsedTimeWaitingForSync = EndPlatformTime - StartPlatformTime;
+
+	// Subtract the multiple of the sync counts that corresponds to the desired divided live link data frame rate
+	// For example if there are 3 sync counts and the frame rate divider is 2, then we'd want to use up 2 of the 3
+	// sync counts and leave the 3rd one for the next cycle.
+
+	LastSyncCountDelta = SyncCount;
+
+	if (FrameRateDivider > 1)
+	{
+		LastSyncCountDelta -= (LastSyncCountDelta % FrameRateDivider);
+	}
+
+	SyncCount -= LastSyncCountDelta;
+
+	// We are using real elapsed time for the time before sync, which will ultimately be used to calculate idle time. 
+	// These values won't make sense if the frame rate settings are not correct.
+	const double TimeAfterSync = FApp::GetLastTime() + GetLastSyncCountDelta() * GetSyncRate().AsInterval();
+	const double TimeBeforeSync = TimeAfterSync - ElapsedTimeWaitingForSync;
+
+	UpdateAppTimes(TimeBeforeSync, TimeAfterSync);
+
+	return false; // false means that the Engine's TimeStep should NOT be performed.
+}
+
+ECustomTimeStepSynchronizationState ULiveLinkCustomTimeStep::GetSynchronizationState() const
+{
+	return State;
+}
+
+FFrameRate ULiveLinkCustomTimeStep::GetFixedFrameRate() const
+{
+	FFrameRate Rate = LiveLinkDataRate;
+
+	if (FrameRateDivider > 1)
+	{
+		Rate.Denominator *= FrameRateDivider;
+	}
+
+	return Rate;
+}
+
+FFrameRate ULiveLinkCustomTimeStep::GetSyncRate() const
+{
+	return LiveLinkDataRate;
+}
+
+uint32 ULiveLinkCustomTimeStep::GetLastSyncCountDelta() const
+{
+	return LastSyncCountDelta;
+}
+
+bool ULiveLinkCustomTimeStep::IsLastSyncDataValid() const
+{
+	return true;
+}
+
+bool ULiveLinkCustomTimeStep::WaitForSync()
+{	
+	EventWaitForSync->Trigger();
+	return EventLiveLink->Wait(FTimespan(TimeoutInSeconds * ETimespan::TicksPerSecond));
+}
+
+
+void ULiveLinkCustomTimeStep::BeginDestroy()
+{
+	UninitLiveLinkClient();
+	Super::BeginDestroy();
+}
+
+#if WITH_EDITOR
+void ULiveLinkCustomTimeStep::PostEditChangeProperty(FPropertyChangedEvent& PropertyChangedEvent)
+{
+	Super::PostEditChangeProperty(PropertyChangedEvent);
+
+	if (PropertyChangedEvent.Property && PropertyChangedEvent.Property->GetFName() == GET_MEMBER_NAME_CHECKED(ULiveLinkCustomTimeStep, SubjectKey))
+	{
+		// If was already registered
+		if (RegisterForFrameDataReceivedHandle.IsValid() && RegisteredSubjectKey != SubjectKey)
+		{
+			UnregisterLiveLinkSubject();
+			RegisterLiveLinkSubject();
+		}
+		// If was waiting for the subject to be added to the client
+		else if (LiveLinkClient && LiveLinkClient->GetSubjectSettings(SubjectKey))
+		{
+			OnLiveLinkSubjectAdded(SubjectKey);
+		}
+	}
+}
+#endif
+
+void ULiveLinkCustomTimeStep::InitLiveLinkClient()
+{
+	IModularFeatures& ModularFeatures = IModularFeatures::Get();
+
+	if (ModularFeatures.IsModularFeatureAvailable(ILiveLinkClient::ModularFeatureName))
+	{
+		LiveLinkClient = &IModularFeatures::Get().GetModularFeature<ILiveLinkClient>(ILiveLinkClient::ModularFeatureName);
+
+		LiveLinkClient->OnLiveLinkSubjectAdded().AddUObject(this, &ULiveLinkCustomTimeStep::OnLiveLinkSubjectAdded);
+		LiveLinkClient->OnLiveLinkSubjectRemoved().AddUObject(this, &ULiveLinkCustomTimeStep::OnLiveLinkSubjectRemoved);
+
+		// if the subject already exist
+		if (LiveLinkClient->GetSubjectSettings(SubjectKey))
+		{
+			OnLiveLinkSubjectAdded(SubjectKey);
+		}
+	}
+}
+
+
+void ULiveLinkCustomTimeStep::UninitLiveLinkClient()
+{
+	if (LiveLinkClient)
+	{
+		UnregisterLiveLinkSubject();
+
+		LiveLinkClient->OnLiveLinkSubjectAdded().RemoveAll(this);
+		LiveLinkClient->OnLiveLinkSubjectRemoved().RemoveAll(this);
+		LiveLinkClient = nullptr;
+		State = ECustomTimeStepSynchronizationState::Closed;
+	}
+}
+
+
+void ULiveLinkCustomTimeStep::RegisterLiveLinkSubject()
+{
+	RegisteredSubjectKey = SubjectKey;
+
+	FDelegateHandle Tmp;
+	LiveLinkClient->RegisterForFrameDataReceived(
+		RegisteredSubjectKey
+		, FOnLiveLinkSubjectStaticDataReceived::FDelegate()
+		, FOnLiveLinkSubjectFrameDataReceived::FDelegate::CreateUObject(this, &ULiveLinkCustomTimeStep::OnLiveLinkFrameDataReceived_AnyThread)
+		, Tmp
+		, RegisterForFrameDataReceivedHandle);
+}
+
+
+void ULiveLinkCustomTimeStep::UnregisterLiveLinkSubject()
+{
+	if (RegisterForFrameDataReceivedHandle.IsValid())
+	{
+		LiveLinkClient->UnregisterForFrameDataReceived(RegisteredSubjectKey, FDelegateHandle(), RegisterForFrameDataReceivedHandle);
+		RegisterForFrameDataReceivedHandle.Reset();
+	}
+
+	RegisteredSubjectKey = FLiveLinkSubjectKey();
+}
+
+
+void ULiveLinkCustomTimeStep::OnLiveLinkClientRegistered(const FName& Type, class IModularFeature* ModularFeature)
+{
+	if (Type == ILiveLinkClient::ModularFeatureName && !LiveLinkClient)
+	{
+		InitLiveLinkClient();
+	}
+}
+
+
+void ULiveLinkCustomTimeStep::OnLiveLinkClientUnregistered(const FName& Type, class IModularFeature* ModularFeature)
+{
+	if (Type == ILiveLinkClient::ModularFeatureName && ModularFeature == LiveLinkClient)
+	{
+		UninitLiveLinkClient();
+		InitLiveLinkClient();
+	}
+}
+
+
+void ULiveLinkCustomTimeStep::OnLiveLinkSubjectAdded(FLiveLinkSubjectKey InSubjectKey)
+{
+	if (InSubjectKey == SubjectKey)
+	{
+		RegisterLiveLinkSubject();
+	}
+}
+
+
+void ULiveLinkCustomTimeStep::OnLiveLinkSubjectRemoved(FLiveLinkSubjectKey InSubjectKey)
+{
+	if (InSubjectKey == RegisteredSubjectKey)
+	{
+		State = ECustomTimeStepSynchronizationState::Error;
+		UnregisterLiveLinkSubject();
+	}
+}
+
+
+void ULiveLinkCustomTimeStep::OnLiveLinkFrameDataReceived_AnyThread(const FLiveLinkFrameDataStruct& InFrameData)
+{
+	// Increment sync counter based on frame id. For example, if frame id is incremented by 2, it means that a frame
+	// was skipped and we must report 2 sync signals, even though only one was received. The purpose of this is to
+	// reflect source delta times in the engine delta times.
+
+	const FLiveLinkFrameIdentifier FrameId = InFrameData.GetBaseData()->FrameId;
+
+	// Receiving data from the desired live link source automatically means that we are in synchronized state.
+	if (State != ECustomTimeStepSynchronizationState::Synchronized)
+	{
+		State = ECustomTimeStepSynchronizationState::Synchronized;
+
+		// Ensure our initial delta frame id is 1.
+		LastFrameId = FrameId - 1;
+	}
+
+	// Loop around cannot be correctly handled unless we know the lower and upper bounds of the frame id.
+	// For now, consider them as single increments even though we may miss skipped frames.
+	// If FrameId is not changing, then we assume that the source does not support this, and we treat this
+	// as a single sync increment as we otherwise lack the means of detecting skipped source frames.
+	if (LastFrameId >= FrameId)
+	{
+		SyncCount++;
+	}
+	else
+	{
+		SyncCount += FrameId - LastFrameId;
+	}
+
+	// Note: We don't expect OnLiveLinkFrameDataReceived_AnyThread to be called for the same live link source concurrently, 
+	// otherwise we would need to avoid race conditions here when updating and comparing with FrameId.
+	LastFrameId = FrameId;
+
+	// Trigger the event to release WaitForSync.
+	// If the engine cannot keep up, it means that SyncCount will be greater than FrameRateDividerAnyThread and delta time will be multiplied accordingly
+	if (SyncCount >= FrameRateDividerAnyThread)
+	{
+		EventLiveLink->Trigger();
+
+		// In lockstep mode, the engine can stall the live link producer until it has had a chance to process its current data.
+		// We can't block if this function is being called in the game thread because it would deadlock.
+		if (bLockStepModeAnyThread && !IsInGameThread())
+		{
+			EventWaitForSync->Wait(FTimespan((TimeoutInSecondsAnyThread / 2) * ETimespan::TicksPerSecond));
+		}
+	}
+}
+
diff --git a/Engine/Plugins/Animation/LiveLink/Source/LiveLink/Public/LiveLinkCustomTimeStep.h b/Engine/Plugins/Animation/LiveLink/Source/LiveLink/Public/LiveLinkCustomTimeStep.h
new file mode 100644
--- /dev/null
+++ b/Engine/Plugins/Animation/LiveLink/Source/LiveLink/Public/LiveLinkCustomTimeStep.h
@@ -0,0 +1,154 @@
+// Copyright Epic Games, Inc. All Rights Reserved.
+
+#pragma once
+
+#include "Engine/EngineCustomTimeStep.h"
+#include "GenlockedCustomTimeStep.h"
+#include "HAL/Event.h"
+#include "LiveLinkTypes.h"
+#include "Misc/FrameRate.h"
+
+#include <atomic>
+
+#include "LiveLinkCustomTimeStep.generated.h"
+
+class UEngine;
+class UObject;
+class ILiveLinkClient;
+
+/**
+ * Control the Engine TimeStep via a live link source
+ * 
+ * Philosophy:
+ * 
+ *   * Quantized time steps based on live link expected data rate.
+ *   * Made for Live Link sources can receive data asynchronously, and therefore trigger the waiting game thread.
+ * 
+ *   * FApp::GetDeltaTime 
+ *       - Forced to a multiple of the desired FrameTime.
+ *       - This multiple will depend on Frame Id increment and user settings.
+ * 
+ *   * FApp::GetCurrentTime
+ *       - Incremented in multiples of the desired FrameTime.
+ * 
+ */
+UCLASS(Blueprintable, editinlinenew, MinimalAPI)
+class ULiveLinkCustomTimeStep : public UGenlockedCustomTimeStep
+{
+	GENERATED_BODY()
+
+public:
+	//~ UFixedFrameRateCustomTimeStep interface
+	LIVELINK_API virtual bool Initialize(UEngine* InEngine) override;
+	LIVELINK_API virtual void Shutdown(UEngine* InEngine) override;
+	LIVELINK_API virtual bool UpdateTimeStep(UEngine* InEngine) override;
+	LIVELINK_API virtual ECustomTimeStepSynchronizationState GetSynchronizationState() const override;
+	LIVELINK_API virtual FFrameRate GetFixedFrameRate() const override;
+	LIVELINK_API virtual FFrameRate GetSyncRate() const override;
+
+	//~ UGenlockedCustomTimeStep interface
+	LIVELINK_API virtual uint32 GetLastSyncCountDelta() const override;
+	LIVELINK_API virtual bool IsLastSyncDataValid() const override;
+	LIVELINK_API virtual bool WaitForSync() override;
+
+	//~ Begin UObject Interface
+	virtual void BeginDestroy() override;
+#if WITH_EDITOR
+	virtual void PostEditChangeProperty(FPropertyChangedEvent& PropertyChangedEvent) override;
+#endif
+	//~End UObject Interface
+
+private:
+
+	/** Initializes the Live Link client (callbacks, etc) */
+	void InitLiveLinkClient();
+
+	/** Uninitializes the Live Link client (callbacks, etc) */
+	void UninitLiveLinkClient();
+
+	/** Registers the selected subject for data callbacks */
+	void RegisterLiveLinkSubject();
+
+	/** Unregisters the data callbacks related to the live link subject */
+	void UnregisterLiveLinkSubject();
+
+	/** Called when a Live Link modular features is registered */
+	void OnLiveLinkClientRegistered(const FName& Type, class IModularFeature* ModularFeature);
+
+	/** Called when a Live Link modular feature is unregistered */
+	void OnLiveLinkClientUnregistered(const FName& Type, class IModularFeature* ModularFeature);
+
+	/** Called when a new live link subject is added, so that we can check if is the one we're listening for */
+	void OnLiveLinkSubjectAdded(FLiveLinkSubjectKey SubjectKey);
+
+	/** Called when a live link subject is removed. It may be the one we were listening to */
+	void OnLiveLinkSubjectRemoved(FLiveLinkSubjectKey SubjectKey);
+
+	/** Called when our live link subject has received data. This is equivalent to a genlock sync signal */
+	void OnLiveLinkFrameDataReceived_AnyThread(const FLiveLinkFrameDataStruct& FrameData);
+
+public:
+
+	/** Expected Live Link data rate. If real rate differs, then delta times will contract/expand with respect to real time */
+	UPROPERTY(EditAnywhere, Category = Timing)
+	FFrameRate LiveLinkDataRate = FFrameRate(60, 1);
+
+	/** The specific subject that we listen to. */
+	UPROPERTY(EditAnywhere, Category = Timing)
+	FLiveLinkSubjectKey SubjectKey;
+
+	/** 
+	 * Lockstep mode means that we only allow FrameRateDivider data frames of the selected subject per engine loop. 
+	 * The idea here is to process all packets and avoid Live Link evaluation to skip frames when the engine hitches, 
+	 * and the live link transport will serve as implicit FIFO buffer. However if the Engine cannot keep up with the data rate,
+	 * a large delay will be introduced and the transport buffer will eventually start dropping data.
+	 */
+	UPROPERTY(EditAnywhere, Category = Timing)
+	bool bLockStepMode = true;
+
+	/** Allows genlock to period to be a multiple of the live link data period. For example a value of 2 will run at half the live link data rate */
+	UPROPERTY(EditAnywhere, Category = Timing, meta = (ClampMin = 1, ClampMax = 256, UIMin = 1, UIMax = 256))
+	uint32 FrameRateDivider = 1;
+
+	/** Determines how long it should wait for live link data before deciding that it is not in synchronized state anymore */
+	UPROPERTY(EditAnywhere, Category = Advanced, AdvancedDisplay)
+	double TimeoutInSeconds = 1.0;
+
+private:
+
+	/** This is how we access live link objects. It is obtained using modular features. */
+	ILiveLinkClient* LiveLinkClient;
+
+	/** Identifies the Live Link subject that we are genlocking to. */
+	FLiveLinkSubjectKey RegisteredSubjectKey;
+
+	/** Gets called when the subject gets new data. This is the basis for the genlock. */
+	FDelegateHandle RegisterForFrameDataReceivedHandle;
+
+	/** Keeps track of the current state. For the most part, if we are getting data from the subject, it is synchronized. */
+	std::atomic<ECustomTimeStepSynchronizationState> State = ECustomTimeStepSynchronizationState::Closed;
+
+	/** Helps trigger the wait for sync function from an asynchronous thread */
+	FEventRef EventLiveLink{ EEventMode::AutoReset };
+
+	/** Used in lockstep mode, where we only allow one live link subject data frame per engine tick */
+	FEventRef EventWaitForSync{ EEventMode::AutoReset };
+
+	/** Keeps track of the last frame id received for a given subject, and used to calcualte delta time based on source data rate. */
+	std::atomic<FLiveLinkFrameIdentifier> LastFrameId = 0;
+
+	/** Incremented based on frame id when live link data is received. Reset when internally read. */
+	std::atomic<uint32> SyncCount = 0;
+
+	/** Cache of the last sync count delta, used for functions in the game thread that need it. */
+	uint32 LastSyncCountDelta = 0;
+
+	/** Shadow of bLockStepMode property */
+	std::atomic<bool> bLockStepModeAnyThread = true;
+
+	/** Shadow of FrameRateDivider property */
+	std::atomic<uint32> FrameRateDividerAnyThread = 1;
+
+	/** Shadow of TimeoutInSeconds property */
+	std::atomic<double> TimeoutInSecondsAnyThread = 1.0;
+};
