Fix CIS caused by missing "typename" in LiveLinkHub

diff --git a/Engine/Plugins/Animation/LiveLink/Source/LiveLinkHub/Private/Clients/LiveLinkHubProvider.h b/Engine/Plugins/Animation/LiveLink/Source/LiveLinkHub/Private/Clients/LiveLinkHubProvider.h
--- a/Engine/Plugins/Animation/LiveLink/Source/LiveLinkHub/Private/Clients/LiveLinkHubProvider.h
+++ b/Engine/Plugins/Animation/LiveLink/Source/LiveLinkHub/Private/Clients/LiveLinkHubProvider.h
@@ -33,7 +33,7 @@ public:
 	 * This is useful to receive some messages on AnyThread and delegate others on the game thread (ie. for methods that will trigger UI updates which need to happen on game thread. )
 	 */
 	template <typename MessageType>
-	TSharedRef<TGameThreadMessageHandler<MessageType, FLiveLinkHubProvider>> MakeHandler(TGameThreadMessageHandler<MessageType, FLiveLinkHubProvider>::FuncType Func)
+	TSharedRef<TGameThreadMessageHandler<MessageType, FLiveLinkHubProvider>> MakeHandler(typename TGameThreadMessageHandler<MessageType, FLiveLinkHubProvider>::FuncType Func)
 	{
 		return MakeShared<TGameThreadMessageHandler<MessageType, FLiveLinkHubProvider>>(this, Func);
 	}
