Fix CIS caused by uninitialized enum in LiveLinkHubClientInfo

diff --git a/Engine/Plugins/Animation/LiveLink/Source/LiveLinkHub/Private/Clients/LiveLinkHubUEClientInfo.h b/Engine/Plugins/Animation/LiveLink/Source/LiveLinkHub/Private/Clients/LiveLinkHubUEClientInfo.h
--- a/Engine/Plugins/Animation/LiveLink/Source/LiveLinkHub/Private/Clients/LiveLinkHubUEClientInfo.h
+++ b/Engine/Plugins/Animation/LiveLink/Source/LiveLinkHub/Private/Clients/LiveLinkHubUEClientInfo.h
@@ -36,7 +36,7 @@ struct FLiveLinkHubUEClientInfo
 	
 	/** Status of the client, ie. is it actively doing a take record at the moment? */
 	UPROPERTY(VisibleAnywhere, Category = "LiveLink Client")
-	ELiveLinkClientStatus Status;
+	ELiveLinkClientStatus Status = ELiveLinkClientStatus::Disconnected;
 	
 	UPROPERTY(VisibleAnywhere, Category = "LiveLink Client", DisplayName = "IP Address")
 	FString IPAddress;
diff --git a/Engine/Plugins/Animation/LiveLink/Source/LiveLinkHubMessaging/Public/LiveLinkHubMessages.h b/Engine/Plugins/Animation/LiveLink/Source/LiveLinkHubMessaging/Public/LiveLinkHubMessages.h
--- a/Engine/Plugins/Animation/LiveLink/Source/LiveLinkHubMessaging/Public/LiveLinkHubMessages.h
+++ b/Engine/Plugins/Animation/LiveLink/Source/LiveLinkHubMessaging/Public/LiveLinkHubMessages.h
@@ -25,6 +25,7 @@ UENUM()
 enum class ELiveLinkClientStatus
 {
 	Connected, /** Default state of a UE client. */
+	Disconnected, /** Client is not connected to the hub. */
 	Recording  /** UE is currently doing a take record. */
 };
 
