Fix CIS caused by uninitialized struct variable

diff --git a/Engine/Plugins/Animation/LiveLink/Source/LiveLinkHubMessaging/Public/LiveLinkHubMessages.h b/Engine/Plugins/Animation/LiveLink/Source/LiveLinkHubMessaging/Public/LiveLinkHubMessages.h
--- a/Engine/Plugins/Animation/LiveLink/Source/LiveLinkHubMessaging/Public/LiveLinkHubMessages.h
+++ b/Engine/Plugins/Animation/LiveLink/Source/LiveLinkHubMessaging/Public/LiveLinkHubMessages.h
@@ -41,7 +41,7 @@ struct LIVELINKHUBMESSAGING_API FLiveLinkClientInfoMessage
 
 	/** Status of the client, ie. is it actively doing a take record at the moment? */
 	UPROPERTY()
-	ELiveLinkClientStatus Status;
+	ELiveLinkClientStatus Status = ELiveLinkClientStatus::Disconnected;
 
 	/** Name of the host of the UE client */
 	UPROPERTY()
