LiveLinkHub: Add status bar with output log and ability to set console variables.

diff --git a/Engine/Plugins/Animation/LiveLink/Source/LiveLinkHub/LiveLinkHub.Build.cs b/Engine/Plugins/Animation/LiveLink/Source/LiveLinkHub/LiveLinkHub.Build.cs
--- a/Engine/Plugins/Animation/LiveLink/Source/LiveLinkHub/LiveLinkHub.Build.cs
+++ b/Engine/Plugins/Animation/LiveLink/Source/LiveLinkHub/LiveLinkHub.Build.cs
@@ -35,6 +35,7 @@ namespace UnrealBuildTool.Rules
 				"TimeManagement",
 				"ToolWidgets",
 				"UnrealEd",
+				"OutputLog",
 			});
 		}
 	}
diff --git a/Engine/Plugins/Animation/LiveLink/Source/LiveLinkHub/Private/LiveLinkHub.cpp b/Engine/Plugins/Animation/LiveLink/Source/LiveLinkHub/Private/LiveLinkHub.cpp
--- a/Engine/Plugins/Animation/LiveLink/Source/LiveLinkHub/Private/LiveLinkHub.cpp
+++ b/Engine/Plugins/Animation/LiveLink/Source/LiveLinkHub/Private/LiveLinkHub.cpp
@@ -22,6 +22,9 @@ void FLiveLinkHub::Initialize()
 	LiveLinkProvider = MakeShared<FLiveLinkHubProvider>();
 	LiveLinkHubClient = MakeShared<FLiveLinkHubClient>(AsShared());
 
+	CommandExecutor = MakeUnique<FConsoleCommandExecutor>();
+	IModularFeatures::Get().RegisterModularFeature(IConsoleCommandExecutor::ModularFeatureName(), CommandExecutor.Get());
+	
 	IModularFeatures::Get().RegisterModularFeature(ILiveLinkClient::ModularFeatureName, LiveLinkHubClient.Get());
 
 	RecordingController = MakeShared<FLiveLinkHubRecordingController>();
diff --git a/Engine/Plugins/Animation/LiveLink/Source/LiveLinkHub/Private/LiveLinkHub.h b/Engine/Plugins/Animation/LiveLink/Source/LiveLinkHub/Private/LiveLinkHub.h
--- a/Engine/Plugins/Animation/LiveLink/Source/LiveLinkHub/Private/LiveLinkHub.h
+++ b/Engine/Plugins/Animation/LiveLink/Source/LiveLinkHub/Private/LiveLinkHub.h
@@ -4,6 +4,8 @@
 
 #include "LiveLinkTypes.h"
 #include "Templates/SubclassOf.h"
+#include "Templates/UniquePtr.h"
+#include "Toolkits/FConsoleCommandExecutor.h"
 
 class FLiveLinkHubClient;
 class FLiveLinkHubClientsController;
@@ -85,6 +87,8 @@ private:
 	TSharedPtr<FLiveLinkHubClient> LiveLinkHubClient;
 	/** LiveLinkProvider used to transfer data to connected UE clients. */
 	TSharedPtr<FLiveLinkHubProvider> LiveLinkProvider;
+	/** Handles execution of commands */
+	TUniquePtr<FConsoleCommandExecutor> CommandExecutor;
 
 	friend class FLiveLinkHubModule;
 };
diff --git a/Engine/Plugins/Animation/LiveLink/Source/LiveLinkHub/Private/UI/Widgets/SLiveLinkHubMainTabView.cpp b/Engine/Plugins/Animation/LiveLink/Source/LiveLinkHub/Private/UI/Widgets/SLiveLinkHubMainTabView.cpp
--- a/Engine/Plugins/Animation/LiveLink/Source/LiveLinkHub/Private/UI/Widgets/SLiveLinkHubMainTabView.cpp
+++ b/Engine/Plugins/Animation/LiveLink/Source/LiveLinkHub/Private/UI/Widgets/SLiveLinkHubMainTabView.cpp
@@ -22,7 +22,6 @@
 
 #define LOCTEXT_NAMESPACE "LiveLinkHub.SLiveLinkHubMainTabView"
 
-const FName SLiveLinkHubMainTabView::StatusBarTabId("StatusBarTabId");
 const FName SLiveLinkHubMainTabView::SourcesTabId("SourcesTabId");
 const FName SLiveLinkHubMainTabView::SourceDetailsTabId("SourceDetailsTabId");
 const FName SLiveLinkHubMainTabView::SubjectsTabId("SubjectsTabId");
diff --git a/Engine/Plugins/Animation/LiveLink/Source/LiveLinkHub/Private/UI/Widgets/SLiveLinkHubMainTabView.h b/Engine/Plugins/Animation/LiveLink/Source/LiveLinkHub/Private/UI/Widgets/SLiveLinkHubMainTabView.h
--- a/Engine/Plugins/Animation/LiveLink/Source/LiveLinkHub/Private/UI/Widgets/SLiveLinkHubMainTabView.h
+++ b/Engine/Plugins/Animation/LiveLink/Source/LiveLinkHub/Private/UI/Widgets/SLiveLinkHubMainTabView.h
@@ -12,7 +12,6 @@ class SLiveLinkHubMainTabView : public SLiveLinkHubTabViewWithManagerBase
 {
 public:
 	//~ Tab IDs
-	static const FName StatusBarTabId;
 	static const FName SourcesTabId;
 	static const FName SourceDetailsTabId;
 	static const FName SubjectsTabId;
diff --git a/Engine/Plugins/Animation/LiveLink/Source/LiveLinkHub/Private/UI/Widgets/SLiveLinkHubStatusBar.cpp b/Engine/Plugins/Animation/LiveLink/Source/LiveLinkHub/Private/UI/Widgets/SLiveLinkHubStatusBar.cpp
new file mode 100644
--- /dev/null
+++ b/Engine/Plugins/Animation/LiveLink/Source/LiveLinkHub/Private/UI/Widgets/SLiveLinkHubStatusBar.cpp
@@ -0,0 +1,146 @@
+// Copyright Epic Games, Inc. All Rights Reserved.
+
+#include "SLiveLinkHubStatusBar.h"
+
+#include "OutputLogCreationParams.h"
+
+#include "Framework/Application/SlateApplication.h"
+#include "OutputLogModule.h"
+#include "SWidgetDrawer.h"
+#include "WidgetDrawerConfig.h"
+#include "Widgets/Layout/SBox.h"
+#include "Widgets/Layout/SBorder.h"
+
+#define LOCTEXT_NAMESPACE "LiveLinkHubStatusBar"
+
+namespace UE::LiveLinkHub::Private
+{
+	class FStatusBarSingleton
+	{
+		TSharedPtr<SWidget> StatusBarOutputLog;
+		TArray<TWeakPtr<SWidgetDrawer>> StatusBars;
+		
+		TSharedRef<SWidget> OnGetOutputLog()
+		{
+			if (!StatusBarOutputLog)
+			{
+				FOutputLogCreationParams Params;
+				Params.bCreateDockInLayoutButton = true;
+				Params.SettingsMenuCreationFlags = EOutputLogSettingsMenuFlags::SkipClearOnPie
+					| EOutputLogSettingsMenuFlags::SkipOpenSourceButton
+					| EOutputLogSettingsMenuFlags::SkipEnableWordWrapping; // Checkbox relies on saving an editor config file and does not work correctly
+				StatusBarOutputLog = FOutputLogModule::Get().MakeOutputLogWidget(Params);
+			}
+
+			return StatusBarOutputLog.ToSharedRef();
+		}
+		
+		void OnOutputLogOpened(FName StatusBarWithDrawerName)
+		{
+			// Dismiss all other open drawers - StatusBarOutputLog is shared and shouldn't be in the layout twice
+			for (const TWeakPtr<SWidgetDrawer>& WidgetDrawer : StatusBars)
+			{
+				if (WidgetDrawer.IsValid())
+				{
+					TSharedPtr<SWidgetDrawer> PinnedDrawer = WidgetDrawer.Pin();
+					if (StatusBarWithDrawerName != PinnedDrawer->GetDrawerName() || PinnedDrawer->IsAnyOtherDrawerOpened(OutputLogId))
+					{
+						PinnedDrawer->CloseDrawerImmediately();
+					}
+				}
+			}
+			
+			FOutputLogModule::Get().FocusOutputLogConsoleBox(StatusBarOutputLog.ToSharedRef());
+		}
+		
+		void OnOutputLogDismissed(const TSharedPtr<SWidget>& NewlyFocusedWidget)
+		{}
+
+		void PreShutdownSlate()
+		{
+			StatusBarOutputLog.Reset();
+		}
+
+	public:
+
+		void Init(TSharedRef<SWidgetDrawer> WidgetDrawer, FWidgetDrawerConfig& OutputLogDrawer)
+		{
+			if (!FSlateApplication::Get().OnPreShutdown().IsBoundToObject(this))
+			{
+				// Destroying StatusBarOutputLog in ~FStatusBarSingleton is too late: it causes a crash
+				FSlateApplication::Get().OnPreShutdown().AddRaw(this, &FStatusBarSingleton::PreShutdownSlate);
+			}
+
+			const bool bIsDrawerNameUnique = !StatusBars.ContainsByPredicate([&WidgetDrawer](TWeakPtr<SWidgetDrawer> WeakDrawer)
+			{
+				return ensure(WeakDrawer.IsValid())
+					&& WeakDrawer.Pin()->GetDrawerName() == WidgetDrawer->GetDrawerName();
+			});
+			checkf(bIsDrawerNameUnique, TEXT("Every widget drawer is expected to have an unique ID"));
+			
+			StatusBars.Add(MoveTemp(WidgetDrawer));
+			
+			OutputLogDrawer.GetDrawerContentDelegate.BindRaw(this, &FStatusBarSingleton::OnGetOutputLog);
+			OutputLogDrawer.OnDrawerOpenedDelegate.BindRaw(this, &FStatusBarSingleton::OnOutputLogOpened);
+			OutputLogDrawer.OnDrawerDismissedDelegate.BindRaw(this, &FStatusBarSingleton::OnOutputLogDismissed);
+		}
+
+		void Remove(TSharedRef<SWidgetDrawer> WidgetDrawer)
+		{
+			StatusBars.RemoveSingle(WidgetDrawer);
+		}
+		
+	} GStatusBarManager;
+}
+
+SLiveLinkHubStatusBar::~SLiveLinkHubStatusBar()
+{
+	UE::LiveLinkHub::Private::GStatusBarManager.Remove(WidgetDrawer.ToSharedRef());
+}
+
+void SLiveLinkHubStatusBar::Construct(const FArguments& InArgs, FName StatusBarId)
+{
+	ChildSlot
+	[
+		SNew(SBox)
+		.HeightOverride(FAppStyle::Get().GetFloat("StatusBar.Height"))
+		[
+			MakeWidgetDrawer(StatusBarId)
+		]
+	];
+}
+
+TSharedRef<SWidgetDrawer> SLiveLinkHubStatusBar::MakeWidgetDrawer(FName StatusBarId)
+{
+	using namespace UE::LiveLinkHub::Private;
+	
+	WidgetDrawer = SNew(SWidgetDrawer, StatusBarId);
+
+	TSharedPtr<SMultiLineEditableTextBox> ConsoleEditBox;
+	FSimpleDelegate OnConsoleClosed;
+	FSimpleDelegate OnConsoleCommandExecuted;
+	const TSharedRef<SWidget> OutputLog = 
+			SNew(SBorder)
+			.BorderImage(FAppStyle::Get().GetBrush("Brushes.Panel"))
+			.VAlign(VAlign_Center)
+			.Padding(FMargin(6.0f, 0.0f))
+			[
+				SNew(SBox)
+				.WidthOverride(350.f)
+				[
+					FOutputLogModule::Get().MakeConsoleInputBox(ConsoleEditBox, OnConsoleClosed, OnConsoleCommandExecuted)
+				]
+			];
+	
+	FWidgetDrawerConfig OutputLogDrawer(OutputLogId);
+	GStatusBarManager.Init(WidgetDrawer.ToSharedRef(), OutputLogDrawer);
+	OutputLogDrawer.CustomWidget = OutputLog;
+
+	OutputLogDrawer.ButtonText = LOCTEXT("StatusBar_OutputLogButton", "Output Log");
+	OutputLogDrawer.Icon = FAppStyle::Get().GetBrush("Log.TabIcon");
+	WidgetDrawer->RegisterDrawer(MoveTemp(OutputLogDrawer));
+
+	return WidgetDrawer.ToSharedRef();
+}
+
+#undef LOCTEXT_NAMESPACE
diff --git a/Engine/Plugins/Animation/LiveLink/Source/LiveLinkHub/Private/UI/Widgets/SLiveLinkHubStatusBar.h b/Engine/Plugins/Animation/LiveLink/Source/LiveLinkHub/Private/UI/Widgets/SLiveLinkHubStatusBar.h
new file mode 100644
--- /dev/null
+++ b/Engine/Plugins/Animation/LiveLink/Source/LiveLinkHub/Private/UI/Widgets/SLiveLinkHubStatusBar.h
@@ -0,0 +1,43 @@
+// Copyright Epic Games, Inc. All Rights Reserved.
+
+#pragma once
+
+#include "Widgets/SCompoundWidget.h"
+
+class SWidgetDrawer;
+
+namespace UE::LiveLinkHub::Private
+{
+	const FName OutputLogId("OutputLog");
+}
+
+/**
+ * The status bar for Live Link Hub which contains console input and the output log.
+ */
+class SLiveLinkHubStatusBar : public SCompoundWidget
+{
+public:
+
+	SLATE_BEGIN_ARGS(SLiveLinkHubStatusBar) { }
+	SLATE_END_ARGS()
+
+	virtual ~SLiveLinkHubStatusBar() override;
+
+	/**
+	 * @param InArgs 
+	 * @param StatusBarId The unique status bar id to assign to the widget drawer.
+	 */
+	void Construct(const FArguments& InArgs, FName StatusBarId);
+
+private:
+	/**
+	 * Construct the widget drawer.
+	 *
+	 * @param StatusBarId The unique status bar id to assign.
+	 */
+	TSharedRef<SWidgetDrawer> MakeWidgetDrawer(FName StatusBarId);
+	
+private:
+	/** The widget drawer for this status bar. */
+	TSharedPtr<SWidgetDrawer> WidgetDrawer;
+};
diff --git a/Engine/Plugins/Animation/LiveLink/Source/LiveLinkHub/Private/UI/Widgets/SLiveLinkHubTabViewBase.cpp b/Engine/Plugins/Animation/LiveLink/Source/LiveLinkHub/Private/UI/Widgets/SLiveLinkHubTabViewBase.cpp
--- a/Engine/Plugins/Animation/LiveLink/Source/LiveLinkHub/Private/UI/Widgets/SLiveLinkHubTabViewBase.cpp
+++ b/Engine/Plugins/Animation/LiveLink/Source/LiveLinkHub/Private/UI/Widgets/SLiveLinkHubTabViewBase.cpp
@@ -9,6 +9,7 @@
 #include "Misc/App.h"
 #include "Modules/ModuleManager.h"
 #include "Recording/LiveLinkHubRecordingController.h"
+#include "SLiveLinkHubStatusBar.h"
 #include "Widgets/Input/SButton.h"
 #include "Widgets/Input/SComboButton.h"
 #include "Widgets/Layout/SBorder.h"
@@ -36,7 +37,8 @@ TSharedRef<SWidget> GetModeSwitcherContent()
 void SLiveLinkHubTabViewBase::Construct(const FArguments& InArgs)
 {
 	TSharedPtr<FLiveLinkHub> LiveLinkHub = FModuleManager::Get().GetModuleChecked<FLiveLinkHubModule>("LiveLinkHub").GetLiveLinkHub();
-
+	const FName StatusBarId = TEXT("LiveLinkHubStatusBar");
+	
 	ChildSlot
 	[
 		SNew(SBorder)
@@ -102,6 +104,17 @@ void SLiveLinkHubTabViewBase::Construct(const FArguments& InArgs)
 			[
 				InArgs._Content.Widget
 			]
+			+SVerticalBox::Slot()
+			.AutoHeight()
+			[
+				SNew(SSeparator)
+				.Orientation(Orient_Horizontal)
+			]
+			+SVerticalBox::Slot()
+			.AutoHeight()
+			[
+				SNew(SLiveLinkHubStatusBar, StatusBarId)
+			]
 		]
 	];
 }
diff --git a/Engine/Plugins/Animation/LiveLink/Source/LiveLinkHub/Private/UI/Widgets/SLiveLinkHubTabViewBase.h b/Engine/Plugins/Animation/LiveLink/Source/LiveLinkHub/Private/UI/Widgets/SLiveLinkHubTabViewBase.h
--- a/Engine/Plugins/Animation/LiveLink/Source/LiveLinkHub/Private/UI/Widgets/SLiveLinkHubTabViewBase.h
+++ b/Engine/Plugins/Animation/LiveLink/Source/LiveLinkHub/Private/UI/Widgets/SLiveLinkHubTabViewBase.h
@@ -20,7 +20,6 @@ public:
 
 	/**
 	 * @param InArgs
-	 * @param InStatusBarId Unique ID needed for the status bar
 	 */
 	void Construct(const FArguments& InArgs);
 };
