Add recording/playback support to livelink hub

diff --git a/Engine/Plugins/Animation/LiveLink/Source/LiveLinkHub/Private/Clients/GameThreadMessageHandler.h b/Engine/Plugins/Animation/LiveLink/Source/LiveLinkHub/Private/Clients/GameThreadMessageHandler.h
--- a/Engine/Plugins/Animation/LiveLink/Source/LiveLinkHub/Private/Clients/GameThreadMessageHandler.h
+++ b/Engine/Plugins/Animation/LiveLink/Source/LiveLinkHub/Private/Clients/GameThreadMessageHandler.h
@@ -38,7 +38,7 @@ public:
 	}
 
 	/** Virtual destructor. */
-	~TGameThreadMessageHandler() { }
+	~TGameThreadMessageHandler() = default;
 
 public:
 
@@ -70,4 +70,4 @@ private:
 
 	/** Holds a pointer to the actual handler function. */
 	FuncType Func;
-};
\ No newline at end of file
+};
diff --git a/Engine/Plugins/Animation/LiveLink/Source/LiveLinkHub/Private/LiveLinkHub.cpp b/Engine/Plugins/Animation/LiveLink/Source/LiveLinkHub/Private/LiveLinkHub.cpp
--- a/Engine/Plugins/Animation/LiveLink/Source/LiveLinkHub/Private/LiveLinkHub.cpp
+++ b/Engine/Plugins/Animation/LiveLink/Source/LiveLinkHub/Private/LiveLinkHub.cpp
@@ -30,7 +30,7 @@ void FLiveLinkHub::Initialize()
 	RecordingController = MakeShared<FLiveLinkHubRecordingController>();
 	PlaybackController = MakeShared<FLiveLinkHubPlaybackController>();
 	RecordingListController = MakeShared<FLiveLinkHubRecordingListController>(AsShared());
-	ClientsController = MakeShared<FLiveLinkHubClientsController>(LiveLinkProvider);
+	ClientsController = MakeShared<FLiveLinkHubClientsController>(LiveLinkProvider.ToSharedRef());
 
 	FString LiveLinkHubLayoutIni = GConfig->GetConfigFilename(TEXT("LiveLinkHubLayout"));
 	WindowController = MakeShared<FLiveLinkHubWindowController>(FLiveLinkHubWindowInitParams{ LiveLinkHubLayoutIni });
@@ -42,11 +42,13 @@ void FLiveLinkHub::Initialize()
 	LiveLinkHubClient->OnLiveLinkSubjectRemoved().AddSP(this, &FLiveLinkHub::OnSubjectRemoved);
 
 	PlaybackController->Start();
-	RecordingController->Initialize(PlaybackController);
 }
 
 FLiveLinkHub::~FLiveLinkHub()
 {
+	RecordingController.Reset();
+	PlaybackController.Reset();
+
 	LiveLinkHubClient->OnLiveLinkSubjectRemoved().RemoveAll(this);
 	LiveLinkHubClient->OnLiveLinkSubjectAdded().RemoveAll(this);
 	LiveLinkHubClient->OnFrameDataReceived_AnyThread().RemoveAll(this);
@@ -116,6 +118,7 @@ void FLiveLinkHub::OnFrameDataReceived_AnyThread(const FLiveLinkSubjectKey& InSu
 		RecordingController->RecordFrameData(InSubjectKey, InFrameDataStruct);
 	}
 
+	
 	FLiveLinkFrameDataStruct FrameDataCopy;
 	FrameDataCopy.InitializeWith(InFrameDataStruct);
 	LiveLinkProvider->UpdateSubjectFrameData(InSubjectKey.SubjectName, MoveTemp(FrameDataCopy));
@@ -140,4 +143,4 @@ void FLiveLinkHub::OnSubjectRemoved(FLiveLinkSubjectKey InSubjectKey)
 	LiveLinkProvider->RemoveSubject(InSubjectKey.SubjectName);
 }
 
-#undef LOCTEXT_NAMESPACE /*LiveLinkHub*/
\ No newline at end of file
+#undef LOCTEXT_NAMESPACE /*LiveLinkHub*/
diff --git a/Engine/Plugins/Animation/LiveLink/Source/LiveLinkHub/Private/LiveLinkHub.h b/Engine/Plugins/Animation/LiveLink/Source/LiveLinkHub/Private/LiveLinkHub.h
--- a/Engine/Plugins/Animation/LiveLink/Source/LiveLinkHub/Private/LiveLinkHub.h
+++ b/Engine/Plugins/Animation/LiveLink/Source/LiveLinkHub/Private/LiveLinkHub.h
@@ -24,7 +24,7 @@ class ULiveLinkRole;
 class ILiveLinkHub
 {
 public:
-	virtual ~ILiveLinkHub() {}
+	virtual ~ILiveLinkHub() = default;
 
 	/** Whether the hub is currently playing a recording. */
 	virtual bool IsInPlayback() const = 0;
diff --git a/Engine/Plugins/Animation/LiveLink/Source/LiveLinkHub/Private/Recording/Implementations/LiveLinkUAssetRecorder.cpp b/Engine/Plugins/Animation/LiveLink/Source/LiveLinkHub/Private/Recording/Implementations/LiveLinkUAssetRecorder.cpp
new file mode 100644
--- /dev/null
+++ b/Engine/Plugins/Animation/LiveLink/Source/LiveLinkHub/Private/Recording/Implementations/LiveLinkUAssetRecorder.cpp
@@ -0,0 +1,225 @@
+// Copyright Epic Games, Inc. All Rights Reserved.
+
+#include "LiveLinkUAssetRecorder.h"
+
+#include "Recording/LiveLinkRecording.h"
+#include "Recording/Implementations/LiveLinkUAssetRecording.h"
+
+#include "AssetToolsModule.h"
+#include "Containers/Map.h"
+#include "ContentBrowserModule.h"
+#include "Features/IModularFeatures.h"
+#include "Framework/MultiBox/MultiBoxBuilder.h"
+#include "FileHelpers.h"
+#include "InstancedStruct.h"
+#include "IContentBrowserSingleton.h"
+#include "LiveLinkHub.h"
+#include "LiveLinkHubClient.h"
+#include "LiveLinkHubModule.h"
+#include "LiveLinkTypes.h"
+#include "LiveLinkPreset.h"
+#include "Misc/FileHelper.h"
+#include "Misc/PackageName.h"
+#include "UI/Window/LiveLinkHubWindowController.h"
+#include "Widgets/DeclarativeSyntaxSupport.h"
+#include "Widgets/Layout/SBox.h"
+#include "Widgets/SWidget.h"
+#include "UObject/Object.h"
+#include "UObject/Package.h"
+#include "UObject/StrongObjectPtr.h"
+
+#define LOCTEXT_NAMESPACE "LiveLinkHub.RecordingController"
+
+
+namespace UAssetRecorderUtils
+{
+	TOptional<FLiveLinkRecordingStaticDataContainer> CreateStaticDataContainerFromFrameData(const FLiveLinkSubjectKey& SubjectKey)
+	{
+		TOptional<FLiveLinkRecordingStaticDataContainer> StaticDataContainer;
+		FLiveLinkHubClient* LiveLinkClient = static_cast<FLiveLinkHubClient*>(&IModularFeatures::Get().GetModularFeature<ILiveLinkClient>(ILiveLinkClient::ModularFeatureName));
+		TSubclassOf<ULiveLinkRole> LiveLinkRole = LiveLinkClient->GetSubjectRole_AnyThread(SubjectKey);
+		
+		if (const FLiveLinkStaticDataStruct* StaticData = LiveLinkClient->GetSubjectStaticData(SubjectKey))
+		{
+			FInstancedStruct StaticDataInstancedStruct;
+			StaticDataInstancedStruct.InitializeAs(StaticData->GetStruct(), (uint8*)StaticData->GetBaseData());
+
+			StaticDataContainer = FLiveLinkRecordingStaticDataContainer();
+			StaticDataContainer->Role = LiveLinkRole;
+			StaticDataContainer->RecordedData.Insert(MoveTemp(StaticDataInstancedStruct), 0);
+			StaticDataContainer->Timestamps.Add(0.0);
+		}
+
+		return StaticDataContainer;
+	}
+}
+
+void FLiveLinkUAssetRecorder::StartRecording()
+{
+	check(!CurrentRecording.IsValid());
+	CurrentRecording = MakePimpl<FLiveLinkUAssetRecordingData>();
+	RecordInitialStaticData();
+
+	bIsRecording = true;
+	TimeRecordingStarted = FPlatformTime::Seconds();
+}
+
+void FLiveLinkUAssetRecorder::StopRecording()
+{
+	if (CurrentRecording)
+	{
+		bIsRecording = false;
+
+		TimeRecordingEnded = FPlatformTime::Seconds();
+
+		SaveRecording();
+		CurrentRecording.Reset();
+	}
+}
+
+bool FLiveLinkUAssetRecorder::IsRecording() const
+{
+	return bIsRecording;
+}
+
+void FLiveLinkUAssetRecorder::RecordBaseData(FLiveLinkRecordingBaseDataContainer& StaticDataContainer, FInstancedStruct&& DataToRecord)
+{
+	const double TimeNowInSeconds = FPlatformTime::Seconds();
+	StaticDataContainer.RecordedData.Add(MoveTemp(DataToRecord));
+	StaticDataContainer.Timestamps.Add(TimeNowInSeconds - TimeRecordingStarted);
+}
+
+void FLiveLinkUAssetRecorder::RecordStaticData(const FLiveLinkSubjectKey& SubjectKey, TSubclassOf<ULiveLinkRole> Role, const FLiveLinkStaticDataStruct& StaticData)
+{
+	if (bIsRecording && CurrentRecording)
+	{
+		FLiveLinkRecordingStaticDataContainer& StaticDataContainer = CurrentRecording->StaticData.FindOrAdd(SubjectKey);
+		FInstancedStruct NewData;
+		NewData.InitializeAs(StaticData.GetStruct(), (uint8*)StaticData.GetBaseData());
+		StaticDataContainer.Role = Role;
+
+		RecordBaseData(StaticDataContainer, MoveTemp(NewData));
+	}
+}
+
+void FLiveLinkUAssetRecorder::RecordFrameData(const FLiveLinkSubjectKey& SubjectKey, const FLiveLinkFrameDataStruct& FrameData)
+{
+	if (bIsRecording && CurrentRecording)
+	{
+		FLiveLinkRecordingBaseDataContainer& FrameDataContainer = CurrentRecording->FrameData.FindOrAdd(SubjectKey);
+		FInstancedStruct NewData;
+		NewData.InitializeAs(FrameData.GetStruct(), (uint8*)FrameData.GetBaseData());
+
+		RecordBaseData(FrameDataContainer, MoveTemp(NewData));
+	}
+}
+
+bool FLiveLinkUAssetRecorder::OpenSaveDialog(const FString& InDefaultPath, const FString& InNewNameSuggestion, FString& OutPackageName)
+{
+	TSharedRef<SWindow> RootWindow = FModuleManager::Get().GetModuleChecked<FLiveLinkHubModule>("LiveLinkHub").GetLiveLinkHub()->GetRootWindow();
+
+	FSaveAssetDialogConfig SaveAssetDialogConfig;
+	{
+		SaveAssetDialogConfig.DefaultPath = InDefaultPath;
+		SaveAssetDialogConfig.DefaultAssetName = InNewNameSuggestion;
+		SaveAssetDialogConfig.AssetClassNames.Add(ULiveLinkRecording::StaticClass()->GetClassPathName());
+		SaveAssetDialogConfig.ExistingAssetPolicy = ESaveAssetDialogExistingAssetPolicy::AllowButWarn;
+		SaveAssetDialogConfig.DialogTitleOverride = LOCTEXT("SaveLiveLinkRecordingDialogTitle", "Save LiveLink Recording");
+		SaveAssetDialogConfig.WindowOverride = RootWindow;
+	}
+
+	FContentBrowserModule& ContentBrowserModule = FModuleManager::LoadModuleChecked<FContentBrowserModule>("ContentBrowser");
+	FString SaveObjectPath = ContentBrowserModule.Get().CreateModalSaveAssetDialog(SaveAssetDialogConfig);
+
+	if (!SaveObjectPath.IsEmpty())
+	{
+		OutPackageName = FPackageName::ObjectPathToPackageName(SaveObjectPath);
+		return true;
+	}
+
+	return false;
+}
+
+bool FLiveLinkUAssetRecorder::GetSavePresetPackageName(FString& OutName)
+{
+	FDateTime Today = FDateTime::Now();
+
+	TMap<FString, FStringFormatArg> FormatArgs;
+	FormatArgs.Add(TEXT("date"), Today.ToString());
+
+	FString DialogStartPath = TEXT("/Game/LiveLinkRecordings");
+
+	FString DefaultName = LOCTEXT("NewLiveLinkRecordingName", "NewLiveLinkRecording").ToString();
+
+	FString UniquePackageName;
+	FString UniqueAssetName;
+
+	FAssetToolsModule& AssetToolsModule = FModuleManager::LoadModuleChecked<FAssetToolsModule>("AssetTools");
+	AssetToolsModule.Get().CreateUniqueAssetName(DialogStartPath / DefaultName, TEXT(""), UniquePackageName, UniqueAssetName);
+
+	FString DialogStartName = FPaths::GetCleanFilename(UniqueAssetName);
+
+	FString UserPackageName;
+	FString NewPackageName;
+
+	// Get destination for asset
+	bool bFilenameValid = false;
+	while (!bFilenameValid)
+	{
+		if (!OpenSaveDialog(DialogStartPath, DialogStartName, UserPackageName))
+		{
+			return false;
+		}
+
+		NewPackageName = FString::Format(*UserPackageName, FormatArgs);
+
+		FText OutError;
+		bFilenameValid = FFileHelper::IsFilenameValidForSaving(NewPackageName, OutError);
+	}
+
+	OutName = MoveTemp(NewPackageName);
+	return true;
+}
+
+void FLiveLinkUAssetRecorder::SaveRecording()
+{
+	FString PackageName;
+	if (!GetSavePresetPackageName(PackageName))
+	{
+		return;
+	}
+
+	// Saving into a new package
+	const FString NewAssetName = FPackageName::GetLongPackageAssetName(PackageName);
+	UPackage* NewPackage = CreatePackage(*PackageName);
+
+	ULiveLinkUAssetRecording* NewRecording = NewObject<ULiveLinkUAssetRecording>(NewPackage, *NewAssetName, RF_Public | RF_Standalone);
+	NewRecording->LengthInSeconds = TimeRecordingEnded - TimeRecordingStarted;
+	NewRecording->RecordingPreset->BuildFromClient();
+	NewRecording->RecordingData = MoveTemp(*CurrentRecording);
+
+	if (NewRecording)
+	{
+		NewRecording->MarkPackageDirty();
+		//FAssetRegistryModule::AssetCreated(NewPreset);  Disabled for now since unsure if needed.
+		UEditorLoadingAndSavingUtils::SavePackages({ NewPackage }, false);
+	}
+	return;
+}
+
+void FLiveLinkUAssetRecorder::RecordInitialStaticData()
+{
+	FLiveLinkHubClient* LiveLinkClient = static_cast<FLiveLinkHubClient*>(&IModularFeatures::Get().GetModularFeature<ILiveLinkClient>(ILiveLinkClient::ModularFeatureName));
+	TArray<FLiveLinkSubjectKey> Subjects = LiveLinkClient->GetSubjects(true, true);
+
+	for (const FLiveLinkSubjectKey& Subject : Subjects)
+	{
+		TOptional<FLiveLinkRecordingStaticDataContainer> StaticDataContainer = UAssetRecorderUtils::CreateStaticDataContainerFromFrameData(Subject);
+		if (StaticDataContainer)
+		{
+			CurrentRecording->StaticData.Add(Subject, MoveTemp(*StaticDataContainer));
+		}
+	}
+}
+
+#undef LOCTEXT_NAMESPACE
diff --git a/Engine/Plugins/Animation/LiveLink/Source/LiveLinkHub/Private/Recording/Implementations/LiveLinkUAssetRecorder.h b/Engine/Plugins/Animation/LiveLink/Source/LiveLinkHub/Private/Recording/Implementations/LiveLinkUAssetRecorder.h
new file mode 100644
--- /dev/null
+++ b/Engine/Plugins/Animation/LiveLink/Source/LiveLinkHub/Private/Recording/Implementations/LiveLinkUAssetRecorder.h
@@ -0,0 +1,49 @@
+// Copyright Epic Games, Inc. All Rights Reserved.
+
+#pragma once
+
+#include "Recording/LiveLinkRecorder.h"
+
+#include "LiveLinkTypes.h"
+#include "Misc/CoreMiscDefines.h"
+#include "Templates/PimplPtr.h"
+
+struct FLiveLinkRecordingBaseDataContainer;
+struct FLiveLinkUAssetRecordingData;
+struct FInstancedStruct;
+
+/** UAsset implementation for serializing recorded livelink data. */
+class FLiveLinkUAssetRecorder : public ILiveLinkRecorder
+{
+public:
+
+	//~ Begin ILiveLinkRecorder
+	virtual void StartRecording() override;
+	virtual void StopRecording() override;
+	virtual void RecordStaticData(const FLiveLinkSubjectKey& SubjectKey, TSubclassOf<ULiveLinkRole> Role, const FLiveLinkStaticDataStruct& StaticData) override;
+	virtual void RecordFrameData(const FLiveLinkSubjectKey& SubjectKey, const FLiveLinkFrameDataStruct& FrameData) override;
+	virtual bool IsRecording() const override;
+	//~ End ILiveLinkRecorder
+
+private:
+	/** Prompt the user for a destination path for the recording. */
+	bool OpenSaveDialog(const FString& InDefaultPath, const FString& InNewNameSuggestion, FString& OutPackageName);
+	/** Creates a unique asset name and prompts the user for the recording name. */
+	bool GetSavePresetPackageName(FString& OutName);
+	/** Create a recording package and save it. */
+	void SaveRecording();
+	/** Record data to a ULiveLinkRecording object. */
+	void RecordBaseData(FLiveLinkRecordingBaseDataContainer& StaticDataContainer, FInstancedStruct&& DataToRecord);
+	/** Record initial data for all livelink subjects. (Useful when static data was sent before the recording started). */
+	void RecordInitialStaticData();
+
+private:
+	/** Holds metadata and recording data. */
+	TPimplPtr<FLiveLinkUAssetRecordingData> CurrentRecording;
+	/** Whether we're currently recording livelink data. */
+	bool bIsRecording = false;
+	/** Timestamp in seconds of when the recording was started. */
+	double TimeRecordingStarted = 0.0;
+	/** Timestamp in seconds of when the recording ended. */
+	double TimeRecordingEnded = 0.0;
+};
diff --git a/Engine/Plugins/Animation/LiveLink/Source/LiveLinkHub/Private/Recording/Implementations/LiveLinkUAssetRecording.h b/Engine/Plugins/Animation/LiveLink/Source/LiveLinkHub/Private/Recording/Implementations/LiveLinkUAssetRecording.h
new file mode 100644
--- /dev/null
+++ b/Engine/Plugins/Animation/LiveLink/Source/LiveLinkHub/Private/Recording/Implementations/LiveLinkUAssetRecording.h
@@ -0,0 +1,69 @@
+// Copyright Epic Games, Inc. All Rights Reserved.
+
+#pragma once
+
+#include "InstancedStruct.h"
+#include "LiveLinkPreset.h"
+#include "LiveLinkRole.h"
+#include "LiveLinkTypes.h"
+#include "Misc/DateTime.h"
+#include "Recording/LiveLinkRecording.h"
+#include "Templates/SubclassOf.h"
+#include "UObject/Object.h"
+
+#include "LiveLinkUAssetRecording.generated.h"
+
+/** Base data container for a recording track. */
+USTRUCT()
+struct FLiveLinkRecordingBaseDataContainer
+{
+	GENERATED_BODY()
+
+	/** Timestamps for the recorded data. Each entry matches an entry in the RecordedData array. */
+	UPROPERTY()
+	TArray<double> Timestamps;
+
+	/** Array of either static or frame recorded for a given timestamp.*/
+	UPROPERTY()
+	TArray<FInstancedStruct> RecordedData;
+};
+
+/** Container for static data. */
+USTRUCT()
+struct FLiveLinkRecordingStaticDataContainer : public FLiveLinkRecordingBaseDataContainer
+{
+	GENERATED_BODY()
+
+	/** The role of the static data being recorded. */
+	UPROPERTY()
+	TSubclassOf<ULiveLinkRole> Role = nullptr;
+};
+
+USTRUCT()
+struct FLiveLinkUAssetRecordingData
+{
+	GENERATED_BODY()
+
+	/** Length of the recording in seconds. */
+	UPROPERTY()
+	double LengthInSeconds = 0;
+
+	/** Static data encountered while recording. */
+	UPROPERTY()
+	TMap<FLiveLinkSubjectKey, FLiveLinkRecordingStaticDataContainer> StaticData;
+	
+	/** Frame data encountered while recording. */
+	UPROPERTY()
+	TMap<FLiveLinkSubjectKey, FLiveLinkRecordingBaseDataContainer> FrameData;
+};
+
+UCLASS()
+class ULiveLinkUAssetRecording : public ULiveLinkRecording
+{
+public:
+	GENERATED_BODY()
+
+	/** Recorded Static and frame data. */
+	UPROPERTY()
+	FLiveLinkUAssetRecordingData RecordingData;
+};
diff --git a/Engine/Plugins/Animation/LiveLink/Source/LiveLinkHub/Private/Recording/Implementations/LiveLinkUAssetRecordingPlayer.cpp b/Engine/Plugins/Animation/LiveLink/Source/LiveLinkHub/Private/Recording/Implementations/LiveLinkUAssetRecordingPlayer.cpp
new file mode 100644
--- /dev/null
+++ b/Engine/Plugins/Animation/LiveLink/Source/LiveLinkHub/Private/Recording/Implementations/LiveLinkUAssetRecordingPlayer.cpp
@@ -0,0 +1,134 @@
+// Copyright Epic Games, Inc. All Rights Reserved.
+
+#include "LiveLinkUAssetRecordingPlayer.h"
+
+#include "Recording/LiveLinkRecording.h"
+
+
+class FLiveLinkPlaybackTrackIterator
+{
+public:
+	FLiveLinkPlaybackTrackIterator(FLiveLinkPlaybackTrack& InTrack, int32 InInitialIndex)
+		: Track(InTrack)
+		, FrameIndex(InInitialIndex)
+	{
+	}
+
+	/** Advances to the next frame */
+	void operator++()
+	{
+		Advance();
+	}
+
+	/* @Return True if there are more frames in this track. */
+	explicit operator bool() const
+	{
+		return HasMoreFrames();
+	}
+
+	double FrameTimestamp() const
+	{
+		return Track.Timestamps[FrameIndex];
+	}
+
+	const FInstancedStruct& FrameData() const
+	{
+		return Track.FrameData[FrameIndex];
+	}
+
+	int32 CurrentIndex() const
+	{
+		return FrameIndex;
+	}
+
+private:
+	/* @Return True if there are more vertices on the component */
+	bool HasMoreFrames() const
+	{
+		return FrameIndex < Track.Timestamps.Num() && FrameIndex < Track.FrameData.Num();
+	}
+
+	/* Advances to the next frame */
+	void Advance()
+	{
+		++FrameIndex;
+	}
+
+private:
+	/** Track that's currently being iterated. */
+	FLiveLinkPlaybackTrack& Track;
+	/** "Playhead" for this track */
+	int32 FrameIndex = 0;
+};
+
+void FLiveLinkPlaybackTrack::GetFramesUntil(double InPlayhead, TArray<FLiveLinkRecordedFrame>& OutFrames)
+{
+	for (auto It = FLiveLinkPlaybackTrackIterator(*this, LastReadIndex + 1); It; ++It)
+	{
+		if (It.FrameTimestamp() > InPlayhead)
+		{
+			break;
+		}
+
+		LastReadIndex = It.CurrentIndex();
+
+		FLiveLinkRecordedFrame FrameToPlay;
+		FrameToPlay.Data = It.FrameData();
+		FrameToPlay.SubjectKey = SubjectKey;
+		FrameToPlay.LiveLinkRole = LiveLinkRole;
+
+		OutFrames.Add(MoveTemp(FrameToPlay));
+	}
+}
+
+TArray<FLiveLinkRecordedFrame> FLiveLinkPlaybackTracks::FetchNextFrames(double Playhead)
+{
+	TArray<FLiveLinkRecordedFrame> NextFrames;
+
+	if (Tracks.Num())
+	{
+		// todo: sort frames by timestamp
+		for (FLiveLinkPlaybackTrack& Track : Tracks)
+		{
+			Track.GetFramesUntil(Playhead, NextFrames);
+		}
+	}
+
+	return NextFrames;
+}
+
+void FLiveLinkPlaybackTracks::Restart()
+{
+	for (FLiveLinkPlaybackTrack& Track : Tracks)
+	{
+		Track.Restart();
+	}
+}
+
+void FLiveLinkUAssetRecordingPlayer::PreparePlayback(const ULiveLinkRecording* CurrentRecording)
+{
+	const ULiveLinkUAssetRecording* UAssetRecording = CastChecked<ULiveLinkUAssetRecording>(CurrentRecording);
+
+	FLiveLinkPlaybackTracks RecordingPlayback;
+
+	for (const TPair<FLiveLinkSubjectKey, FLiveLinkRecordingStaticDataContainer>& Pair : UAssetRecording->RecordingData.StaticData)
+	{
+		FLiveLinkPlaybackTrack PlaybackTrack;
+		PlaybackTrack.FrameData = TConstArrayView<FInstancedStruct>(Pair.Value.RecordedData);
+		PlaybackTrack.Timestamps = TConstArrayView<double>(Pair.Value.Timestamps);
+		PlaybackTrack.LiveLinkRole = Pair.Value.Role;
+		PlaybackTrack.SubjectKey = Pair.Key;
+		RecordingPlayback.Tracks.Add(PlaybackTrack);
+	}
+
+	for (const TPair<FLiveLinkSubjectKey, FLiveLinkRecordingBaseDataContainer>& Pair : UAssetRecording->RecordingData.FrameData)
+	{
+		FLiveLinkPlaybackTrack PlaybackTrack;
+		PlaybackTrack.FrameData = TConstArrayView<FInstancedStruct>(Pair.Value.RecordedData);
+		PlaybackTrack.Timestamps = TConstArrayView<double>(Pair.Value.Timestamps);
+		PlaybackTrack.SubjectKey = Pair.Key;
+		RecordingPlayback.Tracks.Add(PlaybackTrack);
+	}
+
+	CurrentRecordingPlayback = MoveTemp(RecordingPlayback);
+}
diff --git a/Engine/Plugins/Animation/LiveLink/Source/LiveLinkHub/Private/Recording/Implementations/LiveLinkUAssetRecordingPlayer.h b/Engine/Plugins/Animation/LiveLink/Source/LiveLinkHub/Private/Recording/Implementations/LiveLinkUAssetRecordingPlayer.h
new file mode 100644
--- /dev/null
+++ b/Engine/Plugins/Animation/LiveLink/Source/LiveLinkHub/Private/Recording/Implementations/LiveLinkUAssetRecordingPlayer.h
@@ -0,0 +1,64 @@
+// Copyright Epic Games, Inc. All Rights Reserved.
+
+#pragma once
+
+#include "Recording/LiveLinkRecordingPlayer.h"
+
+#include "Recording/Implementations/LiveLinkUAssetRecording.h"
+
+/** Playback track that holds recorded data for a given subject. */
+struct FLiveLinkPlaybackTrack
+{
+	void GetFramesUntil(double InPlayhead, TArray<FLiveLinkRecordedFrame>& OutFrames);
+
+	void Restart()
+	{
+		LastReadIndex = -1;
+	}
+
+	/** Frame data to read. */
+	TConstArrayView<struct FInstancedStruct> FrameData;
+	/** Timestamps for the frames in the track. */
+	TConstArrayView<double> Timestamps;
+	/** Used for static data. */
+	TSubclassOf<ULiveLinkRole> LiveLinkRole;
+	/** Subject key. */
+	FLiveLinkSubjectKey SubjectKey;
+	/** Index of the last frame that was read by the GetFrames method. */
+	int32 LastReadIndex = -1;
+
+	friend class FLiveLinkPlaybackTrackIterator;
+};
+
+/** Reorganized recording data to facilitate playback. */
+struct FLiveLinkPlaybackTracks
+{
+	/** Get the next frames */
+	TArray<FLiveLinkRecordedFrame> FetchNextFrames(double Playhead);
+
+	void Restart();
+
+public:
+	/** LiveLink tracks to playback. */
+	TArray<FLiveLinkPlaybackTrack> Tracks;
+};
+
+class FLiveLinkUAssetRecordingPlayer : public ILiveLinkRecordingPlayer
+{
+public:
+	void PreparePlayback(const class ULiveLinkRecording* CurrentRecording);
+
+	virtual TArray<FLiveLinkRecordedFrame> FetchNextFrames(double Playhead) override
+	{
+		return CurrentRecordingPlayback.FetchNextFrames(Playhead);
+	}
+
+	virtual void RestartPlayback() override
+	{
+		CurrentRecordingPlayback.Restart();
+	}
+	
+private:
+	/** All tracks for the current recording. */
+	FLiveLinkPlaybackTracks CurrentRecordingPlayback;
+};
diff --git a/Engine/Plugins/Animation/LiveLink/Source/LiveLinkHub/Private/Recording/LiveLinkHubPlaybackController.cpp b/Engine/Plugins/Animation/LiveLink/Source/LiveLinkHub/Private/Recording/LiveLinkHubPlaybackController.cpp
new file mode 100644
--- /dev/null
+++ b/Engine/Plugins/Animation/LiveLink/Source/LiveLinkHub/Private/Recording/LiveLinkHubPlaybackController.cpp
@@ -0,0 +1,176 @@
+// Copyright Epic Games, Inc. All Rights Reserved.
+
+#include "LiveLinkHubPlaybackController.h"
+
+#include "Async/Async.h"
+#include "Async/TaskGraphInterfaces.h"
+#include "ILiveLinkClient.h"
+#include "Implementations/LiveLinkUAssetRecordingPlayer.h"
+#include "LiveLinkHubLog.h"
+#include "LiveLinkRecording.h"
+#include "LiveLinkPreset.h"
+#include "LiveLinkTypes.h"
+#include "Features/IModularFeatures.h"
+#include "HAL/RunnableThread.h"
+#include "HAL/Event.h"
+#include "UObject/Object.h"
+#include "UObject/Package.h"
+
+FLiveLinkHubPlaybackController::FLiveLinkHubPlaybackController()
+{
+	Client = &IModularFeatures::Get().GetModularFeature<ILiveLinkClient>(ILiveLinkClient::ModularFeatureName);
+	
+	RecordingPlayer = MakeUnique<FLiveLinkUAssetRecordingPlayer>();
+}
+
+void FLiveLinkHubPlaybackController::StartPlayback()
+{
+	PlaybackStartTime = FPlatformTime::Seconds();
+	Playhead = 0;
+	bIsInPlayback = true;
+	
+	FPlatformMisc::MemoryBarrier();
+	PlaybackEvent->Trigger();
+}
+
+void FLiveLinkHubPlaybackController::PlayRecording(ULiveLinkRecording* InLiveLinkRecording)
+{
+	PreparePlayback(InLiveLinkRecording);
+}
+
+void FLiveLinkHubPlaybackController::StopPlayback()
+{
+	bIsInPlayback = false;
+}
+
+void FLiveLinkHubPlaybackController::PreparePlayback(ULiveLinkRecording* InLiveLinkRecording)
+{
+	if (InLiveLinkRecording)
+	{
+		if (bIsInPlayback)
+		{
+			check(RecordingToPlay);
+			UE_LOG(LogLiveLinkHub, Warning, TEXT("Can't start playback since we are already doing playback"));
+			return;
+		}
+
+		RecordingToPlay.Reset(InLiveLinkRecording);
+
+		RollbackPreset.Reset(NewObject<ULiveLinkPreset>(GetTransientPackage(), TEXT("RecordingRollbackPreset")));
+		// Save the current state of the sources/subjects in a rollback preset.
+		RollbackPreset->BuildFromClient();
+
+		RecordingPlayer->PreparePlayback(RecordingToPlay.Get());
+
+		RecordingToPlay->RecordingPreset->ApplyToClientLatent(
+			[this](bool)
+			{
+				StartPlayback();
+			});
+	}
+	else
+	{
+		UE_LOG(LogLiveLinkHub, Warning, TEXT("Started a recording playback with an invalid recording."));
+	}
+}
+
+FLiveLinkHubPlaybackController::~FLiveLinkHubPlaybackController()
+{
+	StopPlayback();
+	Stopping = true;
+	PlaybackEvent->Trigger();
+
+	if (Thread.IsValid())
+	{
+		Thread->WaitForCompletion();
+		Thread.Reset();
+	}
+}
+
+void FLiveLinkHubPlaybackController::Start()
+{
+	FString ThreadName = TEXT("LiveLinkHub Playback Controller ");
+	ThreadName.AppendInt(FAsyncThreadIndex::GetNext());
+
+	Thread.Reset(FRunnableThread::Create(this, *ThreadName, 128 * 1024, TPri_AboveNormal, FPlatformAffinity::GetPoolThreadMask()));
+}
+
+void FLiveLinkHubPlaybackController::Stop()
+{
+	Stopping = true;
+}
+
+uint32 FLiveLinkHubPlaybackController::Run()
+{
+	while (!Stopping.load())
+	{
+		PlaybackEvent->Wait();
+
+		while (bIsInPlayback)
+		{
+			if (Playhead >= RecordingToPlay->LengthInSeconds)
+			{
+				if (bLoopPlayback && RecordingToPlay->LengthInSeconds != 0)
+				{
+					Playhead = 0.0;
+					PlaybackStartTime = FPlatformTime::Seconds();
+
+					RecordingPlayer->RestartPlayback();
+				}
+				else
+				{
+					// Stop playback
+					break;
+				}
+			}
+			
+			TArray<FLiveLinkRecordedFrame> NextFrames = RecordingPlayer->FetchNextFrames(Playhead);
+
+			for (const FLiveLinkRecordedFrame& NextFrame : NextFrames)
+			{
+				// If we're sending static data
+				if (NextFrame.LiveLinkRole)
+				{
+					FLiveLinkStaticDataStruct StaticDataStruct;
+					StaticDataStruct.InitializeWith(NextFrame.Data.GetScriptStruct(), (FLiveLinkBaseStaticData*)NextFrame.Data.GetMemory());
+					Client->PushSubjectStaticData_AnyThread(NextFrame.SubjectKey, NextFrame.LiveLinkRole, MoveTemp(StaticDataStruct));
+				}
+				else
+				{
+					FLiveLinkFrameDataStruct FrameDataStruct;
+					FrameDataStruct.InitializeWith(NextFrame.Data.GetScriptStruct(), (FLiveLinkBaseFrameData*)NextFrame.Data.GetMemory());
+
+					Client->PushSubjectFrameData_AnyThread(NextFrame.SubjectKey, MoveTemp(FrameDataStruct));
+				}
+			}
+
+			Playhead = FPlatformTime::Seconds() - PlaybackStartTime;
+
+			// Don't sleep if we pushed frames since that can take a small amount of time.
+			if (!NextFrames.Num())
+			{
+				FPlatformProcess::Sleep(0.002);
+				Playhead = FPlatformTime::Seconds() - PlaybackStartTime;
+			}
+		}
+
+		// If the loop ended because the recording is over.
+		bIsInPlayback = false;
+
+		// Recording is done, clear the pointer.
+		CurrentRecording.Reset();
+		
+		// Trigger the playback finished delegate on the game thread.
+		FSimpleDelegateGraphTask::CreateAndDispatchWhenReady(FSimpleDelegateGraphTask::FDelegate::CreateRaw(this, &FLiveLinkHubPlaybackController::OnPlaybackFinished_Internal), TStatId(), nullptr, ENamedThreads::GameThread);
+	}
+
+	return 0;
+}
+
+void FLiveLinkHubPlaybackController::OnPlaybackFinished_Internal()
+{
+	RollbackPreset->ApplyToClientLatent();
+	RecordingToPlay.Reset();
+
+	PlaybackFinishedDelegate.Broadcast();
+}
diff --git a/Engine/Plugins/Animation/LiveLink/Source/LiveLinkHub/Private/Recording/LiveLinkHubPlaybackController.h b/Engine/Plugins/Animation/LiveLink/Source/LiveLinkHub/Private/Recording/LiveLinkHubPlaybackController.h
--- a/Engine/Plugins/Animation/LiveLink/Source/LiveLinkHub/Private/Recording/LiveLinkHubPlaybackController.h
+++ b/Engine/Plugins/Animation/LiveLink/Source/LiveLinkHub/Private/Recording/LiveLinkHubPlaybackController.h
@@ -2,28 +2,106 @@
 
 #pragma once
 
-/** Stub class for the playback controller. Will be replaced with the real implementation in a future CL. */
-class FLiveLinkHubPlaybackController
+#include "HAL/Runnable.h"
+
+#include "HAL/Event.h"
+#include "HAL/ThreadSafeBool.h"
+#include "Misc/Timespan.h"
+#include "Misc/Optional.h"
+#include "Recording/LiveLinkRecordingPlayer.h"
+#include "Templates/SubclassOf.h"
+#include "Templates/SharedPointer.h"
+#include "Templates/PimplPtr.h"
+#include "UObject/StrongObjectPtr.h"
+
+#include <atomic>
+
+class ILiveLinkClient;
+struct FLiveLinkRecordingData;
+class FRunnableThread;
+class ULiveLinkRecording;
+class ULiveLinkPreset;
+class ULiveLinkRecording;
+
+class FLiveLinkHubPlaybackController : public FRunnable
 {
 public:
-	FLiveLinkHubPlaybackController() = default;
+	FLiveLinkHubPlaybackController();
+	virtual ~FLiveLinkHubPlaybackController() override;
+	
+	/** Start playing a livelink recording. */
+	void PlayRecording(ULiveLinkRecording* InLiveLinkRecording);
+
+	/** Stop playing a livelink recording. */
+	void StopPlayback();
 
+	/** Returns whether we've started or are actively playing a recording. */ 
 	bool IsInPlayback() const
 	{
-		return false;
+		return RecordingToPlay || bIsInPlayback.load();
 	}
 
+	/** Returns whether the recording is set to loop. */
 	bool IsLooping() const
 	{
-		return false;
+		return bLoopPlayback.load();
 	}
 
+	/** Set whether a recording should loop. */
 	void SetLooping(bool bInLoop)
 	{
+		bLoopPlayback.store(bInLoop);
 	}
 	
-public:
-	void Start()
+	/** Delegate called when playback is finished (if recording is not set to loop). */
+	FSimpleMulticastDelegate& OnPlaybackFinished()
 	{
+		return PlaybackFinishedDelegate;
 	}
+	
+	/** Create the playback thread. */
+	void Start();
+	
+	//~ Begin FRunnable Interface
+	virtual bool Init() override { return true; }
+	virtual uint32 Run() override;
+	virtual void Stop() override;
+	virtual void Exit() override { }
+	//~ End FRunnable Interface
+
+private:
+	/** Trigger the playback thread to start reading data. */
+	void StartPlayback();
+	/** Apply the recording's preset then prepare the data needed to playback. */
+	void PreparePlayback(ULiveLinkRecording* InLiveLinkRecording);
+	/** Handler called when playback is finished on the playback thread. Is responsible for resetting the livelink state to what it was before we started playback. */
+	void OnPlaybackFinished_Internal();
+	
+private:
+	/** Flag for terminating the thread loop */
+	std::atomic<bool> Stopping = false;
+	/** Thread to do playback on **/
+	TUniquePtr<FRunnableThread> Thread;
+	/** Event signaling that a recording is available for playback. */
+	FEventRef PlaybackEvent = FEventRef();
+	/** Whether we're doing playback. */
+	std::atomic<bool> bIsInPlayback = false;
+	/** Recording we're currently playing. */
+	TSharedPtr<FLiveLinkRecordingData> CurrentRecording;
+	/** LiveLinkRecording to play.  */
+	TStrongObjectPtr<ULiveLinkRecording> RecordingToPlay;
+	/** Delegate called when a recording playback is finished (if it's not looping). */
+	FSimpleMulticastDelegate PlaybackFinishedDelegate;
+	/** Preset used to rollback the hub to its previous state after playing a recording. */
+	TStrongObjectPtr<ULiveLinkPreset> RollbackPreset;
+	/** Atomic bool keeping track of whether we should loop the playback. */
+	std::atomic<bool> bLoopPlayback = false;
+	/** Implementation of the playback functionality. */
+	TUniquePtr<ILiveLinkRecordingPlayer> RecordingPlayer;
+	/** LiveLinkClient used to transmit the data to connected clients. */
+	ILiveLinkClient* Client = nullptr;
+	/** Time that the playback started. */
+	double PlaybackStartTime = 0.0;
+	/** Playhead for the current playback. */
+	double Playhead = 0.0;
 };
diff --git a/Engine/Plugins/Animation/LiveLink/Source/LiveLinkHub/Private/Recording/LiveLinkHubRecordingController.cpp b/Engine/Plugins/Animation/LiveLink/Source/LiveLinkHub/Private/Recording/LiveLinkHubRecordingController.cpp
new file mode 100644
--- /dev/null
+++ b/Engine/Plugins/Animation/LiveLink/Source/LiveLinkHub/Private/Recording/LiveLinkHubRecordingController.cpp
@@ -0,0 +1,48 @@
+// Copyright Epic Games, Inc. All Rights Reserved.
+
+#include "LiveLinkHubRecordingController.h"
+
+#include "Implementations/LiveLinkUAssetRecorder.h"
+#include "LiveLinkRecorder.h"
+#include "SLiveLinkHubRecordingView.h"
+#include "Widgets/DeclarativeSyntaxSupport.h"
+#include "Widgets/Layout/SBox.h"
+#include "Widgets/SWidget.h"
+
+FLiveLinkHubRecordingController::FLiveLinkHubRecordingController()
+{
+	RecorderImplementation = MakeShared<FLiveLinkUAssetRecorder>();
+}
+
+TSharedRef<SWidget> FLiveLinkHubRecordingController::MakeRecordToolbarEntry()
+{
+	return SNew(SLiveLinkHubRecordingView)
+		.IsRecording_Raw(this, &FLiveLinkHubRecordingController::IsRecording)
+		.OnStartRecording_Raw(this, &FLiveLinkHubRecordingController::StartRecording)
+		.OnStopRecording_Raw(this, &FLiveLinkHubRecordingController::StopRecording);
+}
+
+void FLiveLinkHubRecordingController::StartRecording()
+{
+	RecorderImplementation->StartRecording();
+}
+	
+void FLiveLinkHubRecordingController::StopRecording()
+{
+	RecorderImplementation->StopRecording();
+}
+
+bool FLiveLinkHubRecordingController::IsRecording() const
+{
+	return RecorderImplementation->IsRecording();
+}
+
+void FLiveLinkHubRecordingController::RecordStaticData(const FLiveLinkSubjectKey& SubjectKey, TSubclassOf<ULiveLinkRole> Role, const FLiveLinkStaticDataStruct& StaticData)
+{
+	RecorderImplementation->RecordStaticData(SubjectKey, Role, StaticData);
+}
+	
+void FLiveLinkHubRecordingController::RecordFrameData(const FLiveLinkSubjectKey& SubjectKey, const FLiveLinkFrameDataStruct& FrameData)
+{
+	RecorderImplementation->RecordFrameData(SubjectKey, FrameData);
+}
diff --git a/Engine/Plugins/Animation/LiveLink/Source/LiveLinkHub/Private/Recording/LiveLinkHubRecordingController.h b/Engine/Plugins/Animation/LiveLink/Source/LiveLinkHub/Private/Recording/LiveLinkHubRecordingController.h
--- a/Engine/Plugins/Animation/LiveLink/Source/LiveLinkHub/Private/Recording/LiveLinkHubRecordingController.h
+++ b/Engine/Plugins/Animation/LiveLink/Source/LiveLinkHub/Private/Recording/LiveLinkHubRecordingController.h
@@ -1,46 +1,38 @@
 // Copyright Epic Games, Inc. All Rights Reserved.
 
 #pragma once
-
+#include "LiveLinkRecorder.h"
 #include "LiveLinkTypes.h"
-#include "Widgets/SNullWidget.h"
+#include "Templates/PimplPtr.h"
+
+class ILiveLinkRecorder;
+struct FInstancedStruct;
+struct FLiveLinkRecordingBaseDataContainer;
+class SWidget;
 
-/** Stub class for the recording controller. Will be replaced with the real implementation in a future CL. */
 class FLiveLinkHubRecordingController
 {
 public:
+	FLiveLinkHubRecordingController();
 
-	void Initialize(const TSharedPtr<class FLiveLinkHubPlaybackController>& InPlaybackController)
-	{
-	}
+	/** Create the toolbar entry for starting/stopping recordings. */
+	TSharedRef<SWidget> MakeRecordToolbarEntry();
 
-	~FLiveLinkHubRecordingController()
-	{
-	}
+	/** Start recording livelink data. */
+	void StartRecording();
+	
+	/** Stop recording livelink data and prompt the user for a save location. */
+	void StopRecording();
 
-	TSharedRef<SWidget> MakeRecordToolbarEntry()
-	{
-		return SNullWidget::NullWidget;
-	}
+	/** Returns whether we're currently recording. */
+	bool IsRecording() const;
 
-	void StartRecording()
-	{
-	}
-	
-	void StopRecording()
-	{
-	}
-
-	bool IsRecording() const
-	{
-		return false;
-	}
-	
-	void RecordStaticData(const FLiveLinkSubjectKey& SubjectKey, TSubclassOf<ULiveLinkRole> Role, const FLiveLinkStaticDataStruct& StaticData)
-	{
-	}
-	
-	void RecordFrameData(const FLiveLinkSubjectKey& SubjectKey, const FLiveLinkFrameDataStruct& FrameData)
-	{
-	}
+	/** Record static data in the current recording. */
+	void RecordStaticData(const FLiveLinkSubjectKey& SubjectKey, TSubclassOf<class ULiveLinkRole> Role, const FLiveLinkStaticDataStruct& StaticData);
+
+	/** Record frame data in the current recording. */
+	void RecordFrameData(const FLiveLinkSubjectKey& SubjectKey, const FLiveLinkFrameDataStruct& FrameData);
+private:
+	/** Recorder used to serialize livelink data into a given format. */
+	TSharedPtr<ILiveLinkRecorder> RecorderImplementation;
 };
diff --git a/Engine/Plugins/Animation/LiveLink/Source/LiveLinkHub/Private/Recording/LiveLinkHubRecordingListController.h b/Engine/Plugins/Animation/LiveLink/Source/LiveLinkHub/Private/Recording/LiveLinkHubRecordingListController.h
--- a/Engine/Plugins/Animation/LiveLink/Source/LiveLinkHub/Private/Recording/LiveLinkHubRecordingListController.h
+++ b/Engine/Plugins/Animation/LiveLink/Source/LiveLinkHub/Private/Recording/LiveLinkHubRecordingListController.h
@@ -2,19 +2,84 @@
 
 #pragma once
 
-#include "Widgets/SNullWidget.h"
+#include "AssetRegistry/AssetData.h"
+#include "IContentBrowserSingleton.h"
+#include "Framework/MultiBox/MultiBoxBuilder.h"
+#include "LiveLinkHub.h"
+#include "LiveLinkHubLog.h"
+#include "LiveLinkHubPlaybackController.h"
+#include "LiveLinkHubRecordingListController.h"
+#include "SLiveLinkHubRecordingListView.h"
+#include "LiveLinkRecording.h"
+#include "Modules/ModuleManager.h"
+#include "Misc/PackageName.h"
+
+#define LOCTEXT_NAMESPACE "LiveLinkHub.RecordingList"
 
-/** Stub class for the recording list controller. Will be replaced with the real implementation in a future CL. */
 class FLiveLinkHubRecordingListController
 {
 public:
-	FLiveLinkHubRecordingListController(const TSharedPtr<FLiveLinkHub>& InLiveLinkHub)
+	FLiveLinkHubRecordingListController(const TSharedRef<FLiveLinkHub>& InLiveLinkHub)
+		: LiveLinkHub(InLiveLinkHub)
 	{
 	}
 
 	/** Create the list's widget. */
 	TSharedRef<SWidget> MakeRecordingList()
 	{
-		return SNullWidget::NullWidget;
+		return SNew(SLiveLinkHubRecordingListView)
+			.IsLooping_Raw(this, &FLiveLinkHubRecordingListController::IsLooping)
+			.OnSetLooping_Raw(this, &FLiveLinkHubRecordingListController::SetLooping)
+			.OnImportRecording_Raw(this, &FLiveLinkHubRecordingListController::OnImportRecording);
+	}
+
+private:
+	/** Returns whether the recording playback should loop. */
+	bool IsLooping() const
+	{
+		if (TSharedPtr<FLiveLinkHub> HubPtr = LiveLinkHub.Pin())
+		{
+			return HubPtr->GetPlaybackController()->IsLooping();
+		}
+		return false;
+	}
+
+	/** Sets whether the playback should loop. */
+	void SetLooping(bool bLooping) const
+	{
+		if (TSharedPtr<FLiveLinkHub> HubPtr = LiveLinkHub.Pin())
+		{
+			HubPtr->GetPlaybackController()->SetLooping(bLooping);
+		}
+	}
+
+	/** Handler called when a recording a clicked to start the recording.  */
+	void OnImportRecording(const FAssetData& AssetData) const
+	{
+		UObject* RecordingAssetData = AssetData.GetAsset();
+		if (!RecordingAssetData)
+		{
+			UE_LOG(LogLiveLinkHub, Warning, TEXT("Failed to import recording %s"), *AssetData.AssetName.ToString());
+			StopPlayback();
+			return;
+		}
+
+		ULiveLinkRecording* ImportedRecording = Cast<ULiveLinkRecording>(RecordingAssetData);
+		LiveLinkHub.Pin()->GetPlaybackController()->PlayRecording(ImportedRecording);
+	}
+
+	/** Handler called when user clicks in the recording list to stop the current recording (Temporary until we have a stop button).  */
+	void StopPlayback() const
+	{
+		if (TSharedPtr<FLiveLinkHub> HubPtr = LiveLinkHub.Pin())
+		{
+			HubPtr->GetPlaybackController()->StopPlayback();
+		}
 	}
+
+private:
+	/** LiveLinkHub object that holds the different controllers. */
+	TWeakPtr<FLiveLinkHub> LiveLinkHub;
 };
+
+#undef LOCTEXT_NAMESPACE /* LiveLinkHub.RecordingList */
diff --git a/Engine/Plugins/Animation/LiveLink/Source/LiveLinkHub/Private/Recording/LiveLinkPlaybackListView.h b/Engine/Plugins/Animation/LiveLink/Source/LiveLinkHub/Private/Recording/LiveLinkPlaybackListView.h
new file mode 100644
--- /dev/null
+++ b/Engine/Plugins/Animation/LiveLink/Source/LiveLinkHub/Private/Recording/LiveLinkPlaybackListView.h
@@ -0,0 +1,106 @@
+// Copyright Epic Games, Inc. All Rights Reserved.
+
+#pragma once
+
+#include "Widgets/SCompoundWidget.h"
+
+#define LOCTEXT_NAMESPACE "LiveLinkHub.RecordingListView"
+
+class SLiveLinkHubRecordingListView : public SCompounWidget
+{
+public:
+	SLATE_BEGIN_ARGS(SLiveLinkHubRecordingListView)
+		{}
+	SLATE_END_ARGS()
+
+	void Construct(const FArguments& InArgs)
+	{
+		ChildSlot
+		[
+			SNew(SVerticalBox)
+			+ SVerticalBox::Slot()
+			[
+				CreateRecordingPicker()
+			]
+			+ SVerticalBox::Slot()
+			.AutoHeight()
+			[
+				SNew(SHorizontalBox)
+				+ SHorizontalBox::Slot()
+				.Padding(4.f, 0)
+				.VAlign(VAlign_Center)
+				.AutoWidth()
+				[
+					SNew(STextBlock)
+					.Text(LOCTEXT("LoopRecordingLabel", "Loop"))
+				]
+				+ SHorizontalBox::Slot()
+				.Padding(4.f, 0)
+				.VAlign(VAlign_Center)
+				.AutoWidth()
+				[
+					SNew(SCheckBox)
+					.IsChecked_Lambda([]()
+					{
+						FLiveLinkHubModule& LiveLinkHubModule = FModuleManager::Get().GetModuleChecked<FLiveLinkHubModule>("LiveLinkHub");
+						return LiveLinkHubModule.GetPlaybackController()->IsLooping() ? ECheckBoxState::Checked : ECheckBoxState::Unchecked;
+					})
+					.OnCheckStateChanged_Lambda([](ECheckBoxState InState)
+					{
+						FLiveLinkHubModule& LiveLinkHubModule = FModuleManager::Get().GetModuleChecked<FLiveLinkHubModule>("LiveLinkHub");
+						return LiveLinkHubModule.GetPlaybackController()->SetLooping(InState == ECheckBoxState::Checked);
+					})
+					.Padding(FMargin(4.0f, 0.0f))
+				]
+			];
+		]
+	}
+
+	TSharedRef<SWidget> CreateRecordingPicker()
+	{
+		FMenuBuilder MenuBuilder(true, nullptr);
+
+		IContentBrowserSingleton& ContentBrowser = FModuleManager::LoadModuleChecked<FContentBrowserModule>("ContentBrowser").Get();
+
+		FAssetPickerConfig AssetPickerConfig;
+		{
+			AssetPickerConfig.SelectionMode = ESelectionMode::Single;
+			AssetPickerConfig.InitialAssetViewType = EAssetViewType::Column;
+			AssetPickerConfig.bFocusSearchBoxWhenOpened = true;
+			AssetPickerConfig.bAllowNullSelection = false;
+			AssetPickerConfig.bShowBottomToolbar = true;
+			AssetPickerConfig.bAutohideSearchBar = false;
+			AssetPickerConfig.bAllowDragging = false;
+			AssetPickerConfig.bCanShowClasses = false;
+			AssetPickerConfig.bShowPathInColumnView = true;
+			AssetPickerConfig.bSortByPathInColumnView = false;
+			AssetPickerConfig.AssetShowWarningText = LOCTEXT("NoRecordings_Warning", "No Recordings Found");
+
+			AssetPickerConfig.bForceShowEngineContent = true;
+			AssetPickerConfig.bForceShowPluginContent = true;
+
+			AssetPickerConfig.Filter.ClassPaths.Add(ULiveLinkRecording::StaticClass()->GetClassPathName());
+			AssetPickerConfig.OnAssetSelected = FOnAssetSelected::CreateRaw(this, &FLiveLinkHubRecordingListController::OnImportRecording);
+		}
+
+		MenuBuilder.BeginSection(NAME_None, LOCTEXT("ImportRecording_MenuSection", "Import Recording"));
+		{
+			TSharedRef<SWidget> Picker = SNew(SBox)
+				.MinDesiredWidth(400.f)
+				.MinDesiredHeight(400.f)
+				[
+					ContentBrowser.CreateAssetPicker(AssetPickerConfig)
+				];
+
+			MenuBuilder.AddWidget(Picker, FText(), true, false);
+		}
+		MenuBuilder.EndSection();
+
+		return MenuBuilder.MakeWidget();
+	}
+
+	void MakeRecordingList()
+	{
+		
+	}
+};
diff --git a/Engine/Plugins/Animation/LiveLink/Source/LiveLinkHub/Private/Recording/LiveLinkPlaybackSubject.h b/Engine/Plugins/Animation/LiveLink/Source/LiveLinkHub/Private/Recording/LiveLinkPlaybackSubject.h
--- a/Engine/Plugins/Animation/LiveLink/Source/LiveLinkHub/Private/Recording/LiveLinkPlaybackSubject.h
+++ b/Engine/Plugins/Animation/LiveLink/Source/LiveLinkHub/Private/Recording/LiveLinkPlaybackSubject.h
@@ -92,9 +92,9 @@ protected:
 	//~ End ILiveLinkSubject Interface
 	
 private:
-	/** Not used, but needed to implement the FLiveLinkSubject interface. */
+	/** Last evaluated frame for this subject. */
 	FLiveLinkSubjectFrameData CurrentFrameSnapshot;
 
-	/** Subject key for the playback subject. */
+	/** Name of the subject */
 	FLiveLinkSubjectKey SubjectKey;
 };
diff --git a/Engine/Plugins/Animation/LiveLink/Source/LiveLinkHub/Private/Recording/LiveLinkRecorder.h b/Engine/Plugins/Animation/LiveLink/Source/LiveLinkHub/Private/Recording/LiveLinkRecorder.h
new file mode 100644
--- /dev/null
+++ b/Engine/Plugins/Animation/LiveLink/Source/LiveLinkHub/Private/Recording/LiveLinkRecorder.h
@@ -0,0 +1,24 @@
+// Copyright Epic Games, Inc. All Rights Reserved.
+
+#pragma once
+
+#include "LiveLinkRole.h"
+#include "LiveLinkTypes.h"
+
+/** Interface for handling serialization of livelink data. */
+class ILiveLinkRecorder
+{
+public:
+	virtual ~ILiveLinkRecorder() = default;
+
+	/** Start recording livelink data. */
+	virtual void StartRecording() = 0;
+	/** Stop recording livelink data and prompt the user for a save location. */
+	virtual void StopRecording() = 0;
+	/** Record static data in the current recording. */
+	virtual void RecordStaticData(const FLiveLinkSubjectKey& SubjectKey, TSubclassOf<ULiveLinkRole> Role, const FLiveLinkStaticDataStruct& StaticData) = 0;
+	/** Record frame data in the current recording. */
+	virtual void RecordFrameData(const FLiveLinkSubjectKey& SubjectKey, const FLiveLinkFrameDataStruct& FrameData) = 0;
+	/** Return whether we are currently recording livelink data. */
+	virtual bool IsRecording() const = 0;
+};
diff --git a/Engine/Plugins/Animation/LiveLink/Source/LiveLinkHub/Private/Recording/LiveLinkRecording.h b/Engine/Plugins/Animation/LiveLink/Source/LiveLinkHub/Private/Recording/LiveLinkRecording.h
new file mode 100644
--- /dev/null
+++ b/Engine/Plugins/Animation/LiveLink/Source/LiveLinkHub/Private/Recording/LiveLinkRecording.h
@@ -0,0 +1,28 @@
+// Copyright Epic Games, Inc. All Rights Reserved.
+
+#pragma once
+
+#include "LiveLinkPreset.h"
+#include "UObject/Object.h"
+
+#include "LiveLinkRecording.generated.h"
+
+UCLASS(Abstract)
+class ULiveLinkRecording : public UObject
+{
+public:
+	GENERATED_BODY()
+
+	ULiveLinkRecording()
+	{
+		RecordingPreset = CreateDefaultSubobject<ULiveLinkPreset>(TEXT("RecordingPreset"));
+	}
+
+	/** LiveLink Preset used to save the initial state of the sources and subjects at the time of recording. */
+	UPROPERTY(Instanced)
+	TObjectPtr<ULiveLinkPreset> RecordingPreset = nullptr;
+
+	/** Length of the recording. */
+	UPROPERTY()
+	double LengthInSeconds = 0.0;
+};
\ No newline at end of file
diff --git a/Engine/Plugins/Animation/LiveLink/Source/LiveLinkHub/Private/Recording/LiveLinkRecordingPlayer.h b/Engine/Plugins/Animation/LiveLink/Source/LiveLinkHub/Private/Recording/LiveLinkRecordingPlayer.h
new file mode 100644
--- /dev/null
+++ b/Engine/Plugins/Animation/LiveLink/Source/LiveLinkHub/Private/Recording/LiveLinkRecordingPlayer.h
@@ -0,0 +1,37 @@
+// Copyright Epic Games, Inc. All Rights Reserved.
+
+#pragma once
+
+#include "LiveLinkRole.h"
+#include "LiveLinkTypes.h"
+#include "StructView.h"
+#include "Templates/SubclassOf.h"
+
+class ULiveLinkRecording;
+
+/** Frame that was read by the recording player*/
+struct FLiveLinkRecordedFrame
+{
+	/** Recorded frame or static data. */
+	FConstStructView Data;
+	/** Subject that originally sent the data. */
+	FLiveLinkSubjectKey SubjectKey;
+	/** Role used to interpret the data (Only present with recorded static data). */
+	TSubclassOf<ULiveLinkRole> LiveLinkRole;
+};
+
+/** Object responsible for reading a livelink recording and providing the frames to the LiveLinkPlaybackController. */
+class ILiveLinkRecordingPlayer
+{
+public:
+	virtual ~ILiveLinkRecordingPlayer() = default;
+
+	/** Initialize internal structures needed for playback of the recorded data. */
+	virtual void PreparePlayback(const ULiveLinkRecording* Recording) = 0;
+
+	/** Fetch next frames at the provided playhead position. */
+	virtual TArray<FLiveLinkRecordedFrame> FetchNextFrames(double Playhead) = 0;
+
+	/** Restart the recording from the beginning. */
+	virtual void RestartPlayback() = 0;
+};
diff --git a/Engine/Plugins/Animation/LiveLink/Source/LiveLinkHub/Private/Recording/SLiveLinkHubRecordingListView.h b/Engine/Plugins/Animation/LiveLink/Source/LiveLinkHub/Private/Recording/SLiveLinkHubRecordingListView.h
new file mode 100644
--- /dev/null
+++ b/Engine/Plugins/Animation/LiveLink/Source/LiveLinkHub/Private/Recording/SLiveLinkHubRecordingListView.h
@@ -0,0 +1,148 @@
+// Copyright Epic Games, Inc. All Rights Reserved.
+
+#pragma once
+
+#include "Widgets/SCompoundWidget.h"
+
+#include "ContentBrowserModule.h"
+#include "Delegates/DelegateCombinations.h"
+#include "LiveLinkHub.h"
+#include "Recording/LiveLinkRecording.h"
+#include "Recording/LiveLinkHubPlaybackController.h"
+#include "Styling/SlateTypes.h"
+#include "Widgets/Text/STextBlock.h"
+#include "Widgets/Input/SCheckBox.h"
+
+
+#define LOCTEXT_NAMESPACE "LiveLinkHub.RecordingListView"
+
+DECLARE_DELEGATE_RetVal(bool, FOnGetLooping);
+DECLARE_DELEGATE_OneParam(FOnSetLooping, bool);
+DECLARE_DELEGATE_OneParam(FOnImportRecording, const struct FAssetData&);
+
+class SLiveLinkHubRecordingListView : public SCompoundWidget
+{
+public:
+	SLATE_BEGIN_ARGS(SLiveLinkHubRecordingListView)
+		{}
+		SLATE_EVENT(FOnGetLooping, IsLooping)
+		SLATE_EVENT(FOnSetLooping, OnSetLooping)
+		SLATE_EVENT(FOnImportRecording, OnImportRecording)
+	SLATE_END_ARGS()
+
+	//~ Begin SWidget interface
+	void Construct(const FArguments& InArgs)
+	{
+		OnGetLoopingDelegate = InArgs._IsLooping;
+		OnSetLoopingDelegate = InArgs._OnSetLooping;
+		OnImportRecordingDelegate = InArgs._OnImportRecording;
+
+		ChildSlot
+		[
+			SNew(SVerticalBox)
+			+ SVerticalBox::Slot()
+			[
+				CreateRecordingPicker()
+			]
+			+ SVerticalBox::Slot()
+			.AutoHeight()
+			[
+				SNew(SHorizontalBox)
+				+ SHorizontalBox::Slot()
+				.Padding(4.f, 0)
+				.VAlign(VAlign_Center)
+				.AutoWidth()
+				[
+					SNew(STextBlock)
+					.Text(LOCTEXT("LoopRecordingLabel", "Loop"))
+				]
+				+ SHorizontalBox::Slot()
+				.Padding(4.f, 0)
+				.VAlign(VAlign_Center)
+				.AutoWidth()
+				[
+					SNew(SCheckBox)
+					.Padding(FMargin(4.0f, 0.0f))
+					.IsChecked(this, &SLiveLinkHubRecordingListView::IsLoopedChecked)
+					.OnCheckStateChanged(this, &SLiveLinkHubRecordingListView::OnCheckStateChanged)
+				]
+			]
+		];
+	}
+	//~ End SWidget interface
+
+private:
+	/** Returns whether the looping checkbox should be checked or unchecked. */
+	ECheckBoxState IsLoopedChecked() const
+	{
+		return OnGetLoopingDelegate.Execute() ? ECheckBoxState::Checked : ECheckBoxState::Unchecked;
+	}
+
+	/** Handler to change the looping option for playback. */
+	void OnCheckStateChanged(ECheckBoxState InState)
+	{
+		OnSetLoopingDelegate.Execute(InState == ECheckBoxState::Checked);
+	}
+
+	/** Callback to notice the hub that we've selected a recording to play. */
+	void OnImportRecording(const FAssetData& AssetData) const
+	{
+		OnImportRecordingDelegate.Execute(AssetData);
+	}
+
+	/** Creates the asset picker widget for selecting a recording. */
+	TSharedRef<SWidget> CreateRecordingPicker()
+	{
+		FMenuBuilder MenuBuilder(true, nullptr);
+
+		IContentBrowserSingleton& ContentBrowser = FModuleManager::LoadModuleChecked<FContentBrowserModule>("ContentBrowser").Get();
+
+		FAssetPickerConfig AssetPickerConfig;
+		{
+			AssetPickerConfig.SelectionMode = ESelectionMode::Single;
+			AssetPickerConfig.InitialAssetViewType = EAssetViewType::Column;
+			AssetPickerConfig.bFocusSearchBoxWhenOpened = true;
+			AssetPickerConfig.bAllowNullSelection = false;
+			AssetPickerConfig.bShowBottomToolbar = true;
+			AssetPickerConfig.bAutohideSearchBar = false;
+			AssetPickerConfig.bAllowDragging = false;
+			AssetPickerConfig.bCanShowClasses = false;
+			AssetPickerConfig.bShowPathInColumnView = true;
+			AssetPickerConfig.bSortByPathInColumnView = false;
+			AssetPickerConfig.AssetShowWarningText = LOCTEXT("NoRecordings_Warning", "No Recordings Found");
+
+			AssetPickerConfig.bForceShowEngineContent = true;
+			AssetPickerConfig.bForceShowPluginContent = true;
+
+			AssetPickerConfig.Filter.ClassPaths.Add(ULiveLinkRecording::StaticClass()->GetClassPathName());
+			AssetPickerConfig.Filter.bRecursiveClasses = true;
+			AssetPickerConfig.Filter.bRecursivePaths = true;
+			AssetPickerConfig.OnAssetSelected = FOnAssetSelected::CreateRaw(this, &SLiveLinkHubRecordingListView::OnImportRecording);
+		}
+
+		MenuBuilder.BeginSection(NAME_None, LOCTEXT("ImportRecording_MenuSection", "Import Recording"));
+		{
+			TSharedRef<SWidget> PresetPicker = SNew(SBox)
+				.MinDesiredWidth(400.f)
+				.MinDesiredHeight(400.f)
+				[
+					ContentBrowser.CreateAssetPicker(AssetPickerConfig)
+				];
+
+			MenuBuilder.AddWidget(PresetPicker, FText(), true, false);
+		}
+		MenuBuilder.EndSection();
+
+		return MenuBuilder.MakeWidget();
+	}
+
+private:
+	/** Delegate for checking if the recording playback should loop. */
+	FOnGetLooping OnGetLoopingDelegate;
+	/** Delegate to set the looping option. */
+	FOnSetLooping OnSetLoopingDelegate;
+	/** Delegate used for noticing the hub that a recording was selected for playback. */
+	FOnImportRecording OnImportRecordingDelegate;
+};
+
+#undef LOCTEXT_NAMESPACE /* LiveLinkHub.RecordingListView */
diff --git a/Engine/Plugins/Animation/LiveLink/Source/LiveLinkHub/Private/Recording/SLiveLinkHubRecordingView.h b/Engine/Plugins/Animation/LiveLink/Source/LiveLinkHub/Private/Recording/SLiveLinkHubRecordingView.h
new file mode 100644
--- /dev/null
+++ b/Engine/Plugins/Animation/LiveLink/Source/LiveLinkHub/Private/Recording/SLiveLinkHubRecordingView.h
@@ -0,0 +1,86 @@
+// Copyright Epic Games, Inc. All Rights Reserved.
+
+#pragma once
+
+#include "Widgets/SCompoundWidget.h"
+
+#include "Delegates/Delegate.h"
+#include "Delegates/DelegateCombinations.h"
+#include "Styling/SlateTypes.h"
+#include "Widgets/Input/SButton.h"
+
+class FLiveLinkHub;
+
+#define LOCTEXT_NAMESPACE "LiveLinkHub.RecordingView"
+
+DECLARE_DELEGATE_RetVal(bool, FOnGetIsRecording);
+
+/**
+ * Utility base class for all content that is displayed in an entire major tab.
+ */
+class SLiveLinkHubRecordingView : public SCompoundWidget
+{
+public:
+	SLATE_BEGIN_ARGS(SLiveLinkHubRecordingView) {}
+		SLATE_EVENT(FOnGetIsRecording, IsRecording)
+		SLATE_EVENT(FSimpleDelegate, OnStartRecording)
+		SLATE_EVENT(FSimpleDelegate, OnStopRecording)
+	SLATE_END_ARGS()
+
+	/**
+		* @param InArgs
+		* @param InStatusBarId Unique ID needed for the status bar
+		*/
+	void Construct(const FArguments& InArgs)
+	{
+		IsRecordingDelegate = InArgs._IsRecording;
+		OnStartRecordingDelegate = InArgs._OnStartRecording;
+		OnStopRecordingDelegate = InArgs._OnStopRecording;
+
+		ChildSlot
+		[
+			SNew(SButton)
+				.Text(TAttribute<FText>::CreateSP(this, &SLiveLinkHubRecordingView::GetRecordingText))
+				.OnClicked(this, &SLiveLinkHubRecordingView::OnClickRecordButton)
+		];
+	}
+
+	/**  Get the text shown in the Record button in the toolbar. */
+	FText GetRecordingText() const
+	{
+		static FText RecordText = LOCTEXT("RecordButtonLabel", "Record");
+		static FText RecordingText = LOCTEXT("RecordButtonRecordingLabel", "Recording...");
+		return IsRecordingDelegate.Execute() ? RecordingText : RecordText;
+	}
+
+	/** Handler called when record button is clicked */
+	FReply OnClickRecordButton()
+	{
+		if (IsRecordingDelegate.Execute())
+		{
+			OnStopRecordingDelegate.Execute();
+		}
+		else
+		{
+			OnStartRecordingDelegate.Execute();
+		}
+
+		return FReply::Handled();
+	}
+
+	/** Returns whether we're currently recording. */
+	bool IsRecording() const
+	{
+		return IsRecordingDelegate.Execute();
+	}
+
+private:
+	/** Delegate used to know if the hub is current recording. */
+	FOnGetIsRecording IsRecordingDelegate;
+	/** Delegate used to notice the hub that we should start recording. */
+	FSimpleDelegate OnStartRecordingDelegate;
+	/** Delegate used to notice the hub that we should stop recording. */
+	FSimpleDelegate OnStopRecordingDelegate;
+};
+
+#undef LOCTEXT_NAMESPACE /* LiveLinkHub.RecordingView */
\ No newline at end of file
