Add option for enabling/disabling subjects
Fix duplicated hub entries in the livelink panel

diff --git a/Engine/Plugins/Animation/LiveLink/Source/LiveLink/Public/LiveLinkClient.h b/Engine/Plugins/Animation/LiveLink/Source/LiveLink/Public/LiveLinkClient.h
--- a/Engine/Plugins/Animation/LiveLink/Source/LiveLink/Public/LiveLinkClient.h
+++ b/Engine/Plugins/Animation/LiveLink/Source/LiveLink/Public/LiveLinkClient.h
@@ -79,6 +79,7 @@ public:
 	virtual FLiveLinkSourcePreset GetSourcePreset(FGuid SourceGuid, UObject* DuplicatedObjectOuter) const override;
 	virtual FText GetSourceType(FGuid EntryGuid) const override;
 	virtual FText GetSourceStatus(FGuid EntryGuid) const override;
+	virtual FText GetSourceMachineName(FGuid EntryGuid) const override;
 
 
 	virtual void PushSubjectStaticData_AnyThread(const FLiveLinkSubjectKey& SubjectKey, TSubclassOf<ULiveLinkRole> Role, FLiveLinkStaticDataStruct&& StaticData) override;
@@ -161,7 +162,6 @@ public:
 
 	FLiveLinkSubjectTimeSyncData GetTimeSyncData(FLiveLinkSubjectName SubjectName);
 
-	FText GetSourceMachineName(FGuid EntryGuid) const;
 	bool IsSourceStillValid(FGuid EntryGuid) const;
 	UE_DEPRECATED(4.23, "FLiveLinkClient::GetSourceTypeForEntry is deprecated. Please use GetSourceType instead!")
 	FText GetSourceTypeForEntry(FGuid EntryGuid) const { return GetSourceType(EntryGuid); }
diff --git a/Engine/Plugins/Animation/LiveLink/Source/LiveLinkHub/Private/Clients/LiveLinkHubClientsController.h b/Engine/Plugins/Animation/LiveLink/Source/LiveLinkHub/Private/Clients/LiveLinkHubClientsController.h
--- a/Engine/Plugins/Animation/LiveLink/Source/LiveLinkHub/Private/Clients/LiveLinkHubClientsController.h
+++ b/Engine/Plugins/Animation/LiveLink/Source/LiveLinkHub/Private/Clients/LiveLinkHubClientsController.h
@@ -18,11 +18,11 @@
 
 
 /** Controller responsible for holding the list of connected clients and creating the clients view. */
-class FLiveLinkHubClientsController : public TSharedFromThis<FLiveLinkHubClientsController>
+class FLiveLinkHubClientsController
 {
 public:
-	FLiveLinkHubClientsController(const TSharedPtr<ILiveLinkHubClientsModel>& InClientsModel)
-		: ClientsModel (InClientsModel)
+	FLiveLinkHubClientsController(const TSharedRef<ILiveLinkHubClientsModel>& InClientsModel)
+		: ClientsModel(InClientsModel)
 	{
 		ClientsModel->OnClientEvent().AddRaw(this, &FLiveLinkHubClientsController::OnClientEvent);
 		
@@ -37,8 +37,8 @@ public:
 	/** Create the widget that displays connected UE clients. */
 	TSharedRef<SWidget> MakeClientsView()
 	{
-		return SAssignNew(ClientsView, SLiveLinkHubClientsView, ClientsModel)
-			.OnClientSelected(this, &FLiveLinkHubClientsController::OnClientSelected);
+		return SAssignNew(ClientsView, SLiveLinkHubClientsView, ClientsModel.ToSharedRef())
+			.OnClientSelected_Raw(this, &FLiveLinkHubClientsController::OnClientSelected);
 	}
 
 	/** Create the widget that displays information about a given UE client. */
@@ -63,7 +63,7 @@ public:
 
 private:
 	/** Update the client info details if we received new info about it. */
-	void OnClientEvent(const FMessageAddress& MessageAddress, ILiveLinkHubClientsModel::EClientEventType EventType)
+	void OnClientEvent(FMessageAddress MessageAddress, ILiveLinkHubClientsModel::EClientEventType EventType)
 	{
 		switch (EventType)
 		{
@@ -73,9 +73,9 @@ private:
 				break;
 			case ILiveLinkHubClientsModel::EClientEventType::Modified:
 			{
-				if (TSharedPtr<FMessageAddress> ClientAddress = ClientsView->GetSelectedClient())
+				if (FMessageAddress ClientAddress = ClientsView->GetSelectedClient(); ClientAddress.IsValid())
 				{
-					if (MessageAddress == *ClientAddress)
+					if (MessageAddress == ClientAddress)
 					{
 						UpdateClientDetails(MessageAddress);
 					}
@@ -90,11 +90,11 @@ private:
 	}
 
 	/** Handler called when a client is selected in the clients view. */
-	void OnClientSelected(const TSharedPtr<FMessageAddress>& Client)
+	void OnClientSelected(FMessageAddress Client)
 	{
-		if (Client)
+		if (Client.IsValid())
 		{
-			UpdateClientDetails(*Client);
+			UpdateClientDetails(Client);
 		}
 		else
 		{
@@ -104,7 +104,7 @@ private:
 	}
 	
 	/** Handles updating client details in the client details panel. */
-	void UpdateClientDetails(const FMessageAddress& Client)
+	void UpdateClientDetails(FMessageAddress Client)
 	{
 		if (TOptional<FLiveLinkHubUEClientInfo> ClientInfo = ClientsModel->GetClientInfo(Client))
 		{
diff --git a/Engine/Plugins/Animation/LiveLink/Source/LiveLinkHub/Private/Clients/LiveLinkHubClientsModel.h b/Engine/Plugins/Animation/LiveLink/Source/LiveLinkHub/Private/Clients/LiveLinkHubClientsModel.h
--- a/Engine/Plugins/Animation/LiveLink/Source/LiveLinkHub/Private/Clients/LiveLinkHubClientsModel.h
+++ b/Engine/Plugins/Animation/LiveLink/Source/LiveLinkHub/Private/Clients/LiveLinkHubClientsModel.h
@@ -9,6 +9,7 @@ template <typename OptionalType> struct TOptional;
 struct FLiveLinkClientInfoMessage;
 struct FMessageAddress;
 struct FLiveLinkHubUEClientInfo;
+struct FLiveLinkSubjectKey;
 
 
 /** Decouples the UI from the livelink hub functions. */
@@ -23,16 +24,31 @@ public:
 		Modified
 	};
 
-	DECLARE_MULTICAST_DELEGATE_TwoParams(FOnClientEvent, const FMessageAddress&, EClientEventType);
+	DECLARE_MULTICAST_DELEGATE_TwoParams(FOnClientEvent, FMessageAddress, EClientEventType);
 
-	virtual ~ILiveLinkHubClientsModel() {}
+	virtual ~ILiveLinkHubClientsModel() = default;
 
 	/** Delegate used to get noticed about changes to the client list. */
 	virtual FOnClientEvent& OnClientEvent() = 0;
 
+	/** Get the status text of a client. */
+	virtual FText GetClientStatus(FMessageAddress Client) const = 0;
+
 	/** Get the list of clients connected to the hub. */
 	virtual TArray<FMessageAddress> GetClients() const = 0;
 
 	/** Get information about a given client given its address. */
-	virtual TOptional<FLiveLinkHubUEClientInfo> GetClientInfo(const FMessageAddress& InAddress) const = 0;
+	virtual TOptional<FLiveLinkHubUEClientInfo> GetClientInfo(FMessageAddress InAddress) const = 0;
+
+	/** Get whether a client should receive livelink data. */
+	virtual bool IsClientEnabled(FMessageAddress Client) const = 0;
+
+	/** Set whether a client should receive livelink data. */
+	virtual void SetClientEnabled(FMessageAddress Client, bool bInEnable) = 0;
+
+	/** Get whether a subject is enabled on a given client. */
+	virtual bool IsSubjectEnabled(FMessageAddress Client, const FLiveLinkSubjectKey& Subject) const = 0;
+
+	/** Set whether a subject should receive livelink data. */
+	virtual void SetSubjectEnabled(FMessageAddress Client, const FLiveLinkSubjectKey& Subject, bool bInEnable) = 0;
 };
diff --git a/Engine/Plugins/Animation/LiveLink/Source/LiveLinkHub/Private/Clients/LiveLinkHubProvider.h b/Engine/Plugins/Animation/LiveLink/Source/LiveLinkHub/Private/Clients/LiveLinkHubProvider.h
--- a/Engine/Plugins/Animation/LiveLink/Source/LiveLinkHub/Private/Clients/LiveLinkHubProvider.h
+++ b/Engine/Plugins/Animation/LiveLink/Source/LiveLinkHub/Private/Clients/LiveLinkHubProvider.h
@@ -14,6 +14,7 @@
 #include "GameThreadMessageHandler.h"
 #include "HAL/CriticalSection.h"
 #include "IMessageContext.h"
+#include "LiveLinkHubLog.h"
 #include "LiveLinkHubMessages.h"
 #include "LiveLinkProviderImpl.h"
 #include "LiveLinkSettings.h"
@@ -22,6 +23,8 @@
 #include "Misc/ScopeLock.h"
 #include "TimerManager.h"
 
+#define LOCTEXT_NAMESPACE "LiveLinkHub.LiveLinkHubProvider"
+
 /** 
  * LiveLink Provider that allows getting more information about a UE client by communicating with a LiveLinkHub MessageBus Source.
  */
@@ -61,6 +64,26 @@ public:
 		}
 	}
 
+	virtual bool ShouldTransmitToSubject_AnyThread(FName SubjectName, FMessageAddress Address) const override
+	{
+		FReadScopeLock Locker(ClientsMapLock);
+		if (const FLiveLinkHubUEClientInfo* ClientInfoPtr = ClientsMap.Find(Address))
+		{
+			if (!ClientInfoPtr->bEnabled)
+			{
+				return false;
+			}
+
+			return !ClientInfoPtr->DisabledSubjects.Contains(SubjectName);
+		}
+		else
+		{
+			UE_LOG(LogLiveLinkHub, Warning, TEXT("Attempted to transmit data to an invalid client."));
+		}
+
+		return true;
+	}
+
 private:
 	/** Handle a connection message resulting from a livelink hub message bus source connecting to this provider. */
 	void HandleHubConnectMessage(const FLiveLinkHubConnectMessage& Message, const TSharedRef<IMessageContext, ESPMode::ThreadSafe>& Context)
@@ -71,18 +94,22 @@ private:
 
 		FMessageAddress ConnectionAddress = Context->GetSender();
 
-		// Remove old entries if one is found
+
 		TOptional<FMessageAddress> RemovedAddress;
-		for (auto It = ClientsMap.CreateIterator(); It; ++It)
 		{
-			FLiveLinkHubUEClientInfo& IteratedClient = It->Value;
-			if (IteratedClient.Hostname == Message.ClientInfo.Hostname && IteratedClient.LongName == Message.ClientInfo.LongName
-				&& IteratedClient.ProjectName == Message.ClientInfo.ProjectName && IteratedClient.CurrentLevel == Message.ClientInfo.CurrentLevel)
-			{
-				RemovedAddress = It->Key;
-				It.RemoveCurrent();
-				break;
-			}
+			FWriteScopeLock Locker(ClientsMapLock);
+			// Remove old entries if one is found
+            for (auto It = ClientsMap.CreateIterator(); It; ++It)
+            {
+            	FLiveLinkHubUEClientInfo& IteratedClient = It->Value;
+            	if (IteratedClient.Hostname == Message.ClientInfo.Hostname && IteratedClient.LongName == Message.ClientInfo.LongName
+            		&& IteratedClient.ProjectName == Message.ClientInfo.ProjectName && IteratedClient.CurrentLevel == Message.ClientInfo.CurrentLevel)
+            	{
+            		RemovedAddress = It->Key;
+            		It.RemoveCurrent();
+            		break;
+            	}
+            }
 		}
 
 		ClientsMap.Add(ConnectionAddress, FLiveLinkHubUEClientInfo(Message.ClientInfo, ConnectionAddress));
@@ -97,6 +124,7 @@ private:
 	/** Handle a client info message being received. Happens when new information about a client is received (ie. Client has changed map) */
 	void HandleClientInfoMessage(const FLiveLinkClientInfoMessage& Message, const TSharedRef<IMessageContext, ESPMode::ThreadSafe>& Context)
 	{
+		FWriteScopeLock Locker(ClientsMapLock);
 		FMessageAddress Address = Context->GetSender();
 		if (FLiveLinkHubUEClientInfo* ClientInfo = ClientsMap.Find(Address))
 		{
@@ -110,8 +138,9 @@ protected:
 	//~ Begin FLiveLinkProvider interface
 	virtual void OnConnectionsClosed(const TArray<FMessageAddress>& ClosedAddresses) override
 	{
+		FWriteScopeLock Locker(ClientsMapLock);
 		// todo: If we want to show disconnected clients, we should update the status in the clients map rather than remove it.
-		for (const FMessageAddress& TrackedAddress : ClosedAddresses)
+		for (FMessageAddress TrackedAddress : ClosedAddresses)
 		{
 			if (ClientsMap.Contains(TrackedAddress))
 			{
@@ -128,13 +157,15 @@ protected:
 
 	virtual TArray<FMessageAddress> GetClients() const override
 	{
+		FReadScopeLock Locker(ClientsMapLock);
 		TArray<FMessageAddress> ClientAddresses;
 		ClientsMap.GenerateKeyArray(ClientAddresses);
 		return ClientAddresses;
 	}
 
-	TOptional<FLiveLinkHubUEClientInfo> GetClientInfo(const FMessageAddress& InAddress) const override
+	virtual TOptional<FLiveLinkHubUEClientInfo> GetClientInfo(FMessageAddress InAddress) const override
 	{
+		FReadScopeLock Locker(ClientsMapLock);
 		TOptional<FLiveLinkHubUEClientInfo> ClientInfo;
 		if (const FLiveLinkHubUEClientInfo* ClientInfoPtr = ClientsMap.Find(InAddress))
 		{
@@ -147,6 +178,66 @@ protected:
 	{
 		return OnClientEventDelegate;
 	}
+
+	virtual FText GetClientStatus(FMessageAddress Client) const override
+	{
+		FReadScopeLock Locker(ClientsMapLock);
+		if (const FLiveLinkHubUEClientInfo* ClientInfoPtr = ClientsMap.Find(Client))
+		{
+			return StaticEnum<ELiveLinkClientStatus>()->GetDisplayNameTextByValue(static_cast<int64>(ClientInfoPtr->Status));
+		}
+		
+		return LOCTEXT("InvalidStatus", "Invalid");
+	}
+
+	/** Get whether a client should receive livelink data. */
+	virtual bool IsClientEnabled(FMessageAddress Client) const override
+	{
+		FReadScopeLock Locker(ClientsMapLock);
+		if (const FLiveLinkHubUEClientInfo* ClientInfoPtr = ClientsMap.Find(Client))
+		{
+			return ClientInfoPtr->bEnabled;
+		}
+		return false;
+	}
+
+	/** Set whether a client should receive livelink data. */
+	virtual void SetClientEnabled(FMessageAddress Client, bool bInEnable) override
+	{
+		FWriteScopeLock Locker(ClientsMapLock);
+		if (FLiveLinkHubUEClientInfo* ClientInfoPtr = ClientsMap.Find(Client))
+		{
+			ClientInfoPtr->bEnabled = bInEnable;
+		}
+	}
+
+	/** Get whether a subject is enabled on a given client. */
+	virtual bool IsSubjectEnabled(FMessageAddress Client, const FLiveLinkSubjectKey& Subject) const override
+	{
+		FReadScopeLock Locker(ClientsMapLock);
+		if (const FLiveLinkHubUEClientInfo* ClientInfoPtr = ClientsMap.Find(Client))
+		{
+			return !ClientInfoPtr->DisabledSubjects.Contains(Subject.SubjectName);
+		}
+		return false;
+	}
+
+	/** Set whether a subject should receive livelink data. */
+	virtual void SetSubjectEnabled(FMessageAddress Client, const FLiveLinkSubjectKey& Subject, bool bInEnable) override
+	{
+		FWriteScopeLock Locker(ClientsMapLock);
+		if (FLiveLinkHubUEClientInfo* ClientInfoPtr = ClientsMap.Find(Client))
+		{
+			if (bInEnable)
+			{
+				ClientInfoPtr->DisabledSubjects.Remove(Subject.SubjectName);
+			}
+			else
+			{
+				ClientInfoPtr->DisabledSubjects.Add(Subject.SubjectName);
+			}
+		}
+	}
 	//~ End FLiveLinkProvider interface
 
 private:
@@ -158,4 +249,8 @@ private:
 	FOnClientEvent OnClientEventDelegate;
 	/** Annotations sent with every message from this provider. In our case it's use to disambiguate a livelink hub provider from other livelink providers.*/
 	TMap<FName, FString> Annotations;
+	/** Lock used to access the clients map from different threads. */
+	mutable FRWLock ClientsMapLock;
 };
+
+#undef LOCTEXT_NAMESPACE /*LiveLinkHub.LiveLinkHubProvider*/
diff --git a/Engine/Plugins/Animation/LiveLink/Source/LiveLinkHub/Private/Clients/LiveLinkHubUEClientInfo.h b/Engine/Plugins/Animation/LiveLink/Source/LiveLinkHub/Private/Clients/LiveLinkHubUEClientInfo.h
--- a/Engine/Plugins/Animation/LiveLink/Source/LiveLinkHub/Private/Clients/LiveLinkHubUEClientInfo.h
+++ b/Engine/Plugins/Animation/LiveLink/Source/LiveLinkHub/Private/Clients/LiveLinkHubUEClientInfo.h
@@ -15,15 +15,14 @@ struct FLiveLinkHubUEClientInfo
 
 	FLiveLinkHubUEClientInfo() = default;
 
-	explicit FLiveLinkHubUEClientInfo(const FLiveLinkClientInfoMessage& InClientInfo, const FMessageAddress& InMessageAddress)
-		: MessageAddress(InMessageAddress)
+	explicit FLiveLinkHubUEClientInfo(const FLiveLinkClientInfoMessage& InClientInfo, FMessageAddress InMessageAddress)
+		: MessageAddress(MoveTemp(InMessageAddress))
 		, LongName(InClientInfo.LongName)
 		, Status(InClientInfo.Status)
 		, IPAddress(TEXT("192.168.0.1 (Placeholder)"))
 		, Hostname(InClientInfo.Hostname)
 		, ProjectName(InClientInfo.ProjectName)
 		, CurrentLevel(InClientInfo.CurrentLevel)
-
 	{
 	}
 
@@ -35,7 +34,7 @@ struct FLiveLinkHubUEClientInfo
    	FString LongName;
 	
 	/** Status of the client, ie. is it actively doing a take record at the moment? */
-	UPROPERTY(VisibleAnywhere, Category = "LiveLink Client")
+	UPROPERTY()
 	ELiveLinkClientStatus Status = ELiveLinkClientStatus::Disconnected;
 	
 	UPROPERTY(VisibleAnywhere, Category = "LiveLink Client", DisplayName = "IP Address")
@@ -53,7 +52,11 @@ struct FLiveLinkHubUEClientInfo
 	UPROPERTY(VisibleAnywhere, Category = "LiveLink Client")
 	FString CurrentLevel;
 	
+	/** Subjects that should not be transmitted to this client. */
+	UPROPERTY()
+	TSet<FName> DisabledSubjects;
+	
 	/** Whether this client should receive messages. */
-	UPROPERTY(VisibleAnywhere, Category = "LiveLink Client")
+	UPROPERTY()
 	bool bEnabled = true;
 };
diff --git a/Engine/Plugins/Animation/LiveLink/Source/LiveLinkHub/Private/Clients/SLiveLinkHubClientsView.h b/Engine/Plugins/Animation/LiveLink/Source/LiveLinkHub/Private/Clients/SLiveLinkHubClientsView.h
--- a/Engine/Plugins/Animation/LiveLink/Source/LiveLinkHub/Private/Clients/SLiveLinkHubClientsView.h
+++ b/Engine/Plugins/Animation/LiveLink/Source/LiveLinkHub/Private/Clients/SLiveLinkHubClientsView.h
@@ -7,20 +7,267 @@
 #include "Async/Async.h"
 #include "Delegates/Delegate.h"
 #include "Delegates/DelegateCombinations.h"
+#include "Features/IModularFeatures.h"
 #include "Misc/Guid.h"
+#include "LiveLinkClient.h"
 #include "LiveLinkHubClientsModel.h"
-#include "LiveLinkHubModule.h"
 #include "Styling/SlateTypes.h"
 #include "Widgets/Input/SButton.h"
+#include "Widgets/Input/SCheckBox.h"
 #include "Widgets/Text/STextBlock.h"
-#include "Widgets/Views/SListView.h"
+#include "Widgets/Views/STreeView.h"
 
 struct FGuid;
 class FLiveLinkHub;
 
 #define LOCTEXT_NAMESPACE "LiveLinkHub.ClientsView"
 
-DECLARE_DELEGATE_OneParam(FOnClientSelected, const TSharedPtr<FMessageAddress>&/*ClientIdentifier*/);
+DECLARE_DELEGATE_OneParam(FOnClientSelected, FMessageAddress/*ClientIdentifier*/);
+
+static const FName NameColumnId = "Name";
+static const FName StatusColumnId = "Status";
+static const FName EnabledIconColumnId = "EnabledIcon";
+
+/** Tree view item that represents either a client or a livelink subject. */
+struct FClientTreeViewItem
+{
+	virtual ~FClientTreeViewItem() = default;
+
+	FClientTreeViewItem(FMessageAddress InClientAddress, TSharedRef<ILiveLinkHubClientsModel> InClientsModel)
+		: ClientAddress(MoveTemp(InClientAddress))
+		, ClientsModel(MoveTemp(InClientsModel))
+	{
+	}
+
+	/** Get the subject key for this tree item (Invalid key for client rows). */
+	virtual const FLiveLinkSubjectKey& GetSubjectKey() const
+	{
+		static const FLiveLinkSubjectKey InvalidSubjectKey;
+		return InvalidSubjectKey;
+	}
+
+	/**
+	 * For clients, returns if it should receive any livelink data from the hub (except for heartbeat messages).
+	 * For subjects, returns if the hub transmit this subject's data to the client.
+	 */
+	virtual bool IsEnabled() const = 0;
+
+	/** Whether the row should be in read only (ie. if the source is disconnected) */
+	virtual bool IsReadOnly() const = 0;
+
+	/**
+	 * Set whether this item should be transmitted to the client. 
+	 * @See IsEnabled()
+	 */
+	virtual void SetEnabled(bool bInEnabled) = 0;
+
+	/** Get status text for the row. */
+	virtual FText GetStatusText() const = 0;
+
+	/** This item's children, in the case of client rows, these represent the livelink subjects. */
+	TArray<TSharedPtr<FClientTreeViewItem>> Children;
+	/** Name of the tree item (client or subject name). */
+	FText Name;
+	/** Identifier of the unreal client for this item. */
+	FMessageAddress ClientAddress;
+	/** ClientsModel used to retrieve information about clients/subjects. */
+	TWeakPtr<ILiveLinkHubClientsModel> ClientsModel;
+};
+
+/** Holds a client row's data. */
+struct FClientTreeViewClientItem : public FClientTreeViewItem
+{
+	FClientTreeViewClientItem(FMessageAddress InClientAddress, TSharedRef<ILiveLinkHubClientsModel> InClientsModel)
+		: FClientTreeViewItem(MoveTemp(InClientAddress), MoveTemp(InClientsModel))
+	{
+		if (TSharedPtr<ILiveLinkHubClientsModel> ClientModelPtr = ClientsModel.Pin())
+		{
+			if (TOptional<FLiveLinkHubUEClientInfo> ClientInfo = ClientModelPtr->GetClientInfo(ClientAddress))
+			{
+				Name = ClientInfo ? FText::FromString(ClientInfo->LongName) : LOCTEXT("InvalidSourceLabel", "Invalid Source");
+			}
+			else
+			{
+				ensureMsgf(false, TEXT("Client Info was invalid"));
+			}
+		}
+	}
+
+	virtual bool IsEnabled() const override
+	{
+		if (const TSharedPtr<ILiveLinkHubClientsModel> ClientsModelPtr = ClientsModel.Pin())
+		{
+			return ClientsModelPtr->IsClientEnabled(ClientAddress);
+		}
+		return false;
+	}
+
+	virtual bool IsReadOnly() const override
+	{
+		// todo: when the client list will be modified to show disconnected sources, put this in read only when client is disconnected.
+		return false;
+	}
+
+	virtual void SetEnabled(bool bInEnabled) override
+	{
+		if (const TSharedPtr<ILiveLinkHubClientsModel> ClientsModelPtr = ClientsModel.Pin())
+		{
+			ClientsModelPtr->SetClientEnabled(ClientAddress, bInEnabled);
+		}
+	}
+
+	virtual FText GetStatusText() const override
+	{
+		if (const TSharedPtr<ILiveLinkHubClientsModel> ClientsModelPtr = ClientsModel.Pin())
+		{
+			return ClientsModelPtr->GetClientStatus(ClientAddress);
+		}
+
+		return LOCTEXT("InvalidStatus", "Invalid");
+	}
+	//~ End FClientTreeViewItem interface
+};
+
+/** Holds a subject row's data. */
+struct FClientTreeViewSubjectItem : public FClientTreeViewItem
+{
+	FClientTreeViewSubjectItem(FMessageAddress InClientAddress, FLiveLinkSubjectKey InLiveLinkSubjectKey, TSharedRef<ILiveLinkHubClientsModel> InClientsModel)
+		: FClientTreeViewItem(MoveTemp(InClientAddress), MoveTemp(InClientsModel))
+		, LiveLinkSubjectKey(MoveTemp(InLiveLinkSubjectKey))
+	{
+		const FLiveLinkClient& LiveLinkClient = IModularFeatures::Get().GetModularFeature<FLiveLinkClient>(ILiveLinkClient::ModularFeatureName);
+		Name = FText::Format(LOCTEXT("SubjectName", "{0} - {1}"), FText::FromName(LiveLinkSubjectKey.SubjectName), LiveLinkClient.GetSourceType(LiveLinkSubjectKey.Source));
+	}
+
+	//~ Begin FClientTreeViewItem interface
+	virtual const FLiveLinkSubjectKey& GetSubjectKey() const override
+	{
+		return LiveLinkSubjectKey;
+	}
+
+	virtual bool IsEnabled() const override
+	{
+		if (const TSharedPtr<ILiveLinkHubClientsModel> ClientsModelPtr = ClientsModel.Pin())
+		{
+			return ClientsModelPtr->IsSubjectEnabled(ClientAddress, LiveLinkSubjectKey);
+		}
+		return false;
+	}
+
+	virtual bool IsReadOnly() const override
+	{
+		if (const TSharedPtr<ILiveLinkHubClientsModel> ClientsModelPtr = ClientsModel.Pin())
+		{
+			return !ClientsModelPtr->IsClientEnabled(ClientAddress);
+		}
+
+		return false;
+	}
+
+	virtual void SetEnabled(bool bInEnabled) override
+	{
+		if (const TSharedPtr<ILiveLinkHubClientsModel> ClientsModelPtr = ClientsModel.Pin())
+		{
+			ClientsModelPtr->SetSubjectEnabled(ClientAddress, LiveLinkSubjectKey, bInEnabled);
+		}
+	}
+
+	virtual FText GetStatusText() const override
+	{
+		return FText::GetEmpty();
+	}
+	//~ End FClientTreeViewItem interface
+
+	/** Unique key for this item's subject. */
+	FLiveLinkSubjectKey LiveLinkSubjectKey;
+};
+
+
+class SLiveLinkHubClientsRow : public SMultiColumnTableRow<TSharedPtr<FClientTreeViewItem>>
+{
+public:
+	SLATE_BEGIN_ARGS(SLiveLinkHubClientsRow) {}
+		/** The list item for this row */
+		SLATE_ARGUMENT(TSharedPtr<FClientTreeViewItem>, Item)
+	SLATE_END_ARGS()
+
+	void Construct(const FArguments& InArgs, const TSharedRef<STableViewBase>& InOwnerTableView)
+	{
+		TreeItem = InArgs._Item;
+
+		SMultiColumnTableRow<TSharedPtr<FClientTreeViewItem>>::Construct(
+			FSuperRowType::FArguments()
+			.Padding(1.0f),
+			InOwnerTableView
+		);
+	}
+
+	//~ Begin SMultiColumnTableRow interface
+	virtual TSharedRef<SWidget> GenerateWidgetForColumn(const FName& ColumnName) override
+	{
+		if (ColumnName == NameColumnId)
+		{
+			return SNew(SHorizontalBox)
+				+ SHorizontalBox::Slot()
+				.AutoWidth()
+				.Padding(6, 0, 0, 0)
+				[
+					SNew(SExpanderArrow, SharedThis(this)).IndentAmount(12)
+				]
+				+ SHorizontalBox::Slot()
+				.FillWidth(1.0f)
+				[
+					SNew(STextBlock)
+					.ToolTipText(TreeItem->Name)
+					.Text(TreeItem->Name)
+				];
+		}
+		else if (ColumnName == StatusColumnId)
+		{
+			return SNew(STextBlock)
+				.Text(this, &SLiveLinkHubClientsRow::GetStatusText);
+		}
+		else if (ColumnName == EnabledIconColumnId)
+		{
+			return SNew(SCheckBox)
+				.IsChecked(MakeAttributeSP(this, &SLiveLinkHubClientsRow::IsItemEnabled))
+				.IsEnabled(this, &SLiveLinkHubClientsRow::IsCheckboxEnabled)
+				.OnCheckStateChanged(this, &SLiveLinkHubClientsRow::OnEnabledCheckboxChange);
+		}
+
+		return SNullWidget::NullWidget;
+	}
+	//~ End SMultiColumnTableRow interface
+
+private:
+	/** Return whether the enable checkbox should be clickable. */
+	bool IsCheckboxEnabled() const
+	{
+		return !TreeItem->IsReadOnly();
+	}
+
+	/** Return whether the enabled checkbox is checked. */
+	ECheckBoxState IsItemEnabled() const
+	{
+		return TreeItem->IsEnabled() ? ECheckBoxState::Checked : ECheckBoxState::Unchecked;
+	}
+
+	/** Handler called when the enabled checkbox is clicked. */
+	void OnEnabledCheckboxChange(ECheckBoxState State) const
+	{
+		TreeItem->SetEnabled(State == ECheckBoxState::Checked);
+	}
+
+	/** Get the status text from the tree item. */
+	FText GetStatusText() const
+	{
+		return TreeItem->GetStatusText();
+	}
+
+private:
+	/** The data represented by this row (Either an unreal client or a livelink subject). */
+	TSharedPtr<FClientTreeViewItem> TreeItem;
+};
 
 /**
  * Provides the UI that displays the UE clients connected to the hub. 
@@ -32,28 +279,48 @@ public:
 	SLATE_EVENT(FOnClientSelected, OnClientSelected)
 	SLATE_END_ARGS()
 
-	using FClientIdentifierPtr = TSharedPtr<FMessageAddress>;
+	using FClientTreeItemPtr = TSharedPtr<FClientTreeViewItem>;
 
-	/**
-	* @param InArgs
-	*/
-	void Construct(const FArguments& InArgs, TSharedPtr<ILiveLinkHubClientsModel> InClientsModel)
+	//~ Begin SWidget interface
+	void Construct(const FArguments& InArgs, TSharedRef<ILiveLinkHubClientsModel> InClientsModel)
 	{
 		OnClientSelectedDelegate = InArgs._OnClientSelected;
 		ClientsModel = MoveTemp(InClientsModel);
 		 
 		ClientsModel->OnClientEvent().AddSP(this, &SLiveLinkHubClientsView::OnClientEvent);
 
+		FLiveLinkClient& LiveLinkClient = IModularFeatures::Get().GetModularFeature<FLiveLinkClient>(ILiveLinkClient::ModularFeatureName);
+		LiveLinkClient.OnLiveLinkSubjectAdded().AddSP(this, &SLiveLinkHubClientsView::OnSubjectAdded_AnyThread);
+		LiveLinkClient.OnLiveLinkSubjectRemoved().AddSP(this, &SLiveLinkHubClientsView::OnSubjectRemoved_AnyThread);
+
 		ChildSlot
-			[
-				SAssignNew(ListView, SListView<FClientIdentifierPtr>)
-					.ListItemsSource(&Clients)
-					.OnSelectionChanged(this, &SLiveLinkHubClientsView::OnSelectionChanged)
-					.OnGenerateRow(this, &SLiveLinkHubClientsView::OnGenerateClientRow)
-			];
+		[
+			SAssignNew(TreeView, STreeView<FClientTreeItemPtr>)
+				.TreeItemsSource(&Clients)
+				.ItemHeight(20.0f)
+				.OnSelectionChanged(this, &SLiveLinkHubClientsView::OnSelectionChanged)
+				.OnGenerateRow(this, &SLiveLinkHubClientsView::OnGenerateClientRow)
+				.OnGetChildren(this, &SLiveLinkHubClientsView::OnGetChildren)
+				.HeaderRow
+				(
+					SNew(SHeaderRow)
+					+ SHeaderRow::Column(NameColumnId)
+					.FillWidth(0.75f)
+					.DefaultLabel(LOCTEXT("ItemName", "Name"))
+					+ SHeaderRow::Column(StatusColumnId)
+					.DefaultLabel(LOCTEXT("Status", "Status"))
+					.FillWidth(0.25f)
+					+ SHeaderRow::Column(EnabledIconColumnId)
+					.ManualWidth(20.f)
+					.DefaultLabel(LOCTEXT("EnabledIconEmpty", ""))
+				)
+		];
+
+		PopulateClients();
 	}
+	//~ End SWidget interface
 
-	virtual ~SLiveLinkHubClientsView()
+	virtual ~SLiveLinkHubClientsView() override
 	{
 		if (ClientsModel)
 		{
@@ -61,52 +328,63 @@ public:
 		}
 	}
 
-	/** Get the currently selected client, or nullptr if none is currently selected. */
-	FClientIdentifierPtr GetSelectedClient() const
+	/** Get the currently selected client, an invalid address if none is currently selected. */
+	FMessageAddress GetSelectedClient() const
 	{
-		TArray<FClientIdentifierPtr> SelectedClients = ListView->GetSelectedItems();
-		return SelectedClients.Num() ? SelectedClients[0] : nullptr;
+		TArray<FClientTreeItemPtr> SelectedClients = TreeView->GetSelectedItems();
+		if (SelectedClients.Num())
+		{
+			return SelectedClients[0]->ClientAddress;
+		}
+
+		return FMessageAddress();
 	}
-	
+
+private:
 	/** Handler used to generate a widget for a given client row. */
-	TSharedRef<ITableRow> OnGenerateClientRow(FClientIdentifierPtr Item, const TSharedRef<STableViewBase>& OwnerTable)
+	TSharedRef<ITableRow> OnGenerateClientRow(FClientTreeItemPtr Item, const TSharedRef<STableViewBase>& OwnerTable)
 	{
-		TOptional<FLiveLinkHubUEClientInfo> ClientInfo = ClientsModel->GetClientInfo(*Item);
-		const FText Representation = ClientInfo ? FText::FromString(ClientInfo->LongName) : LOCTEXT("InvalidSourceLabel", "Invalid Source");
-		
-		return SNew(STableRow<TSharedPtr<FString>>, OwnerTable)
-		.Content()
-		[
-			SNew(STextBlock)
-			.Text(Representation)
-		];
+		return SNew(SLiveLinkHubClientsRow, OwnerTable)
+			.Item(Item);
+	}
+
+	/** Handler called to fetch a tree row's children. */
+	void OnGetChildren(FClientTreeItemPtr Item, TArray<FClientTreeItemPtr>& OutChildren)
+	{
+		OutChildren.Append(Item->Children);
 	}
 
-private:
 	/** Handler called when selection changes in the list view. */
-	void OnSelectionChanged(FClientIdentifierPtr InClientInfo, const ESelectInfo::Type InSelectInfoType)
+	void OnSelectionChanged(FClientTreeItemPtr InItem, const ESelectInfo::Type InSelectInfoType)
 	{
-		OnClientSelectedDelegate.ExecuteIfBound(InClientInfo);
+		if (InItem)
+		{
+			OnClientSelectedDelegate.ExecuteIfBound(InItem->ClientAddress);
+		}
 	}
 
 	/** Handler called when the client list has changed. */
-	void OnClientEvent(const FMessageAddress& MessageAddress, ILiveLinkHubClientsModel::EClientEventType EventType)
+	void OnClientEvent(FMessageAddress MessageAddress, ILiveLinkHubClientsModel::EClientEventType EventType)
 	{
 		switch (EventType)
 		{
 		case ILiveLinkHubClientsModel::EClientEventType::Connected:
 		{
-			if (!Clients.ContainsByPredicate([MessageAddress](const FClientIdentifierPtr& InClient) { return *InClient == MessageAddress; }))
+			if (!Clients.ContainsByPredicate([MessageAddress](const FClientTreeItemPtr& InClient) { return InClient->ClientAddress == MessageAddress; }))
 			{
-				Clients.Add(MakeShared<FMessageAddress>(MessageAddress));
-				ListView->RequestListRefresh();
+				TSharedPtr<FClientTreeViewClientItem> ClientItem = MakeShared<FClientTreeViewClientItem>(MessageAddress, ClientsModel.ToSharedRef());
+				ClientItem->ClientAddress = MessageAddress;
+				InitializeClientItem(*ClientItem);
+
+				Clients.Add(ClientItem);
+				TreeView->RequestTreeRefresh();
 			}
 			break;
 		}
 		case ILiveLinkHubClientsModel::EClientEventType::Disconnected:
 		{
-			Clients.RemoveAll([MessageAddress](const FClientIdentifierPtr& InClient) { return *InClient == MessageAddress; });
-			ListView->RequestListRefresh();
+			Clients.RemoveAll([MessageAddress](const FClientTreeItemPtr& InClient) { return InClient->ClientAddress == MessageAddress; });
+			TreeView->RequestTreeRefresh();
 			break;
 		}
 		case ILiveLinkHubClientsModel::EClientEventType::Modified:
@@ -119,15 +397,102 @@ private:
 		}
 		}
 	}
+
+	/** Populate a client item row with its data and children. */
+	void InitializeClientItem(FClientTreeViewClientItem& ClientItem)
+	{
+		const FLiveLinkClient& LiveLinkClient = IModularFeatures::Get().GetModularFeature<FLiveLinkClient>(ILiveLinkClient::ModularFeatureName);
+		constexpr bool bIncludeDisabledSubject = true;
+		constexpr bool bIncludeVirtualSubject = true;
+
+		TArray<FLiveLinkSubjectKey> LiveLinkSubjects = LiveLinkClient.GetSubjects(bIncludeDisabledSubject, bIncludeVirtualSubject);
+		ClientItem.Children.Reserve(LiveLinkSubjects.Num());
+
+		for (const FLiveLinkSubjectKey& SubjectKey : LiveLinkSubjects)
+		{
+			TSharedPtr<FClientTreeViewSubjectItem> SubjectItem = MakeShared<FClientTreeViewSubjectItem>(ClientItem.ClientAddress, SubjectKey, ClientsModel.ToSharedRef());
+			ClientItem.Children.Add(SubjectItem);
+		}
+	}
+
+	/** AnyThread handler for the SubjectAdded delegate, dispatches handling on the game thread to avoid asserts in Slate. */
+	void OnSubjectAdded_AnyThread(FLiveLinkSubjectKey SubjectKey)
+	{
+		TWeakPtr<SLiveLinkHubClientsView> Self = StaticCastSharedRef<SLiveLinkHubClientsView>(AsShared());
+		AsyncTask(ENamedThreads::GameThread, [Self, Key = MoveTemp(SubjectKey)]
+		{
+			if (TSharedPtr<SLiveLinkHubClientsView> View = Self.Pin())
+			{
+				View->OnSubjectAdded(Key);
+			}
+		});
+	}
+
+	/** Handles updating the tree view when a subject is added. */
+	void OnSubjectAdded(const FLiveLinkSubjectKey& SubjectKey)
+	{
+		for (const FClientTreeItemPtr& Client : Clients)
+		{
+			TSharedPtr<FClientTreeViewSubjectItem> SubjectItem = MakeShared<FClientTreeViewSubjectItem>(Client->ClientAddress, SubjectKey, ClientsModel.ToSharedRef());
+			SubjectItem->LiveLinkSubjectKey = SubjectKey;
+			Client->Children.Add(SubjectItem);
+		}
+
+		TreeView->RequestTreeRefresh();
+	}
+
+	/** AnyThread handler for the SubjectRemoved delegate, dispatches handling on the game thread to avoid asserts in Slate. */
+	void OnSubjectRemoved_AnyThread(FLiveLinkSubjectKey SubjectKey)
+	{
+		TWeakPtr<SLiveLinkHubClientsView> Self = StaticCastSharedRef<SLiveLinkHubClientsView>(AsShared());
+		AsyncTask(ENamedThreads::GameThread, [Self, Key = MoveTemp(SubjectKey)]
+		{
+			if (TSharedPtr<SLiveLinkHubClientsView> View = Self.Pin())
+			{
+				View->OnSubjectRemoved(Key);
+			}
+		});
+	}
+
+	/** Handles updating the tree view when a subject is removed. */
+	void OnSubjectRemoved(const FLiveLinkSubjectKey& SubjectKey)
+	{
+		for (const FClientTreeItemPtr& Client : Clients)
+		{
+			int32 Index = Client->Children.IndexOfByPredicate([SubjectKey](const TSharedPtr<FClientTreeViewItem>& Item) { return Item->GetSubjectKey() == SubjectKey; });
+			if (Index != INDEX_NONE)
+			{
+				Client->Children.RemoveAt(Index);
+			}
+		}
+
+		TreeView->RequestTreeRefresh();
+	}
+
+	/** Build the client list. */
+	void PopulateClients()
+	{
+		TArray<FMessageAddress> ClientList = ClientsModel->GetClients();
+		Clients.Reset(ClientList.Num());
+
+		for (FMessageAddress Client : ClientList)
+		{
+			TSharedPtr<FClientTreeViewClientItem> ClientItem = MakeShared<FClientTreeViewClientItem>(Client, ClientsModel.ToSharedRef());
+			ClientItem->ClientAddress = Client;
+			InitializeClientItem(*ClientItem);
+			Clients.Add(ClientItem);
+		}
+	}
+
 private:
 	/** Delegate called when a client is selected. */
 	FOnClientSelected OnClientSelectedDelegate;
-	/** ListView widget that displays the clients. */
-	TSharedPtr<SListView<FClientIdentifierPtr>> ListView;
+	/** TreeView widget that displays the clients. */
+	TSharedPtr<STreeView<FClientTreeItemPtr>> TreeView;
 	/** List of message addresses used to identify UE clients and populate the list view. */
-	TArray<FClientIdentifierPtr> Clients;
+	TArray<FClientTreeItemPtr> Clients;
 	/** Model that holds the client data we are displaying. */
 	TSharedPtr<ILiveLinkHubClientsModel> ClientsModel;
 };
 
-#undef LOCTEXT_NAMESPACE /* LiveLinkHub.ClientsView */
\ No newline at end of file
+#undef LOCTEXT_NAMESPACE /* LiveLinkHub.ClientsView */
diff --git a/Engine/Plugins/Animation/LiveLink/Source/LiveLinkHub/Private/LiveLinkHubClient.cpp b/Engine/Plugins/Animation/LiveLink/Source/LiveLinkHub/Private/LiveLinkHubClient.cpp
--- a/Engine/Plugins/Animation/LiveLink/Source/LiveLinkHub/Private/LiveLinkHubClient.cpp
+++ b/Engine/Plugins/Animation/LiveLink/Source/LiveLinkHub/Private/LiveLinkHubClient.cpp
@@ -234,6 +234,8 @@ void FLiveLinkHubClient::PushSubjectStaticData_AnyThread(const FLiveLinkSubjectK
 
 	check(Collection);
 
+	// Check if subject should receive data
+
 	if (!FLiveLinkRoleTrait::Validate(Role, InStaticData))
 	{
 		if (Role == nullptr)
@@ -378,6 +380,7 @@ void FLiveLinkHubClient::PushSubjectStaticData_AnyThread(const FLiveLinkSubjectK
 	if (LiveLinkSubject)
 	{
 		// Dispatch it on the game thread since FLiveLinkSubject::SetStaticData asserts when called outside the game thread.
+		// Note that SetStaticData may not be needed on the hub but we might need it for interpolation or preprocessing in the future.
 		FFunctionGraphTask::CreateAndDispatchWhenReady([this, LiveLinkSubject, SubjectKey, Role, StaticData = MoveTemp(InStaticData)]() mutable
 			{
 				LiveLinkSubject->SetStaticData(Role, MoveTemp(StaticData));
@@ -414,4 +417,4 @@ bool FLiveLinkHubClient::IsSubjectValid(const FLiveLinkSubjectKey& InSubjectKey)
 	return false;
 }
 
-#undef LOCTEXT_NAMESPACE
\ No newline at end of file
+#undef LOCTEXT_NAMESPACE
diff --git a/Engine/Plugins/Animation/LiveLink/Source/LiveLinkHubMessaging/Private/LiveLinkHubConnectionManager.h b/Engine/Plugins/Animation/LiveLink/Source/LiveLinkHubMessaging/Private/LiveLinkHubConnectionManager.h
--- a/Engine/Plugins/Animation/LiveLink/Source/LiveLinkHubMessaging/Private/LiveLinkHubConnectionManager.h
+++ b/Engine/Plugins/Animation/LiveLink/Source/LiveLinkHubMessaging/Private/LiveLinkHubConnectionManager.h
@@ -92,6 +92,7 @@ private:
 				}
 
 				const FString* ProviderType = PollResult->Annotations.Find(FLiveLinkHubMessageAnnotation::ProviderTypeAnnotation);
+
 				if (ProviderType && *ProviderType == UE::LiveLinkHub::Private::LiveLinkHubProviderType)
 				{
 					AddLiveLinkSource(PollResult);
@@ -111,6 +112,14 @@ private:
 		{
 			ILiveLinkClient* LiveLinkClient = &ModularFeatures.GetModularFeature<ILiveLinkClient>(ILiveLinkClient::ModularFeatureName);
 
+			for (const FGuid& SourceId : LiveLinkClient->GetSources())
+			{
+				if (LiveLinkClient->GetSourceMachineName(SourceId).ToString() == PollResult->MachineName)
+				{
+					LiveLinkClient->RemoveSource(SourceId);
+				}
+			}
+
 			TSharedPtr<ILiveLinkSource> LiveLinkSource = MakeShared<FLiveLinkHubMessageBusSource>(FText::FromString(PollResult->Name), FText::FromString(PollResult->MachineName), PollResult->Address, PollResult->MachineTimeOffset);
 			LiveLinkClient->AddSource(LiveLinkSource);
 
diff --git a/Engine/Source/Runtime/LiveLinkInterface/Public/ILiveLinkClient.h b/Engine/Source/Runtime/LiveLinkInterface/Public/ILiveLinkClient.h
--- a/Engine/Source/Runtime/LiveLinkInterface/Public/ILiveLinkClient.h
+++ b/Engine/Source/Runtime/LiveLinkInterface/Public/ILiveLinkClient.h
@@ -162,6 +162,9 @@ public:
 	/** Get the status of a source */
 	virtual FText GetSourceStatus(FGuid EntryGuid) const = 0;
 
+	/** Get the machine name of the source. */
+	virtual FText GetSourceMachineName(FGuid EntryGuid) const = 0;
+
 	/** Push static data for a specific subject for a certain role. This will clear all buffered frames */
 	virtual void PushSubjectStaticData_AnyThread(const FLiveLinkSubjectKey& SubjectKey, TSubclassOf<ULiveLinkRole> Role, FLiveLinkStaticDataStruct&& StaticData) = 0;
 
diff --git a/Engine/Source/Runtime/LiveLinkMessageBusFramework/Private/LiveLinkProvider.cpp b/Engine/Source/Runtime/LiveLinkMessageBusFramework/Private/LiveLinkProvider.cpp
--- a/Engine/Source/Runtime/LiveLinkMessageBusFramework/Private/LiveLinkProvider.cpp
+++ b/Engine/Source/Runtime/LiveLinkMessageBusFramework/Private/LiveLinkProvider.cpp
@@ -4,6 +4,7 @@
 #include "LiveLinkProviderImpl.h"
 
 #include "Algo/RemoveIf.h"
+#include "Algo/Transform.h"
 #include "HAL/PlatformProcess.h"
 #include "IMessageContext.h"
 #include "LiveLinkMessages.h"
@@ -129,7 +130,7 @@ void FLiveLinkProvider::SendSubject(FName SubjectName, const FTrackedSubject& Su
 	SubjectData->SubjectName = SubjectName;
 
 	TArray<FMessageAddress> Addresses;
-	GetConnectedAddresses(Addresses);
+	GetFilteredAddresses(SubjectName, Addresses);
 
 	MessageEndpoint->Send(SubjectData, FLiveLinkSubjectDataMessage::StaticStruct(), EMessageFlags::None, GetAnnotations(), nullptr, Addresses, FTimespan::Zero(), FDateTime::MaxValue());
 }
@@ -145,7 +146,7 @@ void FLiveLinkProvider::SendSubjectFrame(FName SubjectName, const FTrackedSubjec
 	SubjectFrame->Time = Subject.Time;
 
 	TArray<FMessageAddress> Addresses;
-	GetConnectedAddresses(Addresses);
+	GetFilteredAddresses(SubjectName, Addresses);
 
 	MessageEndpoint->Send(SubjectFrame, FLiveLinkSubjectFrameMessage::StaticStruct(), EMessageFlags::None, GetAnnotations(), nullptr, Addresses, FTimespan::Zero(), FDateTime::MaxValue());
 }
@@ -207,7 +208,7 @@ void FLiveLinkProvider::ClearTrackedSubject(const FName& SubjectName)
 void FLiveLinkProvider::SendClearSubjectToConnections(FName SubjectName)
 {
 	TArray<FMessageAddress> MessageAddresses;
-	GetConnectedAddresses(MessageAddresses);
+	GetFilteredAddresses(SubjectName, MessageAddresses);
 
 	MessageEndpoint->Send(FMessageEndpoint::MakeMessage<FLiveLinkClearSubject>(SubjectName), EMessageFlags::Reliable, GetAnnotations(), nullptr, MessageAddresses, FTimespan::Zero(), FDateTime::MaxValue());
 }
@@ -284,11 +285,7 @@ bool FLiveLinkProvider::UpdateSubjectStaticData(const FName SubjectName, TSubcla
 	if (ConnectedAddresses.Num() > 0)
 	{
 		TArray<FMessageAddress> Addresses;
-		Addresses.Reserve(ConnectedAddresses.Num());
-		for (const FTrackedAddress& Address : ConnectedAddresses)
-		{
-			Addresses.Add(Address.Address);
-		}
+		GetFilteredAddresses(SubjectName, Addresses);
 
 		TMap<FName, FString> Annotations;
 		Annotations.Add(FLiveLinkMessageAnnotation::SubjectAnnotation, SubjectName.ToString());
@@ -376,11 +373,7 @@ bool FLiveLinkProvider::UpdateSubjectFrameData(const FName SubjectName, FLiveLin
 	if (ConnectedAddresses.Num() > 0)
 	{
 		TArray<FMessageAddress> Addresses;
-		Addresses.Reserve(ConnectedAddresses.Num());
-		for (const FTrackedAddress& Address : ConnectedAddresses)
-		{
-			Addresses.Add(Address.Address);
-		}
+		GetFilteredAddresses(SubjectName, Addresses);
 
 		TMap<FName, FString> Annotations;
 		Annotations.Add(FLiveLinkMessageAnnotation::SubjectAnnotation, SubjectName.ToString());
@@ -523,3 +516,13 @@ void FLiveLinkProvider::GetConnectedAddresses(TArray<FMessageAddress>& Addresses
 		Addresses.Add(Address.Address);
 	}
 }
+
+void FLiveLinkProvider::GetFilteredAddresses(FName SubjectName, TArray<FMessageAddress>& Addresses)
+{
+	ValidateConnections();
+	Addresses.Reserve(ConnectedAddresses.Num());
+
+	Algo::TransformIf(ConnectedAddresses, Addresses,
+		[this, SubjectName](const FTrackedAddress& Address){ return ShouldTransmitToSubject_AnyThread(SubjectName, Address.Address); },
+		[](const FTrackedAddress& Address){ return Address.Address; });
+}
diff --git a/Engine/Source/Runtime/LiveLinkMessageBusFramework/Public/LiveLinkProviderImpl.h b/Engine/Source/Runtime/LiveLinkMessageBusFramework/Public/LiveLinkProviderImpl.h
--- a/Engine/Source/Runtime/LiveLinkMessageBusFramework/Public/LiveLinkProviderImpl.h
+++ b/Engine/Source/Runtime/LiveLinkMessageBusFramework/Public/LiveLinkProviderImpl.h
@@ -77,8 +77,12 @@ private:
 	// Clear a existing track subject
 	void ClearTrackedSubject(const FName& SubjectName);
 
+	// Send a clear subject message to indicate that the subject should be removed from the connected client.
 	void SendClearSubjectToConnections(FName SubjectName);
 
+	// Get the connected addresses that should receive livelink data.
+	void GetFilteredAddresses(FName SubjectName, TArray<FMessageAddress>& Addresses);
+
 protected:
 	// Update connected addresses and send information to the connected source
 	void HandleConnectMessage(const FLiveLinkConnectMessage& Message, const TSharedRef<class IMessageContext, ESPMode::ThreadSafe>& Context);
@@ -147,6 +151,12 @@ protected:
 		return {};
 	}
 
+	// Get whether a combination of a subject/client should receive livelink data.
+	virtual bool ShouldTransmitToSubject_AnyThread(FName SubjectName, FMessageAddress Address) const
+	{
+		return true;
+	}
+
 	// Constructor for derived classes that allows specifying that no endpoint should be created.
 	FLiveLinkProvider(const FString& InProviderName, bool bInCreateEndpoint);
 
