Prevent deadlock by narrowing the scope of the write lock.

diff --git a/Engine/Plugins/Animation/LiveLink/Source/LiveLinkHub/Private/Clients/LiveLinkHubProvider.h b/Engine/Plugins/Animation/LiveLink/Source/LiveLinkHub/Private/Clients/LiveLinkHubProvider.h
--- a/Engine/Plugins/Animation/LiveLink/Source/LiveLinkHub/Private/Clients/LiveLinkHubProvider.h
+++ b/Engine/Plugins/Animation/LiveLink/Source/LiveLinkHub/Private/Clients/LiveLinkHubProvider.h
@@ -138,16 +138,24 @@ protected:
 	//~ Begin FLiveLinkProvider interface
 	virtual void OnConnectionsClosed(const TArray<FMessageAddress>& ClosedAddresses) override
 	{
-		FWriteScopeLock Locker(ClientsMapLock);
-		// todo: If we want to show disconnected clients, we should update the status in the clients map rather than remove it.
-		for (FMessageAddress TrackedAddress : ClosedAddresses)
+		TArray<FMessageAddress> Notifications;
 		{
-			if (ClientsMap.Contains(TrackedAddress))
+			FWriteScopeLock Locker(ClientsMapLock);
+			// todo: If we want to show disconnected clients, we should update the status in the clients map rather than remove it.
+			for (FMessageAddress TrackedAddress : ClosedAddresses)
 			{
-				ClientsMap.Remove(TrackedAddress);
-				OnClientEventDelegate.Broadcast(TrackedAddress, EClientEventType::Modified);
+				if (ClientsMap.Contains(TrackedAddress))
+				{
+					ClientsMap.Remove(TrackedAddress);
+					Notifications.Add(TrackedAddress);
+				}
 			}
 		}
+
+		for (FMessageAddress TrackedAddress : Notifications)
+		{
+			OnClientEventDelegate.Broadcast(TrackedAddress, EClientEventType::Modified);
+		}
 	}
 
 	virtual TMap<FName, FString> GetAnnotations() const override
