LiveLinkHub: Support loading and saving the live link hub config to json.

diff --git a/Engine/Plugins/Animation/LiveLink/Source/LiveLinkHub/LiveLinkHub.Build.cs b/Engine/Plugins/Animation/LiveLink/Source/LiveLinkHub/LiveLinkHub.Build.cs
--- a/Engine/Plugins/Animation/LiveLink/Source/LiveLinkHub/LiveLinkHub.Build.cs
+++ b/Engine/Plugins/Animation/LiveLink/Source/LiveLinkHub/LiveLinkHub.Build.cs
@@ -23,6 +23,8 @@ namespace UnrealBuildTool.Rules
 				"ContentBrowserData",
 				"Engine",
 				"InputCore",
+				"Json",
+				"JsonUtilities",
 				"LiveLink",
 				"LiveLinkEditor",
 				"LiveLinkHubMessaging",
diff --git a/Engine/Plugins/Animation/LiveLink/Source/LiveLinkHub/Private/Clients/LiveLinkHubProvider.h b/Engine/Plugins/Animation/LiveLink/Source/LiveLinkHub/Private/Clients/LiveLinkHubProvider.h
--- a/Engine/Plugins/Animation/LiveLink/Source/LiveLinkHub/Private/Clients/LiveLinkHubProvider.h
+++ b/Engine/Plugins/Animation/LiveLink/Source/LiveLinkHub/Private/Clients/LiveLinkHubProvider.h
@@ -84,6 +84,23 @@ public:
 		return true;
 	}
 
+	/** Manually add a client to the client map. */
+	void AddClient(const FLiveLinkHubUEClientInfo& InClientInfo, const FMessageAddress& InMessageAddress)
+	{
+		ClientsMap.Add(InMessageAddress, InClientInfo);
+		OnClientEventDelegate.Broadcast(InMessageAddress, EClientEventType::Connected);
+	}
+
+	/** Manually remove a client from the client map. */
+	void RemoveClient(const FMessageAddress& InMessageAddress)
+	{
+		ClientsMap.Remove(InMessageAddress);
+		OnClientEventDelegate.Broadcast(InMessageAddress, EClientEventType::Disconnected);
+	}
+
+	/** Retrieve the existing client map. */
+	const TMap<FMessageAddress, FLiveLinkHubUEClientInfo>& GetClientsMap() const { return ClientsMap; }
+
 private:
 	/** Handle a connection message resulting from a livelink hub message bus source connecting to this provider. */
 	void HandleHubConnectMessage(const FLiveLinkHubConnectMessage& Message, const TSharedRef<IMessageContext, ESPMode::ThreadSafe>& Context)
diff --git a/Engine/Plugins/Animation/LiveLink/Source/LiveLinkHub/Private/Clients/SLiveLinkHubClientsView.h b/Engine/Plugins/Animation/LiveLink/Source/LiveLinkHub/Private/Clients/SLiveLinkHubClientsView.h
--- a/Engine/Plugins/Animation/LiveLink/Source/LiveLinkHub/Private/Clients/SLiveLinkHubClientsView.h
+++ b/Engine/Plugins/Animation/LiveLink/Source/LiveLinkHub/Private/Clients/SLiveLinkHubClientsView.h
@@ -11,6 +11,7 @@
 #include "Misc/Guid.h"
 #include "LiveLinkClient.h"
 #include "LiveLinkHubClientsModel.h"
+#include "LiveLinkHubUEClientInfo.h"
 #include "Styling/SlateTypes.h"
 #include "Widgets/Input/SButton.h"
 #include "Widgets/Input/SCheckBox.h"
@@ -435,7 +436,14 @@ private:
 		{
 			TSharedPtr<FClientTreeViewSubjectItem> SubjectItem = MakeShared<FClientTreeViewSubjectItem>(Client->ClientAddress, SubjectKey, ClientsModel.ToSharedRef());
 			SubjectItem->LiveLinkSubjectKey = SubjectKey;
-			Client->Children.Add(SubjectItem);
+
+			if (!Client->Children.ContainsByPredicate([&](const TSharedPtr<FClientTreeViewItem>& Child)
+			{
+				return Child->GetSubjectKey() == SubjectKey;
+			}))
+			{
+				Client->Children.Add(SubjectItem);
+			}
 		}
 
 		TreeView->RequestTreeRefresh();
diff --git a/Engine/Plugins/Animation/LiveLink/Source/LiveLinkHub/Private/Config/LiveLinkHubConfigData.h b/Engine/Plugins/Animation/LiveLink/Source/LiveLinkHub/Private/Config/LiveLinkHubConfigData.h
new file mode 100644
--- /dev/null
+++ b/Engine/Plugins/Animation/LiveLink/Source/LiveLinkHub/Private/Config/LiveLinkHubConfigData.h
@@ -0,0 +1,29 @@
+// Copyright Epic Games, Inc. All Rights Reserved.
+
+#pragma once
+
+#include "Clients/LiveLinkHubUEClientInfo.h"
+#include "CoreTypes.h"
+#include "LiveLinkPresetTypes.h"
+#include "Misc/Guid.h"
+
+#include "LiveLinkHubConfigData.generated.h"
+
+/** Live link hub config data which can be serialized to disk. */
+USTRUCT()
+struct FLiveLinkHubConfigData
+{
+	GENERATED_BODY()
+
+	/** Live link hub sources. */
+	UPROPERTY()
+	TArray<FLiveLinkSourcePreset> Sources;
+
+	/** Live link hub subjects. */
+	UPROPERTY()
+	TArray<FLiveLinkSubjectPreset> Subjects;
+
+	/** Live link hub client info. */
+	UPROPERTY()
+	TArray<FLiveLinkHubUEClientInfo> Clients;
+};
\ No newline at end of file
diff --git a/Engine/Plugins/Animation/LiveLink/Source/LiveLinkHub/Private/Config/LiveLinkHubFileUtilities.cpp b/Engine/Plugins/Animation/LiveLink/Source/LiveLinkHub/Private/Config/LiveLinkHubFileUtilities.cpp
new file mode 100644
--- /dev/null
+++ b/Engine/Plugins/Animation/LiveLink/Source/LiveLinkHub/Private/Config/LiveLinkHubFileUtilities.cpp
@@ -0,0 +1,83 @@
+// Copyright Epic Games, Inc. All Rights Reserved.
+
+#include "LiveLinkHubFileUtilities.h"
+
+#include "HAL/FileManager.h"
+#include "JsonObjectConverter.h"
+#include "LiveLinkHubConfigData.h"
+#include "LiveLinkHubLog.h"
+
+void UE::LiveLinkHub::FileUtilities::Private::SaveConfig(const FLiveLinkHubConfigData& InConfigData, const FString& InFilePath)
+{
+	if (ensure(!InFilePath.IsEmpty()))
+	{
+		const TUniquePtr<FArchive> Ar(IFileManager::Get().CreateFileWriter(*InFilePath));
+		if (Ar)
+		{
+			const TSharedPtr<FJsonObject> JsonObject = ToJson(InConfigData);
+			
+			const TSharedRef<TJsonWriter<>> JsonWriter = TJsonWriterFactory<>::Create(Ar.Get(), 0);
+			FJsonSerializer::Serialize(JsonObject.ToSharedRef(), JsonWriter);
+
+			ensure(Ar->Close());
+		}
+	}
+}
+
+TSharedPtr<FLiveLinkHubConfigData> UE::LiveLinkHub::FileUtilities::Private::LoadConfig(const FString& InFilePath)
+{
+	if (IFileManager::Get().FileExists(*InFilePath))
+	{
+		const TUniquePtr<FArchive> Ar(IFileManager::Get().CreateFileReader(*InFilePath));
+		if (Ar)
+		{
+			const TSharedRef<TJsonReader<>> JsonReader = TJsonReaderFactory<>::Create(Ar.Get());
+
+			TSharedPtr<FJsonObject> JsonObject = MakeShared<FJsonObject>();
+			if (FJsonSerializer::Deserialize(JsonReader, JsonObject) && JsonObject.IsValid())
+			{
+				int32 Version = 0;
+				if (JsonObject->TryGetNumberField(JsonVersionKey, Version) &&
+					Version <= LiveLinkHubVersion)
+				{
+					return FromJson(JsonObject);
+				}
+
+				UE_LOG(LogLiveLinkHub, Error, TEXT("Could not load config %s because the version field %s was missing or invalid."),
+					*InFilePath, *JsonVersionKey);
+			}
+		}
+	}
+
+	return nullptr;
+}
+
+TSharedPtr<FJsonObject> UE::LiveLinkHub::FileUtilities::Private::ToJson(const FLiveLinkHubConfigData& InConfigData)
+{
+	TSharedRef<FJsonObject> JsonObject = MakeShared<FJsonObject>();
+	const TSharedPtr<FJsonValueNumber> NumberValue = MakeShared<FJsonValueNumber>(LiveLinkHubVersion);
+	JsonObject->SetField(JsonVersionKey, NumberValue);
+	
+	FJsonObjectConverter::UStructToJsonObject(FLiveLinkHubConfigData::StaticStruct(), &InConfigData, JsonObject);
+	
+	return JsonObject;
+}
+
+TSharedPtr<FLiveLinkHubConfigData> UE::LiveLinkHub::FileUtilities::Private::FromJson(const TSharedPtr<FJsonObject>& InJsonObject)
+{
+	if (InJsonObject.IsValid())
+	{
+		FLiveLinkHubConfigData OutConfigData;
+		const bool bResult =
+			FJsonObjectConverter::JsonObjectToUStruct(InJsonObject.ToSharedRef(),
+				FLiveLinkHubConfigData::StaticStruct(), &OutConfigData);
+
+		if (!bResult)
+		{
+			UE_LOG(LogLiveLinkHub, Error, TEXT("Could not convert from json to LiveLinkHubConfigData."))
+		}
+		return MakeShared<FLiveLinkHubConfigData>(OutConfigData);
+	}
+
+	return nullptr;
+}
diff --git a/Engine/Plugins/Animation/LiveLink/Source/LiveLinkHub/Private/Config/LiveLinkHubFileUtilities.h b/Engine/Plugins/Animation/LiveLink/Source/LiveLinkHub/Private/Config/LiveLinkHubFileUtilities.h
new file mode 100644
--- /dev/null
+++ b/Engine/Plugins/Animation/LiveLink/Source/LiveLinkHub/Private/Config/LiveLinkHubFileUtilities.h
@@ -0,0 +1,36 @@
+// Copyright Epic Games, Inc. All Rights Reserved.
+
+#pragma once
+
+#include "Containers/UnrealString.h"
+#include "Templates/SharedPointer.h"
+
+struct FLiveLinkHubConfigData;
+class FJsonObject;
+
+namespace UE::LiveLinkHub::FileUtilities::Private
+{
+	/** A key which must be present in json configs and mapped to the saved json version. */
+	const FString JsonVersionKey = TEXT("liveLinkHub");
+	/** The json version we support. */
+	constexpr int32 LiveLinkHubVersion = 1;
+
+	/** The extension of the config file. */
+	const FString ConfigExtension = TEXT("json");
+	/** The default name of the config file. */
+	const FString ConfigDefaultFileName = TEXT("LiveLinkHubConfig");
+	/** The description of the config file. */
+	const FString ConfigDescription = TEXT("Live Link Hub Config");
+
+	/** Save config data to disk. */
+	void SaveConfig(const FLiveLinkHubConfigData& InConfigData, const FString& InFilePath);
+
+	/** Load config data from disk. */
+	TSharedPtr<FLiveLinkHubConfigData> LoadConfig(const FString& InFilePath);
+
+	/** Convert config data to json. */
+	TSharedPtr<FJsonObject> ToJson(const FLiveLinkHubConfigData& InConfigData);
+
+	/** Convert config data from json. */
+	TSharedPtr<FLiveLinkHubConfigData> FromJson(const TSharedPtr<FJsonObject>& InJsonObject);
+}
\ No newline at end of file
diff --git a/Engine/Plugins/Animation/LiveLink/Source/LiveLinkHub/Private/LiveLinkHub.cpp b/Engine/Plugins/Animation/LiveLink/Source/LiveLinkHub/Private/LiveLinkHub.cpp
--- a/Engine/Plugins/Animation/LiveLink/Source/LiveLinkHub/Private/LiveLinkHub.cpp
+++ b/Engine/Plugins/Animation/LiveLink/Source/LiveLinkHub/Private/LiveLinkHub.cpp
@@ -4,8 +4,15 @@
 
 #include "Clients/LiveLinkHubClientsController.h"
 #include "Clients/LiveLinkHubProvider.h"
+#include "Config/LiveLinkHubConfigData.h"
+#include "Config/LiveLinkHubFileUtilities.h"
+#include "DesktopPlatformModule.h"
+#include "EditorDirectories.h"
 #include "Features/IModularFeatures.h"
+#include "Framework/Application/SlateApplication.h"
+#include "IDesktopPlatform.h"
 #include "LiveLinkHubClient.h"
+#include "LiveLinkHubCommands.h"
 #include "LiveLinkProvider.h"
 #include "LiveLinkSubject.h"
 #include "LiveLinkSubjectSettings.h"
@@ -31,7 +38,11 @@ void FLiveLinkHub::Initialize()
 	PlaybackController = MakeShared<FLiveLinkHubPlaybackController>();
 	RecordingListController = MakeShared<FLiveLinkHubRecordingListController>(AsShared());
 	ClientsController = MakeShared<FLiveLinkHubClientsController>(LiveLinkProvider.ToSharedRef());
+	CommandList = MakeShared<FUICommandList>();
 
+	FLiveLinkHubCommands::Register();
+	BindCommands();
+	
 	FString LiveLinkHubLayoutIni = GConfig->GetConfigFilename(TEXT("LiveLinkHubLayout"));
 	WindowController = MakeShared<FLiveLinkHubWindowController>(FLiveLinkHubWindowInitParams{ LiveLinkHubLayoutIni });
 	WindowController->RestoreLayout();
@@ -143,4 +154,159 @@ void FLiveLinkHub::OnSubjectRemoved(FLiveLinkSubjectKey InSubjectKey)
 	LiveLinkProvider->RemoveSubject(InSubjectKey.SubjectName);
 }
 
+void FLiveLinkHub::BindCommands()
+{
+	const FLiveLinkHubCommands& Commands = FLiveLinkHubCommands::Get();
+	CommandList->MapAction(Commands.NewConfig, FExecuteAction::CreateSP(this, &FLiveLinkHub::NewConfig));
+	CommandList->MapAction(Commands.OpenConfig, FExecuteAction::CreateSP(this, &FLiveLinkHub::OpenConfig));
+	CommandList->MapAction(Commands.SaveConfigAs, FExecuteAction::CreateSP(this, &FLiveLinkHub::SaveConfigAs));
+	CommandList->MapAction(Commands.SaveConfig, FExecuteAction::CreateSP(this, &FLiveLinkHub::SaveConfig),
+		FCanExecuteAction::CreateSP(this, &FLiveLinkHub::CanSaveConfig));
+}
+
+void FLiveLinkHub::ClearClient()
+{
+	check(LiveLinkProvider.IsValid());
+	check(LiveLinkHubClient.IsValid());
+	
+	TMap<FMessageAddress, FLiveLinkHubUEClientInfo> Clients = LiveLinkProvider->GetClientsMap();
+	
+	for (const TTuple<FMessageAddress, FLiveLinkHubUEClientInfo>& Client : Clients)
+	{
+		LiveLinkProvider->RemoveClient(Client.Key);
+	}
+	
+	LiveLinkHubClient->RemoveAllSources();
+
+	// Removing sources only sets bPendingKill to true, need to tick to make sure they are removed.
+	LiveLinkHubClient->ForceTick();
+}
+
+void FLiveLinkHub::NewConfig()
+{
+	ClearClient();
+	LastConfigPath.Empty();
+}
+
+void FLiveLinkHub::SaveConfigAs()
+{
+	const FString FileDescription = UE::LiveLinkHub::FileUtilities::Private::ConfigDescription;
+	const FString Extensions = UE::LiveLinkHub::FileUtilities::Private::ConfigExtension;
+	const FString FileTypes = FString::Printf(TEXT("%s (*.%s)|*.%s"), *FileDescription, *Extensions, *Extensions);
+
+	const FString DefaultFile = UE::LiveLinkHub::FileUtilities::Private::ConfigDefaultFileName;
+	
+	TArray<FString> SaveFileNames;
+
+	IDesktopPlatform* DesktopPlatform = FDesktopPlatformModule::Get();
+	const void* ParentWindowWindowHandle = FSlateApplication::Get().FindBestParentWindowHandleForDialogs(nullptr);
+	
+	const bool bFileSelected = DesktopPlatform->SaveFileDialog(
+		ParentWindowWindowHandle,
+		LOCTEXT("LiveLinkHubSaveAsTitle", "Save As").ToString(),
+		FEditorDirectories::Get().GetLastDirectory(ELastDirectory::GENERIC_SAVE),
+		DefaultFile,
+		FileTypes,
+		EFileDialogFlags::None,
+		SaveFileNames);
+
+	if (bFileSelected && SaveFileNames.Num() > 0)
+	{
+		LastConfigPath = SaveFileNames[0];
+		FEditorDirectories::Get().SetLastDirectory(ELastDirectory::GENERIC_SAVE, FPaths::GetPath(LastConfigPath));
+		SaveConfig();
+	}
+}
+
+bool FLiveLinkHub::CanSaveConfig() const
+{
+	return !LastConfigPath.IsEmpty();
+}
+
+void FLiveLinkHub::SaveConfig()
+{
+	if (LastConfigPath.IsEmpty())
+	{
+		return;
+	}
+	
+	check(LiveLinkProvider.IsValid())
+	check(LiveLinkHubClient.IsValid());
+	
+	FLiveLinkHubConfigData LiveLinkHubConfigData;
+	
+	TArray<FGuid> SourceGuids = LiveLinkHubClient->GetSources();
+	for (const FGuid& SourceGuid : SourceGuids)
+	{
+		LiveLinkHubConfigData.Sources.Add(LiveLinkHubClient->GetSourcePreset(SourceGuid, nullptr));
+	}
+
+	TArray<FLiveLinkSubjectKey> Subjects = LiveLinkHubClient->GetSubjects(true, true);
+	for (const FLiveLinkSubjectKey& Subject : Subjects)
+	{
+		LiveLinkHubConfigData.Subjects.Add(LiveLinkHubClient->GetSubjectPreset(Subject, nullptr));
+	}
+
+	const TMap<FMessageAddress, FLiveLinkHubUEClientInfo>& ClientMap = LiveLinkProvider->GetClientsMap();
+	
+	for (const TTuple<FMessageAddress, FLiveLinkHubUEClientInfo>& ClientKeyVal : ClientMap)
+	{
+		LiveLinkHubConfigData.Clients.Add(ClientKeyVal.Value);
+	}
+
+	UE::LiveLinkHub::FileUtilities::Private::SaveConfig(LiveLinkHubConfigData, LastConfigPath);
+}
+
+void FLiveLinkHub::OpenConfig()
+{
+	const FString FileDescription = UE::LiveLinkHub::FileUtilities::Private::ConfigDescription;
+	const FString Extensions = UE::LiveLinkHub::FileUtilities::Private::ConfigExtension;
+	const FString FileTypes = FString::Printf(TEXT("%s (*.%s)|*.%s"), *FileDescription, *Extensions, *Extensions);
+
+	const FString DefaultFile = UE::LiveLinkHub::FileUtilities::Private::ConfigDefaultFileName;
+
+	check(LiveLinkProvider.IsValid())
+	check(LiveLinkHubClient.IsValid());
+	
+	ClearClient();
+	
+	TArray<FString> OpenFileNames;
+
+	IDesktopPlatform* DesktopPlatform = FDesktopPlatformModule::Get();
+	const void* ParentWindowWindowHandle = FSlateApplication::Get().FindBestParentWindowHandleForDialogs(nullptr);
+	const bool bFileSelected = DesktopPlatform->OpenFileDialog(
+		ParentWindowWindowHandle,
+		LOCTEXT("LiveLinkHubOpenTitle", "Open").ToString(),
+		FEditorDirectories::Get().GetLastDirectory(ELastDirectory::GENERIC_OPEN),
+		DefaultFile,
+		FileTypes,
+		EFileDialogFlags::None,
+		OpenFileNames);
+
+	if (bFileSelected && OpenFileNames.Num() > 0)
+	{
+		LastConfigPath = OpenFileNames[0];
+		FEditorDirectories::Get().SetLastDirectory(ELastDirectory::GENERIC_OPEN, FPaths::GetPath(LastConfigPath));
+
+		const TSharedPtr<FLiveLinkHubConfigData> ConfigData = UE::LiveLinkHub::FileUtilities::Private::LoadConfig(LastConfigPath);
+		if (ConfigData.IsValid())
+		{
+			for (const FLiveLinkSourcePreset& SourcePreset : ConfigData->Sources)
+			{
+				LiveLinkHubClient->CreateSource(SourcePreset);
+			}
+
+			for (const FLiveLinkSubjectPreset& SubjectPreset : ConfigData->Subjects)
+			{
+				LiveLinkHubClient->CreateSubject(SubjectPreset);	
+			}
+
+			for (const FLiveLinkHubUEClientInfo& Client : ConfigData->Clients)
+			{
+				LiveLinkProvider->AddClient(Client, FMessageAddress::NewAddress());
+			}
+		}
+	}
+}
+
 #undef LOCTEXT_NAMESPACE /*LiveLinkHub*/
diff --git a/Engine/Plugins/Animation/LiveLink/Source/LiveLinkHub/Private/LiveLinkHub.h b/Engine/Plugins/Animation/LiveLink/Source/LiveLinkHub/Private/LiveLinkHub.h
--- a/Engine/Plugins/Animation/LiveLink/Source/LiveLinkHub/Private/LiveLinkHub.h
+++ b/Engine/Plugins/Animation/LiveLink/Source/LiveLinkHub/Private/LiveLinkHub.h
@@ -63,6 +63,8 @@ public:
 	TSharedPtr<FLiveLinkHubPlaybackController> GetPlaybackController() const;
 	/** Get the controller that manages clients. */
 	TSharedPtr<FLiveLinkHubClientsController> GetClientsController() const;
+	/** Get the live link hub command list. */
+	TSharedPtr<FUICommandList> GetCommandList() const { return CommandList; }
 	
 private:
 	//~ LiveLink Client delegates
@@ -72,6 +74,27 @@ private:
 	void OnSubjectRemoved(FLiveLinkSubjectKey InSubjectKey);
 	//~ LiveLink Client delegates
 
+	/** Bind all available live link hub commands. */
+	void BindCommands();
+
+	/** Clear all client settings. */
+	void ClearClient();
+
+	/** Create a new config. */
+	void NewConfig();
+
+	/** Save an existing config to a new file. */
+	void SaveConfigAs();
+
+	/** Whether the Save command can be used. */
+	bool CanSaveConfig() const;
+	
+	/** Save the config to the current file. */
+	void SaveConfig();
+
+	/** Open an existing config. */
+	void OpenConfig();
+	
 private:
 	/** Implements the logic to manage the clients tabs. */
 	TSharedPtr<FLiveLinkHubClientsController> ClientsController;
@@ -87,8 +110,12 @@ private:
 	TSharedPtr<FLiveLinkHubClient> LiveLinkHubClient;
 	/** LiveLinkProvider used to transfer data to connected UE clients. */
 	TSharedPtr<FLiveLinkHubProvider> LiveLinkProvider;
-	/** Handles execution of commands */
+	/** Handles execution of commands. */
 	TUniquePtr<FConsoleCommandExecutor> CommandExecutor;
+	/** Available live link hub commands. */
+	TSharedPtr<FUICommandList> CommandList;
+	/** The last opened config path. */
+	FString LastConfigPath;
 
 	friend class FLiveLinkHubModule;
 };
diff --git a/Engine/Plugins/Animation/LiveLink/Source/LiveLinkHub/Private/LiveLinkHubCommands.h b/Engine/Plugins/Animation/LiveLink/Source/LiveLinkHub/Private/LiveLinkHubCommands.h
new file mode 100644
--- /dev/null
+++ b/Engine/Plugins/Animation/LiveLink/Source/LiveLinkHub/Private/LiveLinkHubCommands.h
@@ -0,0 +1,54 @@
+// Copyright Epic Games, Inc. All Rights Reserved.
+
+#pragma once
+
+#include "CoreTypes.h"
+#include "Styling/AppStyle.h"
+#include "Framework/Commands/Commands.h"
+#include "Framework/Commands/InputChord.h"
+#include "Internationalization/Text.h"
+#include "Templates/SharedPointer.h"
+
+#define LOCTEXT_NAMESPACE "LiveLinkHubCommands"
+
+/**
+ * The Live Link Hub commands.
+ */
+class FLiveLinkHubCommands
+	: public TCommands<FLiveLinkHubCommands>
+{
+public:
+
+	/** Default constructor. */
+	FLiveLinkHubCommands()
+		: TCommands<FLiveLinkHubCommands>(
+			"LiveLinkHub",
+			NSLOCTEXT("Contexts", "LiveLinkHub", "Live Link Hub"),
+			NAME_None, FAppStyle::GetAppStyleSetName()
+		)
+	{ }
+
+public:
+
+	//~ TCommands interface
+
+	virtual void RegisterCommands() override
+	{
+		UI_COMMAND(NewConfig, "New Config", "Create a new config", EUserInterfaceActionType::Button, FInputChord(EKeys::N, EModifierKey::Control));
+		UI_COMMAND(OpenConfig, "Open Config...", "Open an existing config", EUserInterfaceActionType::Button, FInputChord(EKeys::O, EModifierKey::Control));
+		UI_COMMAND(SaveConfigAs, "Save Config As...", "Save a config to a new file", EUserInterfaceActionType::Button, FInputChord(EModifierKey::Control | EModifierKey::Alt, EKeys::S));
+		UI_COMMAND(SaveConfig, "Save Config", "Save the current config", EUserInterfaceActionType::Button, FInputChord(EKeys::S, EModifierKey::Control));
+	}
+
+public:
+	/** Start a new config. */
+	TSharedPtr<FUICommandInfo> NewConfig;
+	/** Open an existing config. */
+	TSharedPtr<FUICommandInfo> OpenConfig;
+	/** Save the current configuration to a new file. */
+	TSharedPtr<FUICommandInfo> SaveConfigAs;
+	/** Save the current configuration to an existing file. */
+	TSharedPtr<FUICommandInfo> SaveConfig;
+};
+
+#undef LOCTEXT_NAMESPACE
\ No newline at end of file
diff --git a/Engine/Plugins/Animation/LiveLink/Source/LiveLinkHub/Private/UI/Widgets/LiveLinkHubMainTabController.cpp b/Engine/Plugins/Animation/LiveLink/Source/LiveLinkHub/Private/UI/Widgets/LiveLinkHubMainTabController.cpp
--- a/Engine/Plugins/Animation/LiveLink/Source/LiveLinkHub/Private/UI/Widgets/LiveLinkHubMainTabController.cpp
+++ b/Engine/Plugins/Animation/LiveLink/Source/LiveLinkHub/Private/UI/Widgets/LiveLinkHubMainTabController.cpp
@@ -32,7 +32,7 @@ TSharedRef<SDockTab> FLiveLinkHubMainTabController::SpawnMainTab(const FSpawnTab
 		.CanEverClose(false);
 
 	DockTab->SetContent(
-		SAssignNew(MainTabView, SLiveLinkHubMainTabView, LiveLinkHubTabID)
+		SAssignNew(MainTabView, SLiveLinkHubMainTabView)
 			.ConstructUnderMajorTab(DockTab)
 			.ConstructUnderWindow(RootWindow)
 	);
diff --git a/Engine/Plugins/Animation/LiveLink/Source/LiveLinkHub/Private/UI/Widgets/SLiveLinkHubMainTabView.cpp b/Engine/Plugins/Animation/LiveLink/Source/LiveLinkHub/Private/UI/Widgets/SLiveLinkHubMainTabView.cpp
--- a/Engine/Plugins/Animation/LiveLink/Source/LiveLinkHub/Private/UI/Widgets/SLiveLinkHubMainTabView.cpp
+++ b/Engine/Plugins/Animation/LiveLink/Source/LiveLinkHub/Private/UI/Widgets/SLiveLinkHubMainTabView.cpp
@@ -38,7 +38,7 @@ const FText SLiveLinkHubMainTabView::PlaybackTabName = LOCTEXT("PlaybackTabLabel
 const FText SLiveLinkHubMainTabView::ClientsTabName = LOCTEXT("ClientsTabLabel", "Clients");
 const FText SLiveLinkHubMainTabView::ClientDetailsTabName = LOCTEXT("ClientDetailsTabLabel", "Client Details");
 
-void SLiveLinkHubMainTabView::Construct(const FArguments& InArgs, FName InStatusBarID)
+void SLiveLinkHubMainTabView::Construct(const FArguments& InArgs)
 {	
 	PanelController = MakeShared<FLiveLinkPanelController>();
 
diff --git a/Engine/Plugins/Animation/LiveLink/Source/LiveLinkHub/Private/UI/Widgets/SLiveLinkHubMainTabView.h b/Engine/Plugins/Animation/LiveLink/Source/LiveLinkHub/Private/UI/Widgets/SLiveLinkHubMainTabView.h
--- a/Engine/Plugins/Animation/LiveLink/Source/LiveLinkHub/Private/UI/Widgets/SLiveLinkHubMainTabView.h
+++ b/Engine/Plugins/Animation/LiveLink/Source/LiveLinkHub/Private/UI/Widgets/SLiveLinkHubMainTabView.h
@@ -35,7 +35,7 @@ public:
 		SLATE_ARGUMENT(TSharedPtr<SWindow>, ConstructUnderWindow)
 	SLATE_END_ARGS()
 
-	void Construct(const FArguments& InArgs, FName InStatusBarID);
+	void Construct(const FArguments& InArgs);
 	
 private:
 	/** Create all livelink hub tabs. */
diff --git a/Engine/Plugins/Animation/LiveLink/Source/LiveLinkHub/Private/UI/Widgets/SLiveLinkHubTabViewWithManagerBase.cpp b/Engine/Plugins/Animation/LiveLink/Source/LiveLinkHub/Private/UI/Widgets/SLiveLinkHubTabViewWithManagerBase.cpp
--- a/Engine/Plugins/Animation/LiveLink/Source/LiveLinkHub/Private/UI/Widgets/SLiveLinkHubTabViewWithManagerBase.cpp
+++ b/Engine/Plugins/Animation/LiveLink/Source/LiveLinkHub/Private/UI/Widgets/SLiveLinkHubTabViewWithManagerBase.cpp
@@ -6,6 +6,10 @@
 #include "Framework/Commands/UICommandList.h"
 #include "Framework/Docking/LayoutService.h"
 #include "Framework/MultiBox/MultiBoxBuilder.h"
+#include "LiveLinkHub.h"
+#include "LiveLinkHubCommands.h"
+#include "LiveLinkHubModule.h"
+#include "Modules/ModuleManager.h"
 #include "UI/Widgets/SLiveLinkHubTabViewBase.h"
 
 #define LOCTEXT_NAMESPACE "LiveLinkHub.SLiveLinkHubTabViewWithManagerBase"
@@ -44,7 +48,10 @@ TSharedRef<SWidget> SLiveLinkHubTabViewWithManagerBase::CreateTabs(const FArgume
 	);
 
 	TSharedRef<SWidget> Result = TabManager->RestoreFrom(Layout, InArgs._ConstructUnderWindow).ToSharedRef();
-	FMenuBarBuilder MenuBarBuilder = FMenuBarBuilder(TSharedPtr<FUICommandList>());
+	
+	const TSharedPtr<FLiveLinkHub> LiveLinkHub = FModuleManager::Get().GetModuleChecked<FLiveLinkHubModule>("LiveLinkHub").GetLiveLinkHub();
+	check(LiveLinkHub.IsValid());
+	FMenuBarBuilder MenuBarBuilder = FMenuBarBuilder(LiveLinkHub->GetCommandList());
 	FillInDefaultMenuItems(MenuBarBuilder);
 	InArgs._CreateMenuBar.ExecuteIfBound(MenuBarBuilder);
 	
@@ -57,6 +64,13 @@ TSharedRef<SWidget> SLiveLinkHubTabViewWithManagerBase::CreateTabs(const FArgume
 
 void SLiveLinkHubTabViewWithManagerBase::FillInDefaultMenuItems(FMenuBarBuilder MenuBarBuilder)
 {
+	MenuBarBuilder.AddPullDownMenu(
+		LOCTEXT("FileMenuLabel", "File"),
+		FText::GetEmpty(),
+		FNewMenuDelegate::CreateSP(this, &SLiveLinkHubTabViewWithManagerBase::FillFileMenu),
+		"File"
+	);
+	
 	MenuBarBuilder.AddPullDownMenu(
 		LOCTEXT("WindowMenuLabel", "Window"),
 		FText::GetEmpty(),
@@ -65,6 +79,19 @@ void SLiveLinkHubTabViewWithManagerBase::FillInDefaultMenuItems(FMenuBarBuilder
 	);
 }
 
+void SLiveLinkHubTabViewWithManagerBase::FillFileMenu(FMenuBuilder& MenuBuilder)
+{
+	MenuBuilder.BeginSection("Open", LOCTEXT("OpenHeader", "Open"));
+	MenuBuilder.AddMenuEntry(FLiveLinkHubCommands::Get().NewConfig);
+	MenuBuilder.AddMenuEntry(FLiveLinkHubCommands::Get().OpenConfig);
+	MenuBuilder.EndSection();
+	
+	MenuBuilder.BeginSection("Save", LOCTEXT("SaveHeader", "Save"));
+	MenuBuilder.AddMenuEntry(FLiveLinkHubCommands::Get().SaveConfig);
+	MenuBuilder.AddMenuEntry(FLiveLinkHubCommands::Get().SaveConfigAs);
+	MenuBuilder.EndSection();
+}
+
 void SLiveLinkHubTabViewWithManagerBase::FillWindowMenu(FMenuBuilder& MenuBuilder)
 {
 	TabManager->PopulateLocalTabSpawnerMenu(MenuBuilder);
diff --git a/Engine/Plugins/Animation/LiveLink/Source/LiveLinkHub/Private/UI/Widgets/SLiveLinkHubTabViewWithManagerBase.h b/Engine/Plugins/Animation/LiveLink/Source/LiveLinkHub/Private/UI/Widgets/SLiveLinkHubTabViewWithManagerBase.h
--- a/Engine/Plugins/Animation/LiveLink/Source/LiveLinkHub/Private/UI/Widgets/SLiveLinkHubTabViewWithManagerBase.h
+++ b/Engine/Plugins/Animation/LiveLink/Source/LiveLinkHub/Private/UI/Widgets/SLiveLinkHubTabViewWithManagerBase.h
@@ -33,7 +33,6 @@ public:
 
 	/**
 	 * @param InArgs
-	 * @param InStatusBarId Unique ID needed for the status bar
 	 */
 	void Construct(const FArguments& InArgs);
 
@@ -49,6 +48,8 @@ private:
 	TSharedRef<SWidget> CreateTabs(const FArguments& InArgs);
 	/** Creates the top bar menu items. */
 	void FillInDefaultMenuItems(FMenuBarBuilder MenuBarBuilder);
-	/** Populates the tab spawner's menu. */
+	/** Populates the tab spawner's file menu. */
+	void FillFileMenu(FMenuBuilder& MenuBuilder);
+	/** Populates the tab spawner's window menu. */
 	void FillWindowMenu(FMenuBuilder& MenuBuilder);
 };
diff --git a/Engine/Source/Runtime/LiveLinkInterface/Public/LiveLinkPresetTypes.h b/Engine/Source/Runtime/LiveLinkInterface/Public/LiveLinkPresetTypes.h
--- a/Engine/Source/Runtime/LiveLinkInterface/Public/LiveLinkPresetTypes.h
+++ b/Engine/Source/Runtime/LiveLinkInterface/Public/LiveLinkPresetTypes.h
@@ -22,7 +22,7 @@ struct FLiveLinkSourcePreset
 	UPROPERTY(VisibleAnywhere, Category="LiveLinkSourcePresets")
 	FGuid Guid;
 
-	UPROPERTY(VisibleAnywhere, Category = "LiveLinkSourcePresets")
+	UPROPERTY(VisibleAnywhere, Category = "LiveLinkSourcePresets", Instanced)
 	TObjectPtr<ULiveLinkSourceSettings> Settings = nullptr;
 
 	/** The SourceType when the source was saved to a Preset. */
@@ -42,10 +42,10 @@ struct FLiveLinkSubjectPreset
 	UPROPERTY(VisibleAnywhere, Category = "LiveLinkSubjectPresets")
 	TSubclassOf<ULiveLinkRole> Role;
 
-	UPROPERTY(VisibleAnywhere, Category = "LiveLinkSubjectPresets")
+	UPROPERTY(VisibleAnywhere, Category = "LiveLinkSubjectPresets", Instanced)
 	TObjectPtr<ULiveLinkSubjectSettings> Settings = nullptr;
 
-	UPROPERTY(VisibleAnywhere, Category = "LiveLinkSubjectPresets")
+	UPROPERTY(VisibleAnywhere, Category = "LiveLinkSubjectPresets", Instanced)
 	TObjectPtr<ULiveLinkVirtualSubject> VirtualSubject = nullptr;
 
 	UPROPERTY(VisibleAnywhere, Category = "LiveLinkSubjectPresets")
