Add subject renaming to livelink hub

diff --git a/Engine/Plugins/Animation/LiveLink/Source/LiveLinkEditor/Internal/LiveLinkPanelController.h b/Engine/Plugins/Animation/LiveLink/Source/LiveLinkEditor/Internal/LiveLinkPanelController.h
--- a/Engine/Plugins/Animation/LiveLink/Source/LiveLinkEditor/Internal/LiveLinkPanelController.h
+++ b/Engine/Plugins/Animation/LiveLink/Source/LiveLinkEditor/Internal/LiveLinkPanelController.h
@@ -2,12 +2,15 @@
 
 #pragma once
 
+#include "Delegates/Delegate.h"
 #include "Delegates/IDelegateInstance.h"
+#include "Delegates/DelegateCombinations.h"
 #include "Templates/SharedPointer.h"
 
 class IDetailsView;
 class FLiveLinkClient;
 struct FLiveLinkSourceUIEntry;
+struct FLiveLinkSubjectKey;
 struct FLiveLinkSubjectUIEntry;
 class FLiveLinkSourcesView;
 class SLiveLinkSourceListView;
@@ -19,6 +22,7 @@ namespace ESelectInfo { enum Type : int; }
 typedef TSharedPtr<FLiveLinkSourceUIEntry> FLiveLinkSourceUIEntryPtr;
 typedef TSharedPtr<FLiveLinkSubjectUIEntry> FLiveLinkSubjectUIEntryPtr;
 
+
 /** Handles callback connections between the sources, subjects and details views. */
 class LIVELINKEDITOR_API FLiveLinkPanelController : public TSharedFromThis<FLiveLinkPanelController>
 {
@@ -26,6 +30,13 @@ public:
 	FLiveLinkPanelController();
 	~FLiveLinkPanelController();
 
+	DECLARE_MULTICAST_DELEGATE_OneParam(FOnSubjectSelectionChanged, const FLiveLinkSubjectKey&);
+	/** Subject Selection changed callback. */
+	FOnSubjectSelectionChanged& OnSubjectSelectionChanged()
+	{
+		return SubjectSelectionChangedDelegate;
+	}
+
 private:
 	// Bind live link commands 
 	void BindCommands();
@@ -50,9 +61,9 @@ private:
 	// Recreates the subject list data behind the tree view.
 	void RebuildSubjectList();
 	// Handles source selection changing.
-	void OnSourceSelectionChanged(FLiveLinkSourceUIEntryPtr Entry, ESelectInfo::Type SelectionType) const;
+	void OnSourceSelectionChangedHandler(FLiveLinkSourceUIEntryPtr Entry, ESelectInfo::Type SelectionType) const;
 	// Hadnles subject selection changing.
-	void OnSubjectSelectionChanged(FLiveLinkSubjectUIEntryPtr SubjectEntry, ESelectInfo::Type SelectInfo);
+	void OnSubjectSelectionChangedHandler(FLiveLinkSubjectUIEntryPtr SubjectEntry, ESelectInfo::Type SelectInfo);
 
 public:
 	// Sources view
@@ -73,4 +84,6 @@ public:
 	mutable bool bSelectionChangedGuard = false;
 	// Command list
 	TSharedPtr<FUICommandList> CommandList;
+	// Delegate called when the subject selection changes
+	FOnSubjectSelectionChanged SubjectSelectionChangedDelegate;
 };
\ No newline at end of file
diff --git a/Engine/Plugins/Animation/LiveLink/Source/LiveLinkEditor/Private/LiveLinkPanelController.cpp b/Engine/Plugins/Animation/LiveLink/Source/LiveLinkEditor/Private/LiveLinkPanelController.cpp
--- a/Engine/Plugins/Animation/LiveLink/Source/LiveLinkEditor/Private/LiveLinkPanelController.cpp
+++ b/Engine/Plugins/Animation/LiveLink/Source/LiveLinkEditor/Private/LiveLinkPanelController.cpp
@@ -24,8 +24,8 @@ FLiveLinkPanelController::FLiveLinkPanelController()
 	CommandList = MakeShared<FUICommandList>();
 	BindCommands();
 
-	SourcesView = MakeShared<FLiveLinkSourcesView>(Client, CommandList, FLiveLinkSourcesView::FOnSourceSelectionChanged::CreateRaw(this, &FLiveLinkPanelController::OnSourceSelectionChanged));
-	SubjectsView = MakeShared<FLiveLinkSubjectsView>(FLiveLinkSubjectsView::FOnSubjectSelectionChanged::CreateRaw(this, &FLiveLinkPanelController::OnSubjectSelectionChanged), CommandList);
+	SourcesView = MakeShared<FLiveLinkSourcesView>(Client, CommandList, FLiveLinkSourcesView::FOnSourceSelectionChanged::CreateRaw(this, &FLiveLinkPanelController::OnSourceSelectionChangedHandler));
+	SubjectsView = MakeShared<FLiveLinkSubjectsView>(FLiveLinkSubjectsView::FOnSubjectSelectionChanged::CreateRaw(this, &FLiveLinkPanelController::OnSubjectSelectionChangedHandler), CommandList);
 	SourcesDetailsView = UE::LiveLink::CreateSourcesDetailsView(SourcesView);
 	SubjectsDetailsView = UE::LiveLink::CreateSubjectsDetailsView(Client);
 
@@ -45,7 +45,7 @@ FLiveLinkPanelController::~FLiveLinkPanelController()
 	}
 }
 
-void FLiveLinkPanelController::OnSourceSelectionChanged(FLiveLinkSourceUIEntryPtr Entry, ESelectInfo::Type SelectionType) const
+void FLiveLinkPanelController::OnSourceSelectionChangedHandler(FLiveLinkSourceUIEntryPtr Entry, ESelectInfo::Type SelectionType) const
 {
 	if (bSelectionChangedGuard)
 	{
@@ -80,7 +80,7 @@ void FLiveLinkPanelController::OnSourceSelectionChanged(FLiveLinkSourceUIEntryPt
 	}
 }
 
-void FLiveLinkPanelController::OnSubjectSelectionChanged(FLiveLinkSubjectUIEntryPtr SubjectEntry, ESelectInfo::Type SelectInfo)
+void FLiveLinkPanelController::OnSubjectSelectionChangedHandler(FLiveLinkSubjectUIEntryPtr SubjectEntry, ESelectInfo::Type SelectInfo)
 {
 	if (bSelectionChangedGuard)
 	{
@@ -111,6 +111,8 @@ void FLiveLinkPanelController::OnSubjectSelectionChanged(FLiveLinkSubjectUIEntry
 			SubjectsDetailsView->SetSubjectKey(SubjectEntry->SubjectKey);
 		}
 		bDetailViewSet = true;
+
+		SubjectSelectionChangedDelegate.Broadcast(SubjectEntry->SubjectKey);
 	}
 
 	if (!bDetailViewSet)
@@ -157,7 +159,7 @@ void FLiveLinkPanelController::BindCommands()
 
 void FLiveLinkPanelController::OnSourcesChangedHandler()
 {
-	// Since  this can be called from any thread, make sure we only update slate on the game thread.
+	// Since this can be called from any thread, make sure we only update slate on the game thread.
 	FFunctionGraphTask::CreateAndDispatchWhenReady([this]()
 	{
 		SourcesView->RefreshSourceData(true);
@@ -167,7 +169,7 @@ void FLiveLinkPanelController::OnSourcesChangedHandler()
 
 void FLiveLinkPanelController::OnSubjectsChangedHandler()
 {
-	// Since  this can be called from any thread, make sure we only update slate on the game thread.
+	// Since this can be called from any thread, make sure we only update slate on the game thread.
 	FFunctionGraphTask::CreateAndDispatchWhenReady([this]()
 	{
 		RebuildSubjectList();
diff --git a/Engine/Plugins/Animation/LiveLink/Source/LiveLinkHub/Private/Clients/LiveLinkHubProvider.h b/Engine/Plugins/Animation/LiveLink/Source/LiveLinkHub/Private/Clients/LiveLinkHubProvider.h
--- a/Engine/Plugins/Animation/LiveLink/Source/LiveLinkHub/Private/Clients/LiveLinkHubProvider.h
+++ b/Engine/Plugins/Animation/LiveLink/Source/LiveLinkHub/Private/Clients/LiveLinkHubProvider.h
@@ -10,7 +10,7 @@
 #include "CoreMinimal.h"
 #include "Delegates/DelegateCombinations.h"
 #include "Editor.h"
-#include "Engine/TimerHandle.h" 
+#include "Engine/TimerHandle.h"
 #include "GameThreadMessageHandler.h"
 #include "HAL/CriticalSection.h"
 #include "IMessageContext.h"
@@ -19,8 +19,9 @@
 #include "LiveLinkProviderImpl.h"
 #include "LiveLinkSettings.h"
 #include "MessageEndpointBuilder.h"
-#include "MessageHandlers.h"
 #include "Misc/ScopeLock.h"
+#include "Subjects/LiveLinkHubSubjectSessionConfig.h"
+#include "Session/LiveLinkHubSession.h"
 #include "TimerManager.h"
 
 #define LOCTEXT_NAMESPACE "LiveLinkHub.LiveLinkHubProvider"
@@ -31,6 +32,9 @@
 class FLiveLinkHubProvider : public FLiveLinkProvider, public ILiveLinkHubClientsModel, public TSharedFromThis<FLiveLinkHubProvider>
 {
 public:
+	using FLiveLinkProvider::SendClearSubjectToConnections;
+	using FLiveLinkProvider::GetLastSubjectStaticDataStruct;
+
 	/**
 	 * Create a message bus handler that will dispatch messages on the game thread. 
 	 * This is useful to receive some messages on AnyThread and delegate others on the game thread (ie. for methods that will trigger UI updates which need to happen on game thread. )
@@ -41,8 +45,9 @@ public:
 		return MakeShared<TGameThreadMessageHandler<MessageType, FLiveLinkHubProvider>>(this, Func);
 	}
 
-	FLiveLinkHubProvider()
+	FLiveLinkHubProvider(const TSharedRef<ILiveLinkHubSessionManager>& InSessionManager)
 		: FLiveLinkProvider(TEXT("LiveLink Hub"), false)
+		, SessionManager(InSessionManager)
 	{
 		Annotations.Add(FLiveLinkHubMessageAnnotation::ProviderTypeAnnotation, UE::LiveLinkHub::Private::LiveLinkHubProviderType.ToString());
 
@@ -56,7 +61,7 @@ public:
 		GEditor->GetTimerManager()->SetTimer(ValidateConnectionsTimer, FTimerDelegate::CreateRaw(this, &FLiveLinkHubProvider::ValidateConnections), ValidateConnectionsRate, true);
 	}
 
-	virtual ~FLiveLinkHubProvider()
+	virtual ~FLiveLinkHubProvider() override
 	{
 		if (GEditor)
 		{
@@ -274,6 +279,8 @@ private:
 	FOnClientEvent OnClientEventDelegate;
 	/** Annotations sent with every message from this provider. In our case it's use to disambiguate a livelink hub provider from other livelink providers.*/
 	TMap<FName, FString> Annotations;
+	/** LiveLinkHub session manager. */
+	TWeakPtr<ILiveLinkHubSessionManager> SessionManager;
 	/** Lock used to access the clients map from different threads. */
 	mutable FRWLock ClientsMapLock;
 };
diff --git a/Engine/Plugins/Animation/LiveLink/Source/LiveLinkHub/Private/Clients/SLiveLinkHubClientsView.h b/Engine/Plugins/Animation/LiveLink/Source/LiveLinkHub/Private/Clients/SLiveLinkHubClientsView.h
--- a/Engine/Plugins/Animation/LiveLink/Source/LiveLinkHub/Private/Clients/SLiveLinkHubClientsView.h
+++ b/Engine/Plugins/Animation/LiveLink/Source/LiveLinkHub/Private/Clients/SLiveLinkHubClientsView.h
@@ -323,6 +323,13 @@ public:
 
 	virtual ~SLiveLinkHubClientsView() override
 	{
+		if (IModularFeatures::Get().IsModularFeatureAvailable(ILiveLinkClient::ModularFeatureName))
+		{
+			FLiveLinkClient& LiveLinkClient = IModularFeatures::Get().GetModularFeature<FLiveLinkClient>(ILiveLinkClient::ModularFeatureName);
+			LiveLinkClient.OnLiveLinkSubjectRemoved().RemoveAll(this);
+			LiveLinkClient.OnLiveLinkSubjectAdded().RemoveAll(this);
+		}
+
 		if (ClientsModel)
 		{
 			ClientsModel->OnClientEvent().RemoveAll(this);
diff --git a/Engine/Plugins/Animation/LiveLink/Source/LiveLinkHub/Private/LiveLinkHub.cpp b/Engine/Plugins/Animation/LiveLink/Source/LiveLinkHub/Private/LiveLinkHub.cpp
--- a/Engine/Plugins/Animation/LiveLink/Source/LiveLinkHub/Private/LiveLinkHub.cpp
+++ b/Engine/Plugins/Animation/LiveLink/Source/LiveLinkHub/Private/LiveLinkHub.cpp
@@ -12,33 +12,40 @@
 #include "Framework/Application/SlateApplication.h"
 #include "IDesktopPlatform.h"
 #include "LiveLinkHubClient.h"
-#include "LiveLinkHubCommands.h"
 #include "LiveLinkProvider.h"
 #include "LiveLinkSubject.h"
+#include "LiveLinkHubCommands.h"
+#include "LiveLinkSubject.h"
+#include "LiveLinkProviderImpl.h"
 #include "LiveLinkSubjectSettings.h"
 #include "Misc/ConfigCacheIni.h"
 #include "Recording/LiveLinkHubPlaybackController.h"
 #include "Recording/LiveLinkHubRecordingController.h"
 #include "Recording/LiveLinkHubRecordingListController.h"
+#include "Subjects/LiveLinkHubSubjectController.h"
 #include "UI/Window/LiveLinkHubWindowController.h"
 
 #define LOCTEXT_NAMESPACE "LiveLinkHub"
 
 void FLiveLinkHub::Initialize()
 {
-	LiveLinkProvider = MakeShared<FLiveLinkHubProvider>();
+	// We must register the livelink client first since we might rely on the modular feature to initialize the controllers/managers.
 	LiveLinkHubClient = MakeShared<FLiveLinkHubClient>(AsShared());
+	IModularFeatures::Get().RegisterModularFeature(ILiveLinkClient::ModularFeatureName, LiveLinkHubClient.Get());
+
+	SessionManager = MakeShared<FLiveLinkHubSessionManager>();
+	LiveLinkProvider = MakeShared<FLiveLinkHubProvider>(SessionManager.ToSharedRef());
 
 	CommandExecutor = MakeUnique<FConsoleCommandExecutor>();
 	IModularFeatures::Get().RegisterModularFeature(IConsoleCommandExecutor::ModularFeatureName(), CommandExecutor.Get());
-	
-	IModularFeatures::Get().RegisterModularFeature(ILiveLinkClient::ModularFeatureName, LiveLinkHubClient.Get());
 
 	RecordingController = MakeShared<FLiveLinkHubRecordingController>();
 	PlaybackController = MakeShared<FLiveLinkHubPlaybackController>();
 	RecordingListController = MakeShared<FLiveLinkHubRecordingListController>(AsShared());
 	ClientsController = MakeShared<FLiveLinkHubClientsController>(LiveLinkProvider.ToSharedRef());
 	CommandList = MakeShared<FUICommandList>();
+	SubjectController = MakeShared<FLiveLinkHubSubjectController>();
+
 
 	FLiveLinkHubCommands::Register();
 	BindCommands();
@@ -95,6 +102,11 @@ TSharedPtr<FLiveLinkHubClientsController> FLiveLinkHub::GetClientsController() c
 	return ClientsController;
 }
 
+TSharedPtr<ILiveLinkHubSessionManager> FLiveLinkHub::GetSessionManager() const
+{
+	return SessionManager;
+}
+
 TSharedPtr<FLiveLinkHubRecordingController> FLiveLinkHub::GetRecordingController() const
 {
 	return RecordingController;
@@ -110,7 +122,7 @@ TSharedPtr<FLiveLinkHubPlaybackController> FLiveLinkHub::GetPlaybackController()
 	return PlaybackController;
 }
 
-void FLiveLinkHub::OnStaticDataReceived_AnyThread(const FLiveLinkSubjectKey& InSubjectKey, TSubclassOf<ULiveLinkRole> InRole, const FLiveLinkStaticDataStruct& InStaticDataStruct)
+void FLiveLinkHub::OnStaticDataReceived_AnyThread(const FLiveLinkSubjectKey& InSubjectKey, TSubclassOf<ULiveLinkRole> InRole, const FLiveLinkStaticDataStruct& InStaticDataStruct) const
 {
 	if (RecordingController->IsRecording())
 	{
@@ -119,10 +131,12 @@ void FLiveLinkHub::OnStaticDataReceived_AnyThread(const FLiveLinkSubjectKey& InS
 
 	FLiveLinkStaticDataStruct StaticDataCopy;
 	StaticDataCopy.InitializeWith(InStaticDataStruct);
-	LiveLinkProvider->UpdateSubjectStaticData(InSubjectKey.SubjectName, InRole, MoveTemp(StaticDataCopy));
+
+	const FName OverridenName = GetSubjectNameOverride(InSubjectKey);
+	LiveLinkProvider->UpdateSubjectStaticData(OverridenName, InRole, MoveTemp(StaticDataCopy));
 }
 
-void FLiveLinkHub::OnFrameDataReceived_AnyThread(const FLiveLinkSubjectKey& InSubjectKey, const FLiveLinkFrameDataStruct& InFrameDataStruct)
+void FLiveLinkHub::OnFrameDataReceived_AnyThread(const FLiveLinkSubjectKey& InSubjectKey, const FLiveLinkFrameDataStruct& InFrameDataStruct) const
 {
 	if (RecordingController->IsRecording())
 	{
@@ -132,10 +146,12 @@ void FLiveLinkHub::OnFrameDataReceived_AnyThread(const FLiveLinkSubjectKey& InSu
 	
 	FLiveLinkFrameDataStruct FrameDataCopy;
 	FrameDataCopy.InitializeWith(InFrameDataStruct);
-	LiveLinkProvider->UpdateSubjectFrameData(InSubjectKey.SubjectName, MoveTemp(FrameDataCopy));
+
+	const FName OverridenName = GetSubjectNameOverride(InSubjectKey);
+	LiveLinkProvider->UpdateSubjectFrameData(OverridenName, MoveTemp(FrameDataCopy));
 }
 
-void FLiveLinkHub::OnSubjectAdded(FLiveLinkSubjectKey InSubjectKey)
+void FLiveLinkHub::OnSubjectAdded(FLiveLinkSubjectKey InSubjectKey) const
 {
 	// Send an update to connected clients as well.
 	ULiveLinkSubjectSettings* SubjectSettings = Cast<ULiveLinkSubjectSettings>(LiveLinkHubClient->GetSubjectSettings(InSubjectKey));
@@ -144,14 +160,17 @@ void FLiveLinkHub::OnSubjectAdded(FLiveLinkSubjectKey InSubjectKey)
 	{
 		FLiveLinkStaticDataStruct StaticDataCopy;
 		StaticDataCopy.InitializeWith(*StaticData);
-		LiveLinkProvider->UpdateSubjectStaticData(InSubjectKey.SubjectName, SubjectSettings->Role, MoveTemp(StaticDataCopy));
+
+		const FName OverridenName = GetSubjectNameOverride(InSubjectKey);
+		LiveLinkProvider->UpdateSubjectStaticData(OverridenName, SubjectSettings->Role, MoveTemp(StaticDataCopy));
 	}
 }
 
-void FLiveLinkHub::OnSubjectRemoved(FLiveLinkSubjectKey InSubjectKey)
+void FLiveLinkHub::OnSubjectRemoved(FLiveLinkSubjectKey InSubjectKey) const
 {
 	// Send an update to connected clients as well.
-	LiveLinkProvider->RemoveSubject(InSubjectKey.SubjectName);
+	const FName OverridenName = GetSubjectNameOverride(InSubjectKey);
+	LiveLinkProvider->RemoveSubject(OverridenName);
 }
 
 void FLiveLinkHub::BindCommands()
@@ -309,4 +328,20 @@ void FLiveLinkHub::OpenConfig()
 	}
 }
 
+FName FLiveLinkHub::GetSubjectNameOverride(const FLiveLinkSubjectKey& InSubjectKey) const
+{
+	if (const TSharedPtr<ILiveLinkHubSessionManager> Manager = SessionManager)
+	{
+		if (const TSharedPtr<ILiveLinkHubSession> CurrentSession = Manager->GetCurrentSession())
+		{
+			if (const ULiveLinkHubSubjectProxy* SubjectProxy = CurrentSession->GetSubjectConfig(InSubjectKey))
+			{
+				return SubjectProxy->GetOutboundName();
+			}
+		}
+	}
+
+	return InSubjectKey.SubjectName;
+}
+
 #undef LOCTEXT_NAMESPACE /*LiveLinkHub*/
diff --git a/Engine/Plugins/Animation/LiveLink/Source/LiveLinkHub/Private/LiveLinkHub.h b/Engine/Plugins/Animation/LiveLink/Source/LiveLinkHub/Private/LiveLinkHub.h
--- a/Engine/Plugins/Animation/LiveLink/Source/LiveLinkHub/Private/LiveLinkHub.h
+++ b/Engine/Plugins/Animation/LiveLink/Source/LiveLinkHub/Private/LiveLinkHub.h
@@ -12,9 +12,11 @@ class FLiveLinkHubClientsController;
 class FLiveLinkHubPlaybackController;
 class FLiveLinkHubRecordingController;
 class FLiveLinkHubRecordingListController;
+class FLiveLinkHubSubjectController;
 class FLiveLinkHubWindowController;
 struct FLiveLinkSubjectKey;
 class FLiveLinkHubProvider;
+class ILiveLinkHubSessionManager;
 class SWindow;
 class ULiveLinkRole;
 
@@ -39,7 +41,7 @@ public:
 class FLiveLinkHub : public ILiveLinkHub, public TSharedFromThis<FLiveLinkHub>
 {
 public:
-	virtual ~FLiveLinkHub();
+	virtual ~FLiveLinkHub() override;
 
 	//~ Begin ILiveLinkHub interface
 	virtual bool IsInPlayback() const override;
@@ -65,13 +67,15 @@ public:
 	TSharedPtr<FLiveLinkHubClientsController> GetClientsController() const;
 	/** Get the live link hub command list. */
 	TSharedPtr<FUICommandList> GetCommandList() const { return CommandList; }
+	/** Get the session manager. */
+	TSharedPtr<ILiveLinkHubSessionManager> GetSessionManager() const;
 	
 private:
 	//~ LiveLink Client delegates
-	void OnStaticDataReceived_AnyThread(const FLiveLinkSubjectKey& InSubjectKey, TSubclassOf<ULiveLinkRole> InRole, const FLiveLinkStaticDataStruct& InStaticDataStruct);
-	void OnFrameDataReceived_AnyThread(const FLiveLinkSubjectKey& InSubjectKey, const FLiveLinkFrameDataStruct& InFrameDataStruct);
-	void OnSubjectAdded(FLiveLinkSubjectKey InSubjectKey);
-	void OnSubjectRemoved(FLiveLinkSubjectKey InSubjectKey);
+	void OnStaticDataReceived_AnyThread(const FLiveLinkSubjectKey& InSubjectKey, TSubclassOf<ULiveLinkRole> InRole, const FLiveLinkStaticDataStruct& InStaticDataStruct) const;
+	void OnFrameDataReceived_AnyThread(const FLiveLinkSubjectKey& InSubjectKey, const FLiveLinkFrameDataStruct& InFrameDataStruct) const;
+	void OnSubjectAdded(FLiveLinkSubjectKey InSubjectKey) const;
+	void OnSubjectRemoved(FLiveLinkSubjectKey InSubjectKey) const;
 	//~ LiveLink Client delegates
 
 	/** Bind all available live link hub commands. */
@@ -79,6 +83,7 @@ private:
 
 	/** Clear all client settings. */
 	void ClearClient();
+	FName GetSubjectNameOverride(const FLiveLinkSubjectKey& InSubjectKey) const;
 
 	/** Create a new config. */
 	void NewConfig();
@@ -104,8 +109,12 @@ private:
 	TSharedPtr<FLiveLinkHubRecordingListController> RecordingListController;
 	/** Implements the logic for triggering the playback of a livelink recording. */
 	TSharedPtr<FLiveLinkHubPlaybackController> PlaybackController;
+	/** Implements the controller responsible for displaying and managing subject data. */
+	TSharedPtr<FLiveLinkHubSubjectController> SubjectController;
 	/** Controller responsible for creating and managing the app's slate windows. */
 	TSharedPtr<FLiveLinkHubWindowController> WindowController;
+	/** Object responsible for managing sessions.  */
+	TSharedPtr<ILiveLinkHubSessionManager> SessionManager;
 	/** LiveLinkHub's livelink client. */
 	TSharedPtr<FLiveLinkHubClient> LiveLinkHubClient;
 	/** LiveLinkProvider used to transfer data to connected UE clients. */
diff --git a/Engine/Plugins/Animation/LiveLink/Source/LiveLinkHub/Private/LiveLinkHubModule.cpp b/Engine/Plugins/Animation/LiveLink/Source/LiveLinkHub/Private/LiveLinkHubModule.cpp
--- a/Engine/Plugins/Animation/LiveLink/Source/LiveLinkHub/Private/LiveLinkHubModule.cpp
+++ b/Engine/Plugins/Animation/LiveLink/Source/LiveLinkHub/Private/LiveLinkHubModule.cpp
@@ -45,5 +45,14 @@ TSharedPtr<FLiveLinkHubPlaybackController> FLiveLinkHubModule::GetPlaybackContro
 	return LiveLinkHub->PlaybackController;
 }
 
-IMPLEMENT_MODULE(FLiveLinkHubModule, LiveLinkHub);
+TSharedPtr<FLiveLinkHubSubjectController> FLiveLinkHubModule::GetSubjectController() const
+{
+	return LiveLinkHub->SubjectController;
+}
 
+TSharedPtr<ILiveLinkHubSessionManager> FLiveLinkHubModule::GetSessionManager() const
+{
+	return LiveLinkHub->SessionManager;
+}
+
+IMPLEMENT_MODULE(FLiveLinkHubModule, LiveLinkHub);
diff --git a/Engine/Plugins/Animation/LiveLink/Source/LiveLinkHub/Private/LiveLinkHubModule.h b/Engine/Plugins/Animation/LiveLink/Source/LiveLinkHub/Private/LiveLinkHubModule.h
--- a/Engine/Plugins/Animation/LiveLink/Source/LiveLinkHub/Private/LiveLinkHubModule.h
+++ b/Engine/Plugins/Animation/LiveLink/Source/LiveLinkHub/Private/LiveLinkHubModule.h
@@ -10,6 +10,8 @@ class FLiveLinkHubPlaybackController;
 class FLiveLinkHubProvider;
 class FLiveLinkHubRecordingController;
 class FLiveLinkHubRecordingListController;
+class FLiveLinkHubSubjectController;
+class ILiveLinkHubSessionManager;
 
 class FLiveLinkHubModule : public ILiveLinkHubModule
 {
@@ -28,6 +30,11 @@ public:
 	TSharedPtr<FLiveLinkHubRecordingListController> GetRecordingListController() const;
 	/** Get the playback controller. */
 	TSharedPtr<FLiveLinkHubPlaybackController> GetPlaybackController() const;
+	/** Get the subject controller. */
+	TSharedPtr<FLiveLinkHubSubjectController> GetSubjectController() const;
+	/** Get the subject controller. */
+    TSharedPtr<ILiveLinkHubSessionManager> GetSessionManager() const;
+
 private:
 	/** LiveLinkHub object responsible for initializing the different controllers. */
 	TSharedPtr<FLiveLinkHub> LiveLinkHub;
diff --git a/Engine/Plugins/Animation/LiveLink/Source/LiveLinkHub/Private/Session/LiveLinkHubSession.h b/Engine/Plugins/Animation/LiveLink/Source/LiveLinkHub/Private/Session/LiveLinkHubSession.h
new file mode 100644
--- /dev/null
+++ b/Engine/Plugins/Animation/LiveLink/Source/LiveLinkHub/Private/Session/LiveLinkHubSession.h
@@ -0,0 +1,64 @@
+// Copyright Epic Games, Inc. All Rights Reserved.
+
+#pragma once
+
+#include "LiveLinkTypes.h"
+#include "Subjects/LiveLinkHubSubjectSessionConfig.h"
+#include "UObject/StrongObjectPtr.h"
+
+/**
+ * Note: This is just a stub implementation until the session manager work is submitted.
+ */
+class ILiveLinkHubSession
+{
+public:
+	virtual ~ILiveLinkHubSession() = default;
+	virtual ULiveLinkHubSubjectProxy* GetSubjectConfig(const FLiveLinkSubjectKey& SubjectKey) const = 0;
+};
+
+class FLiveLinkHubSession : public ILiveLinkHubSession
+{
+public:
+	FLiveLinkHubSession()
+	{
+		SessionConfig = TStrongObjectPtr<ULiveLinkHubSubjectSessionConfig>(NewObject<ULiveLinkHubSubjectSessionConfig>());
+	}
+
+	virtual ~FLiveLinkHubSession() override = default;
+
+	virtual ULiveLinkHubSubjectProxy* GetSubjectConfig(const FLiveLinkSubjectKey& SubjectKey) const override
+	{
+		return SessionConfig->GetSubjectConfig(SubjectKey);
+	}
+
+private:
+	TStrongObjectPtr<ULiveLinkHubSubjectSessionConfig> SessionConfig;
+};
+
+class ILiveLinkHubSessionManager
+{
+public:
+	virtual ~ILiveLinkHubSessionManager() = default;
+
+	// Todo: Real implementation must be thread safe
+	virtual TSharedPtr<ILiveLinkHubSession> GetCurrentSession() const = 0;
+};
+
+class FLiveLinkHubSessionManager : public ILiveLinkHubSessionManager
+{
+public:
+	FLiveLinkHubSessionManager()
+	{
+		CurrentSession = MakeShared<FLiveLinkHubSession>();
+	}
+
+	virtual ~FLiveLinkHubSessionManager() override = default;
+
+	virtual TSharedPtr<ILiveLinkHubSession> GetCurrentSession() const override
+	{
+		return CurrentSession;
+	}
+
+private:
+	TSharedPtr<ILiveLinkHubSession> CurrentSession;
+};
diff --git a/Engine/Plugins/Animation/LiveLink/Source/LiveLinkHub/Private/Subjects/LiveLinkHubSubjectController.h b/Engine/Plugins/Animation/LiveLink/Source/LiveLinkHub/Private/Subjects/LiveLinkHubSubjectController.h
new file mode 100644
--- /dev/null
+++ b/Engine/Plugins/Animation/LiveLink/Source/LiveLinkHub/Private/Subjects/LiveLinkHubSubjectController.h
@@ -0,0 +1,35 @@
+// Copyright Epic Games, Inc. All Rights Reserved.
+
+#pragma once
+
+#include "Clients/LiveLinkHubProvider.h"
+#include "LiveLinkTypes.h"
+#include "SLiveLinkHubSubjectView.h"
+
+/** Controller responsible for handling the hub's subjects and creating the subject view. */
+class FLiveLinkHubSubjectController
+{
+public:
+	FLiveLinkHubSubjectController()
+	{
+		SubjectModel = MakeShared<FLiveLinkHubSubjectModel>();
+	}
+
+	/** Create the widget for displaying a subject's settings. */
+	TSharedRef<SWidget> MakeSubjectView()
+	{
+		return SAssignNew(SubjectsView, SLiveLinkHubSubjectView, SubjectModel.ToSharedRef());
+	}
+
+	/** Set the displayed subject in the subject view. */
+	void SetSubject(const FLiveLinkSubjectKey& Subject)
+	{
+		SubjectsView->SetSubject(Subject);
+	}
+
+private:
+	/** View widget for the selected subject. */
+	TSharedPtr<SLiveLinkHubSubjectView> SubjectsView;
+	/** Model responsible for the subjects data. */
+	TSharedPtr<ILiveLinkHubSubjectModel> SubjectModel;
+};
diff --git a/Engine/Plugins/Animation/LiveLink/Source/LiveLinkHub/Private/Subjects/LiveLinkHubSubjectModel.h b/Engine/Plugins/Animation/LiveLink/Source/LiveLinkHub/Private/Subjects/LiveLinkHubSubjectModel.h
new file mode 100644
--- /dev/null
+++ b/Engine/Plugins/Animation/LiveLink/Source/LiveLinkHub/Private/Subjects/LiveLinkHubSubjectModel.h
@@ -0,0 +1,35 @@
+// Copyright Epic Games, Inc. All Rights Reserved.
+
+#pragma once
+
+#include "LiveLinkHubModule.h"
+#include "LiveLinkHubSubjectSessionConfig.h"
+#include "LiveLinkTypes.h"
+#include "Session/LiveLinkHubSession.h"
+
+
+/** Subject Model used by the view to access the subject settings. */
+class ILiveLinkHubSubjectModel
+{
+public:
+	virtual ~ILiveLinkHubSubjectModel() = default;
+
+	/** Get the settings for a livelinkhub subject. */
+	virtual ULiveLinkHubSubjectProxy* GetSubjectConfig(const FLiveLinkSubjectKey& InSubject) const = 0;
+};
+
+/** Implementation of the ILiveLinkHubSubjectModel. */
+class FLiveLinkHubSubjectModel : public ILiveLinkHubSubjectModel
+{
+public:
+	/** Get the settings for a livelinkhub subject. */
+	virtual ULiveLinkHubSubjectProxy* GetSubjectConfig(const FLiveLinkSubjectKey& InSubject) const override
+	{
+		const FLiveLinkHubModule& LiveLinkHubModule = FModuleManager::Get().GetModuleChecked<FLiveLinkHubModule>("LiveLinkHub");
+		if (const TSharedPtr<ILiveLinkHubSession> Session = LiveLinkHubModule.GetSessionManager()->GetCurrentSession())
+		{
+			return Session->GetSubjectConfig(InSubject);
+		}
+		return nullptr;
+	}
+};
diff --git a/Engine/Plugins/Animation/LiveLink/Source/LiveLinkHub/Private/Subjects/LiveLinkHubSubjectSessionConfig.cpp b/Engine/Plugins/Animation/LiveLink/Source/LiveLinkHub/Private/Subjects/LiveLinkHubSubjectSessionConfig.cpp
new file mode 100644
--- /dev/null
+++ b/Engine/Plugins/Animation/LiveLink/Source/LiveLinkHub/Private/Subjects/LiveLinkHubSubjectSessionConfig.cpp
@@ -0,0 +1,165 @@
+// Copyright Epic Games, Inc. All Rights Reserved.
+
+#include "LiveLinkHubSubjectSessionConfig.h"
+
+#include "Async/Async.h"
+#include "Clients/LiveLinkHubProvider.h"
+#include "Features/IModularFeatures.h"
+#include "LiveLinkHubModule.h"
+#include "Modules/ModuleManager.h"
+#include "LiveLinkClient.h"
+#include "UObject/StrongObjectPtr.h"
+
+ULiveLinkHubSubjectSessionConfig::ULiveLinkHubSubjectSessionConfig()
+{
+	if (!IsTemplate())
+	{
+		FLiveLinkClient& LiveLinkClient = IModularFeatures::Get().GetModularFeature<FLiveLinkClient>(ILiveLinkClient::ModularFeatureName);
+		LiveLinkClient.OnLiveLinkSubjectAdded().AddUObject(this, &ULiveLinkHubSubjectSessionConfig::OnSubjectAdded_AnyThread);
+		LiveLinkClient.OnLiveLinkSubjectRemoved().AddUObject(this, &ULiveLinkHubSubjectSessionConfig::OnSubjectRemoved_AnyThread);
+
+		constexpr bool bIncludeDisabledSubject = true;
+		constexpr bool bIncludeVirtualSubject = true;
+		for (const FLiveLinkSubjectKey& SubjectKey : LiveLinkClient.GetSubjects(bIncludeDisabledSubject, bIncludeVirtualSubject))
+		{
+			ULiveLinkHubSubjectProxy* SubjectSettings = NewObject<ULiveLinkHubSubjectProxy>();
+			SubjectSettings->Initialize(SubjectKey, LiveLinkClient.GetSourceType(SubjectKey.Source).ToString());
+
+			SubjectProxies.Add(SubjectKey, SubjectSettings);
+		}
+	}
+}
+
+ULiveLinkHubSubjectSessionConfig::~ULiveLinkHubSubjectSessionConfig()
+{
+	if (IModularFeatures::Get().IsModularFeatureAvailable(ILiveLinkClient::ModularFeatureName))
+    {
+    	FLiveLinkClient& LiveLinkClient = IModularFeatures::Get().GetModularFeature<FLiveLinkClient>(ILiveLinkClient::ModularFeatureName);
+    	LiveLinkClient.OnLiveLinkSubjectRemoved().RemoveAll(this);
+    	LiveLinkClient.OnLiveLinkSubjectAdded().RemoveAll(this);
+    }
+}
+
+ULiveLinkHubSubjectProxy* ULiveLinkHubSubjectSessionConfig::GetSubjectConfig(const FLiveLinkSubjectKey& InSubject)
+{
+    ULiveLinkHubSubjectProxy* Settings = nullptr;
+
+    if (TObjectPtr<ULiveLinkHubSubjectProxy>* SettingsPtr = SubjectProxies.Find(InSubject))
+    {
+    	Settings = SettingsPtr->Get();
+    }
+
+    return Settings;
+}
+
+void ULiveLinkHubSubjectSessionConfig::OnSubjectAdded_AnyThread(FLiveLinkSubjectKey SubjectKey)
+{
+	TWeakObjectPtr<ULiveLinkHubSubjectSessionConfig> Self = this;
+	AsyncTask(ENamedThreads::GameThread, [Self, Key = MoveTemp(SubjectKey)]
+	{
+		ULiveLinkHubSubjectSessionConfig* Container = Self.Get();
+		if (UObjectInitialized() && Container)
+		{
+			Container->OnSubjectAdded(Key);
+		}
+	});
+}
+
+void ULiveLinkHubSubjectSessionConfig::OnSubjectAdded(const FLiveLinkSubjectKey& SubjectKey)
+{
+	if (!SubjectProxies.Contains(SubjectKey))
+	{
+		ILiveLinkClient& LiveLinkClient = IModularFeatures::Get().GetModularFeature<ILiveLinkClient>(ILiveLinkClient::ModularFeatureName);
+
+		ULiveLinkHubSubjectProxy* SubjectSettings = NewObject<ULiveLinkHubSubjectProxy>();
+		SubjectSettings->Initialize(SubjectKey, LiveLinkClient.GetSourceType(SubjectKey.Source).ToString());
+
+		SubjectProxies.FindOrAdd(SubjectKey) = SubjectSettings;
+	}
+}
+
+void ULiveLinkHubSubjectSessionConfig::OnSubjectRemoved_AnyThread(FLiveLinkSubjectKey SubjectKey)
+{
+	TWeakObjectPtr<ULiveLinkHubSubjectSessionConfig> Self = this;
+	AsyncTask(ENamedThreads::GameThread, [Self, Key = MoveTemp(SubjectKey)]
+	{
+		ULiveLinkHubSubjectSessionConfig* Container = Self.Get();
+		if (UObjectInitialized() && Container)
+		{
+			Container->OnSubjectAdded(Key);
+		}
+	});
+}
+
+void ULiveLinkHubSubjectSessionConfig::OnSubjectRemoved(const FLiveLinkSubjectKey& SubjectKey)
+{
+	SubjectProxies.Remove(SubjectKey);
+}
+
+void ULiveLinkHubSubjectProxy::Initialize(const FLiveLinkSubjectKey& InSubjectKey, FString InSource)
+{
+	SubjectName = SubjectKey.SubjectName.Name.ToString();
+	SubjectKey = InSubjectKey;
+	OutboundName = SubjectKey.SubjectName.Name.ToString();
+	Source = MoveTemp(InSource);
+}
+
+void ULiveLinkHubSubjectProxy::PreEditChange(FProperty* PropertyAboutToChange)
+{
+	Super::PreEditChange(PropertyAboutToChange);
+
+	if (PropertyAboutToChange && PropertyAboutToChange->GetFName() == GET_MEMBER_NAME_CHECKED(ULiveLinkHubSubjectProxy, OutboundName))
+	{
+		PreviousOutboundName = *OutboundName;
+	}
+}
+
+void ULiveLinkHubSubjectProxy::PostEditChangeProperty(struct FPropertyChangedEvent& PropertyChangedEvent)
+{
+	Super::PostEditChangeProperty(PropertyChangedEvent);
+
+	if (PropertyChangedEvent.Property && PropertyChangedEvent.Property->GetName() == GET_MEMBER_NAME_CHECKED(ULiveLinkHubSubjectProxy, OutboundName))
+	{
+		NotifyRename();
+	}
+}
+
+FName ULiveLinkHubSubjectProxy::GetOutboundName() const
+{
+	if (bPendingOutboundNameChange)
+	{
+		return PreviousOutboundName;
+	}
+
+	return *OutboundName;
+}
+
+void ULiveLinkHubSubjectProxy::NotifyRename()
+{
+	FLiveLinkHubModule& LiveLinkHubModule = FModuleManager::Get().GetModuleChecked<FLiveLinkHubModule>("LiveLinkHub");
+
+	if (TSharedPtr<FLiveLinkHubProvider> Provider = LiveLinkHubModule.GetLiveLinkProvider())
+	{
+		if (PreviousOutboundName != *OutboundName)
+		{
+			bPendingOutboundNameChange = true;
+
+			// We need to keep the last subject static data to send it.
+			Provider->SendClearSubjectToConnections(SubjectKey.SubjectName.Name);
+
+			TPair<UClass*, FLiveLinkStaticDataStruct*> StaticData = Provider->GetLastSubjectStaticDataStruct(SubjectKey.SubjectName.Name);
+			if (StaticData.Key && StaticData.Value)
+			{
+				FLiveLinkStaticDataStruct StaticDataCopy;
+				StaticDataCopy.InitializeWith(*StaticData.Value);
+
+				Provider->UpdateSubjectStaticData(*OutboundName, StaticData.Key, MoveTemp(StaticDataCopy));
+			}
+
+			Provider->RemoveSubject(PreviousOutboundName);
+
+			bPendingOutboundNameChange = false;
+			PreviousOutboundName = *OutboundName;
+		}
+	}
+}
diff --git a/Engine/Plugins/Animation/LiveLink/Source/LiveLinkHub/Private/Subjects/LiveLinkHubSubjectSessionConfig.h b/Engine/Plugins/Animation/LiveLink/Source/LiveLinkHub/Private/Subjects/LiveLinkHubSubjectSessionConfig.h
new file mode 100644
--- /dev/null
+++ b/Engine/Plugins/Animation/LiveLink/Source/LiveLinkHub/Private/Subjects/LiveLinkHubSubjectSessionConfig.h
@@ -0,0 +1,88 @@
+// Copyright Epic Games, Inc. All Rights Reserved.
+
+#pragma once
+
+#include "Delegates/Delegate.h"
+#include "LiveLinkTypes.h"
+#include "UObject/Object.h"
+#include "UObject/ObjectMacros.h"
+
+#include "LiveLinkHubSubjectSessionConfig.generated.h"
+
+/** Config pertaining to livelink hub subjects for a given session. */
+UCLASS()
+class ULiveLinkHubSubjectSessionConfig : public UObject
+{
+public:
+	GENERATED_BODY()
+
+	ULiveLinkHubSubjectSessionConfig();
+	virtual ~ULiveLinkHubSubjectSessionConfig() override;
+
+	/** Get the config for a livelinkhub subject. */
+	ULiveLinkHubSubjectProxy* GetSubjectConfig(const FLiveLinkSubjectKey& InSubject);
+
+private:
+	/** AnyThread handler for the SubjectAdded delegate, dispatches handling on the game thread to avoid asserts in Slate. */
+	void OnSubjectAdded_AnyThread(FLiveLinkSubjectKey SubjectKey);
+	/** Handles updating the tree view when a subject is added. */
+	void OnSubjectAdded(const FLiveLinkSubjectKey& SubjectKey);
+	/** AnyThread handler for the SubjectRemoved delegate, dispatches handling on the game thread to avoid asserts in Slate. */
+	void OnSubjectRemoved_AnyThread(FLiveLinkSubjectKey SubjectKey);
+	/** Handles updating the tree view when a subject is removed. */
+	void OnSubjectRemoved(const FLiveLinkSubjectKey& SubjectKey);
+
+private:
+	/** Settings for subjects displayed in the livelink hub. */
+	UPROPERTY(Instanced)
+	TMap<FLiveLinkSubjectKey, TObjectPtr<ULiveLinkHubSubjectProxy>> SubjectProxies;
+};
+
+
+/** Holds information and config for a subject for the duration of a livelink hub session.  */
+UCLASS()
+class ULiveLinkHubSubjectProxy : public UObject
+{
+public:
+	GENERATED_BODY()
+
+	/** Initialize from a subject key. */
+	void Initialize(const FLiveLinkSubjectKey& InSubjectKey, FString InSource);
+
+	/** Get the outbound name for this subject, allows  */
+	FName GetOutboundName() const;
+
+	//~ Begin UObject interface
+	virtual void PreEditChange(FProperty* PropertyAboutToChange) override;
+	virtual void PostEditChangeProperty(struct FPropertyChangedEvent& PropertyChangedEvent) override;
+	//~ End UObject interface
+private:
+	/** Notify clients that the old subject should be deleted and replaced with a new one with the updated name. */
+	void NotifyRename();
+
+private:
+	/** Name of  the subject, */
+	UPROPERTY(VisibleAnywhere, Category = "LiveLink")
+	FString SubjectName;
+
+	/** Name override that will be transmitted to clients instead of the subject name. */
+	UPROPERTY(EditAnywhere, Category = "LiveLink")
+	FString OutboundName;
+
+	/** Source that contains the subject. */
+	UPROPERTY(VisibleAnywhere, Category = "LiveLink")
+	FString Source;
+
+	/** SubjectKey for this subject, */
+	UPROPERTY()
+	FLiveLinkSubjectKey SubjectKey;
+
+	/**
+	 * If this is set, then the outbound name is currently undergoing a rename,
+	 * If this is not set, then the OutboundName is committed.
+	 */
+	bool bPendingOutboundNameChange = false;
+
+	/* Previous outbound name to be used for noticing clients to remove this entry from their subject list. */
+	FName PreviousOutboundName;
+};
diff --git a/Engine/Plugins/Animation/LiveLink/Source/LiveLinkHub/Private/Subjects/SLiveLinkHubSubjectView.h b/Engine/Plugins/Animation/LiveLink/Source/LiveLinkHub/Private/Subjects/SLiveLinkHubSubjectView.h
new file mode 100644
--- /dev/null
+++ b/Engine/Plugins/Animation/LiveLink/Source/LiveLinkHub/Private/Subjects/SLiveLinkHubSubjectView.h
@@ -0,0 +1,70 @@
+// Copyright Epic Games, Inc. All Rights Reserved.
+
+#pragma once
+
+#include "Widgets/SCompoundWidget.h"
+
+#include "IDetailsView.h"
+#include "LiveLinkHubSubjectModel.h"
+#include "LiveLinkTypes.h"
+#include "PropertyEditorModule.h"
+#include "Modules/ModuleManager.h"
+
+/**
+ * Provides the UI that displays information about a livelink hub subject.
+ */
+class SLiveLinkHubSubjectView : public SCompoundWidget
+{
+public:
+	SLATE_BEGIN_ARGS(SLiveLinkHubSubjectView) {}
+	SLATE_END_ARGS()
+
+	//~ Begin SWidget interface
+	void Construct(const FArguments& InArgs, const TSharedRef<ILiveLinkHubSubjectModel>& InSubjectModel)
+	{
+		SubjectModel = InSubjectModel;
+
+		FDetailsViewArgs DetailsViewArgs;
+		DetailsViewArgs.bUpdatesFromSelection = false;
+		DetailsViewArgs.bLockable = false;
+		DetailsViewArgs.bShowPropertyMatrixButton = false;
+		DetailsViewArgs.NameAreaSettings = FDetailsViewArgs::HideNameArea;
+		DetailsViewArgs.ViewIdentifier = NAME_None;
+		DetailsViewArgs.bShowCustomFilterOption = false;
+		DetailsViewArgs.bShowOptions = false;
+
+		FPropertyEditorModule& PropertyEditorModule = FModuleManager::GetModuleChecked<FPropertyEditorModule>("PropertyEditor");
+
+		SettingsDetailsView = PropertyEditorModule.CreateDetailView(DetailsViewArgs);
+
+		ChildSlot
+		[
+			SettingsDetailsView.ToSharedRef()
+		];
+	}
+
+	/** Set the subject to be displayed in the details view. */
+	void SetSubject(const FLiveLinkSubjectKey& InSubjectKey)
+	{
+		SubjectKey = InSubjectKey;
+		SettingsDetailsView->SetObject(SubjectModel->GetSubjectConfig(SubjectKey));
+	}
+
+	/** Get the outbound name for a subject. */
+	FString GetSubjectOutboundName(const FLiveLinkSubjectKey& InSubjectKey)
+	{
+		if (ULiveLinkHubSubjectProxy* SubjectProxy = SubjectModel->GetSubjectConfig(InSubjectKey))
+		{
+			return SubjectProxy->GetOutboundName().ToString();
+		}
+		return InSubjectKey.SubjectName.ToString();
+	}
+
+private:
+	/** Details for the selected subject. */
+	TSharedPtr<IDetailsView> SettingsDetailsView;
+	/** Subject being shown. */
+	FLiveLinkSubjectKey SubjectKey;
+	/** Model that holds the data for a given subject. */
+	TSharedPtr<ILiveLinkHubSubjectModel> SubjectModel;
+};
diff --git a/Engine/Plugins/Animation/LiveLink/Source/LiveLinkHub/Private/UI/Widgets/SLiveLinkHubMainTabView.cpp b/Engine/Plugins/Animation/LiveLink/Source/LiveLinkHub/Private/UI/Widgets/SLiveLinkHubMainTabView.cpp
--- a/Engine/Plugins/Animation/LiveLink/Source/LiveLinkHub/Private/UI/Widgets/SLiveLinkHubMainTabView.cpp
+++ b/Engine/Plugins/Animation/LiveLink/Source/LiveLinkHub/Private/UI/Widgets/SLiveLinkHubMainTabView.cpp
@@ -11,11 +11,10 @@
 #include "LiveLinkHub.h"
 #include "LiveLinkHubModule.h"
 #include "LiveLinkPanelController.h"
+#include "LiveLinkTypes.h"
 #include "Modules/ModuleManager.h"
-#include "Recording/LiveLinkHubPlaybackController.h"
-#include "Recording/LiveLinkHubRecordingController.h"
 #include "Recording/LiveLinkHubRecordingListController.h"
-#include "SLiveLinkDataView.h"
+#include "Subjects/LiveLinkHubSubjectController.h"
 #include "Widgets/Docking/SDockTab.h"
 #include "Widgets/Input/SCheckBox.h"
 #include "Widgets/Text/STextBlock.h"
@@ -39,8 +38,9 @@ const FText SLiveLinkHubMainTabView::ClientsTabName = LOCTEXT("ClientsTabLabel",
 const FText SLiveLinkHubMainTabView::ClientDetailsTabName = LOCTEXT("ClientDetailsTabLabel", "Client Details");
 
 void SLiveLinkHubMainTabView::Construct(const FArguments& InArgs)
-{	
+{
 	PanelController = MakeShared<FLiveLinkPanelController>();
+	PanelController->OnSubjectSelectionChanged().AddSP(this, &SLiveLinkHubMainTabView::OnSubjectSelectionChanged);
 
 	SLiveLinkHubTabViewWithManagerBase::Construct(
 		SLiveLinkHubTabViewWithManagerBase::FArguments()
@@ -54,6 +54,14 @@ void SLiveLinkHubMainTabView::Construct(const FArguments& InArgs)
 	);
 }
 
+SLiveLinkHubMainTabView::~SLiveLinkHubMainTabView()
+{
+	if (PanelController)
+	{
+		PanelController->OnSubjectSelectionChanged().RemoveAll(this);
+	}
+}
+
 void SLiveLinkHubMainTabView::CreateTabs(const TSharedRef<FTabManager>& InTabManager, const TSharedRef<FTabManager::FLayout>& InLayout, const FArguments& InArgs)
 {
 	InTabManager->RegisterTabSpawner(SourcesTabId, FOnSpawnTab::CreateSP(this, &SLiveLinkHubMainTabView::SpawnSourcesTab))
@@ -187,11 +195,13 @@ TSharedRef<SDockTab> SLiveLinkHubMainTabView::SpawnSubjectsTab(const FSpawnTabAr
 
 TSharedRef<SDockTab> SLiveLinkHubMainTabView::SpawnSubjectsDetailsTab(const FSpawnTabArgs& InTabArgs)
 {
+	FLiveLinkHubModule& LiveLinkHubModule = FModuleManager::Get().GetModuleChecked<FLiveLinkHubModule>("LiveLinkHub");
+
 	return SNew(SDockTab)
 		.Label(SubjectsDetailsTabName)
 		.TabRole(PanelTab)
 		[
-			PanelController->SubjectsDetailsView.ToSharedRef()
+			LiveLinkHubModule.GetSubjectController()->MakeSubjectView()
 		];
 }
 
@@ -233,4 +243,10 @@ TSharedRef<SDockTab> SLiveLinkHubMainTabView::SpawnClientDetailsTab(const FSpawn
 		];
 }
 
+void SLiveLinkHubMainTabView::OnSubjectSelectionChanged(const FLiveLinkSubjectKey& SubjectKey)
+{
+	FLiveLinkHubModule& LiveLinkHubModule = FModuleManager::Get().GetModuleChecked<FLiveLinkHubModule>("LiveLinkHub");
+	LiveLinkHubModule.GetSubjectController()->SetSubject(SubjectKey);
+}
+
 #undef LOCTEXT_NAMESPACE
diff --git a/Engine/Plugins/Animation/LiveLink/Source/LiveLinkHub/Private/UI/Widgets/SLiveLinkHubMainTabView.h b/Engine/Plugins/Animation/LiveLink/Source/LiveLinkHub/Private/UI/Widgets/SLiveLinkHubMainTabView.h
--- a/Engine/Plugins/Animation/LiveLink/Source/LiveLinkHub/Private/UI/Widgets/SLiveLinkHubMainTabView.h
+++ b/Engine/Plugins/Animation/LiveLink/Source/LiveLinkHub/Private/UI/Widgets/SLiveLinkHubMainTabView.h
@@ -5,6 +5,7 @@
 #include "SLiveLinkHubTabViewWithManagerBase.h"
 
 class FLiveLinkHubRecordingController;
+struct FLiveLinkSubjectKey;
 class ILiveLinkClient;
 
 /** Manages the UI logic of the Clients tab */
@@ -36,6 +37,7 @@ public:
 	SLATE_END_ARGS()
 
 	void Construct(const FArguments& InArgs);
+	virtual ~SLiveLinkHubMainTabView() override;
 	
 private:
 	/** Create all livelink hub tabs. */
@@ -50,6 +52,9 @@ private:
 	TSharedRef<SDockTab> SpawnClientsTab(const FSpawnTabArgs& InTabArgs);
 	TSharedRef<SDockTab> SpawnClientDetailsTab(const FSpawnTabArgs& InTabArgs);
 
+	/** Handles selection change in the subjects view. */
+	void OnSubjectSelectionChanged(const FLiveLinkSubjectKey& SubjectKey);
+
 private:
 	/** Holds the livelink panel controller responsible for creating sources and subjects tabs. */
 	TSharedPtr<class FLiveLinkPanelController> PanelController;
diff --git a/Engine/Source/Runtime/LiveLinkInterface/Public/ILiveLinkClient.h b/Engine/Source/Runtime/LiveLinkInterface/Public/ILiveLinkClient.h
--- a/Engine/Source/Runtime/LiveLinkInterface/Public/ILiveLinkClient.h
+++ b/Engine/Source/Runtime/LiveLinkInterface/Public/ILiveLinkClient.h
@@ -18,13 +18,13 @@ struct FLiveLinkSubjectFrameData;
 struct FTimecode;
 class ULiveLinkSourceSettings;
 
-DECLARE_EVENT_OneParam(ILiveLinkClient, FOnLiveLinkSourceChangedDelegate, FGuid /*SourceGuid*/);
-DECLARE_EVENT_OneParam(ILiveLinkClient, FOnLiveLinkSubjectChangedDelegate, FLiveLinkSubjectKey /*SubjectKey*/);
-DECLARE_EVENT_OneParam(ILiveLinkClient, FOnLiveLinkSubjectStaticDataReceived, const FLiveLinkStaticDataStruct& /*InStaticData*/)
-DECLARE_EVENT_OneParam(ILiveLinkClient, FOnLiveLinkSubjectFrameDataReceived, const FLiveLinkFrameDataStruct& /*InFrameData*/)
-DECLARE_EVENT_ThreeParams(ILiveLinkClient, FOnLiveLinkSubjectStaticDataAdded, FLiveLinkSubjectKey /*InSubjectKey*/, TSubclassOf<ULiveLinkRole> /*SubjectRole*/, const FLiveLinkStaticDataStruct& /*InStaticData*/)
-DECLARE_EVENT_ThreeParams(ILiveLinkClient, FOnLiveLinkSubjectFrameDataAdded, FLiveLinkSubjectKey /*InSubjectKey*/, TSubclassOf<ULiveLinkRole> /*SubjectRole*/, const FLiveLinkFrameDataStruct& /*InFrameData*/)
-DECLARE_EVENT_FiveParams(ILiveLinkClient, FOnLiveLinkSubjectEvaluated, FLiveLinkSubjectKey /*InSubjectKey*/, TSubclassOf<ULiveLinkRole> /*RequestedRole*/, const FLiveLinkTime& /*RequestedTime*/, bool /*bResult*/, const FLiveLinkTime& /*EvaluatedFrameTime*/)
+DECLARE_TS_MULTICAST_DELEGATE_OneParam(FOnLiveLinkSourceChangedDelegate, FGuid /*SourceGuid*/);
+DECLARE_TS_MULTICAST_DELEGATE_OneParam(FOnLiveLinkSubjectChangedDelegate, FLiveLinkSubjectKey /*SubjectKey*/);
+DECLARE_TS_MULTICAST_DELEGATE_OneParam(FOnLiveLinkSubjectStaticDataReceived, const FLiveLinkStaticDataStruct& /*InStaticData*/)
+DECLARE_TS_MULTICAST_DELEGATE_OneParam(FOnLiveLinkSubjectFrameDataReceived, const FLiveLinkFrameDataStruct& /*InFrameData*/)
+DECLARE_TS_MULTICAST_DELEGATE_ThreeParams(FOnLiveLinkSubjectStaticDataAdded, FLiveLinkSubjectKey /*InSubjectKey*/, TSubclassOf<ULiveLinkRole> /*SubjectRole*/, const FLiveLinkStaticDataStruct& /*InStaticData*/)
+DECLARE_TS_MULTICAST_DELEGATE_ThreeParams(FOnLiveLinkSubjectFrameDataAdded, FLiveLinkSubjectKey /*InSubjectKey*/, TSubclassOf<ULiveLinkRole> /*SubjectRole*/, const FLiveLinkFrameDataStruct& /*InFrameData*/)
+DECLARE_TS_MULTICAST_DELEGATE_FiveParams(FOnLiveLinkSubjectEvaluated, FLiveLinkSubjectKey /*InSubjectKey*/, TSubclassOf<ULiveLinkRole> /*RequestedRole*/, const FLiveLinkTime& /*RequestedTime*/, bool /*bResult*/, const FLiveLinkTime& /*EvaluatedFrameTime*/)
 
 PRAGMA_DISABLE_DEPRECATION_WARNINGS
 class ILiveLinkClient_Base_DEPRECATED : public IModularFeature
diff --git a/Engine/Source/Runtime/LiveLinkMessageBusFramework/Private/LiveLinkProvider.cpp b/Engine/Source/Runtime/LiveLinkMessageBusFramework/Private/LiveLinkProvider.cpp
--- a/Engine/Source/Runtime/LiveLinkMessageBusFramework/Private/LiveLinkProvider.cpp
+++ b/Engine/Source/Runtime/LiveLinkMessageBusFramework/Private/LiveLinkProvider.cpp
@@ -151,6 +151,23 @@ void FLiveLinkProvider::SendSubjectFrame(FName SubjectName, const FTrackedSubjec
 	MessageEndpoint->Send(SubjectFrame, FLiveLinkSubjectFrameMessage::StaticStruct(), EMessageFlags::None, GetAnnotations(), nullptr, Addresses, FTimespan::Zero(), FDateTime::MaxValue());
 }
 
+TPair<UClass*, FLiveLinkStaticDataStruct*> FLiveLinkProvider::GetLastSubjectStaticDataStruct(FName SubjectName)
+{
+	FScopeLock Lock(&CriticalSection);
+	TPair<UClass*, FLiveLinkStaticDataStruct*> Pair = { nullptr, nullptr };
+
+	if (FTrackedStaticData* TrackedStaticData = GetLastSubjectStaticData(SubjectName))
+	{
+		if (TrackedStaticData->RoleClass.IsValid() && TrackedStaticData->StaticData.IsValid())
+		{
+			Pair.Key = TrackedStaticData->RoleClass.Get();
+			Pair.Value = &TrackedStaticData->StaticData;
+		}
+	}
+
+	return Pair;
+}
+
 // Get the cached data for the named subject
 FTrackedStaticData* FLiveLinkProvider::GetLastSubjectStaticData(const FName& SubjectName)
 {
@@ -205,14 +222,6 @@ void FLiveLinkProvider::ClearTrackedSubject(const FName& SubjectName)
 	}
 }
 
-void FLiveLinkProvider::SendClearSubjectToConnections(FName SubjectName)
-{
-	TArray<FMessageAddress> MessageAddresses;
-	GetFilteredAddresses(SubjectName, MessageAddresses);
-
-	MessageEndpoint->Send(FMessageEndpoint::MakeMessage<FLiveLinkClearSubject>(SubjectName), EMessageFlags::Reliable, GetAnnotations(), nullptr, MessageAddresses, FTimespan::Zero(), FDateTime::MaxValue());
-}
-
 FLiveLinkProvider::FLiveLinkProvider(const FString& InProviderName)
 	: ProviderName(InProviderName)
 	, MachineName(FPlatformProcess::ComputerName())
@@ -261,6 +270,14 @@ void FLiveLinkProvider::UpdateSubject(const FName& SubjectName, const TArray<FNa
 	SendSubject(SubjectName, Subject);
 }
 
+void FLiveLinkProvider::SendClearSubjectToConnections(FName SubjectName)
+{
+	TArray<FMessageAddress> MessageAddresses;
+	GetFilteredAddresses(SubjectName, MessageAddresses);
+
+	MessageEndpoint->Send(FMessageEndpoint::MakeMessage<FLiveLinkClearSubject>(SubjectName), EMessageFlags::Reliable, GetAnnotations(), nullptr, MessageAddresses, FTimespan::Zero(), FDateTime::MaxValue());
+}
+
 bool FLiveLinkProvider::UpdateSubjectStaticData(const FName SubjectName, TSubclassOf<ULiveLinkRole> Role, FLiveLinkStaticDataStruct&& StaticData)
 {
 	FScopeLock Lock(&CriticalSection);
diff --git a/Engine/Source/Runtime/LiveLinkMessageBusFramework/Public/LiveLinkProvider.h b/Engine/Source/Runtime/LiveLinkMessageBusFramework/Public/LiveLinkProvider.h
--- a/Engine/Source/Runtime/LiveLinkMessageBusFramework/Public/LiveLinkProvider.h
+++ b/Engine/Source/Runtime/LiveLinkMessageBusFramework/Public/LiveLinkProvider.h
@@ -22,32 +22,7 @@ class ULiveLinkRole;
 /** Delegate called when the connection status of the provider has changed. */
 DECLARE_MULTICAST_DELEGATE(FLiveLinkProviderConnectionStatusChanged);
 
-PRAGMA_DISABLE_DEPRECATION_WARNINGS
-struct ILiveLinkProvider_Base_DEPRECATED
-{
-	virtual ~ILiveLinkProvider_Base_DEPRECATED() {}
-
-	/** Remove named subject */
-	UE_DEPRECATED(4.23, "ILiveLinkProvider::ClearSubject is deprecated. Please use RemoveSubject instead!")
-	virtual void ClearSubject(const FName& SubjectName) = 0;
-
-	// Update hierarchy for named subject
-	UE_DEPRECATED(4.23, "ILiveLinkProvider::UpdateSubject is deprecated. Please use UpdateSubject with the proper LiveLinkRole's structure instead!")
-	virtual void UpdateSubject(const FName& SubjectName, const TArray<FName>& BoneNames, const TArray<int32>& BoneParents) {}
-
-	// Update subject with transform data
-	UE_DEPRECATED(4.23, "ILiveLinkProvider::UpdateSubjectFrame is deprecated. Please use UpdateSubjectFrame with the proper LiveLinkRole's structure instead!")
-	virtual void UpdateSubjectFrame(const FName& SubjectName, const TArray<FTransform>& BoneTransforms, const TArray<FLiveLinkCurveElement>& CurveData,
-			double Time) {};
-
-	// Update subject with additional metadata
-	UE_DEPRECATED(4.23, "ILiveLinkProvider::UpdateSubjectFrame is deprecated. Please use UpdateSubjectFrame with the proper LiveLinkRole's structure instead!")
-	virtual void UpdateSubjectFrame(const FName& SubjectName, const TArray<FTransform>& BoneTransforms, const TArray<FLiveLinkCurveElement>& CurveData,
-			const FLiveLinkMetaData& MetaData, double Time) {};
-};
-PRAGMA_ENABLE_DEPRECATION_WARNINGS
-
-struct ILiveLinkProvider : public ILiveLinkProvider_Base_DEPRECATED
+struct ILiveLinkProvider
 {
 public:
 	LIVELINKMESSAGEBUSFRAMEWORK_API static TSharedPtr<ILiveLinkProvider> CreateLiveLinkProvider(const FString& ProviderName);
diff --git a/Engine/Source/Runtime/LiveLinkMessageBusFramework/Public/LiveLinkProviderImpl.h b/Engine/Source/Runtime/LiveLinkMessageBusFramework/Public/LiveLinkProviderImpl.h
--- a/Engine/Source/Runtime/LiveLinkMessageBusFramework/Public/LiveLinkProviderImpl.h
+++ b/Engine/Source/Runtime/LiveLinkMessageBusFramework/Public/LiveLinkProviderImpl.h
@@ -37,7 +37,7 @@ private:
 
 	TSharedPtr<class FMessageEndpoint, ESPMode::ThreadSafe> MessageEndpoint;
 
-	/** Lock to stop multiple threads accessing the CurrentPreset at the same time */
+	// Lock to stop multiple threads accessing the CurrentPreset at the same time
 	mutable FCriticalSection CriticalSection;
 
 	// Array of our current connections
@@ -77,9 +77,6 @@ private:
 	// Clear a existing track subject
 	void ClearTrackedSubject(const FName& SubjectName);
 
-	// Send a clear subject message to indicate that the subject should be removed from the connected client.
-	void SendClearSubjectToConnections(FName SubjectName);
-
 	// Get the connected addresses that should receive livelink data.
 	void GetFilteredAddresses(FName SubjectName, TArray<FMessageAddress>& Addresses);
 
@@ -96,6 +93,9 @@ protected:
 	// Validate our current connections, removing those that have timed out.
 	void ValidateConnections();
 
+	// Get the cached data struct for a subject
+	TPair<UClass*, FLiveLinkStaticDataStruct*> GetLastSubjectStaticDataStruct(FName SubjectName);
+
 	template<typename MessageType>
 	void SendMessage(MessageType* Message)
 	{
@@ -157,6 +157,9 @@ protected:
 		return true;
 	}
 
+	// Send a clear subject message to indicate that the subject should be removed from the connected client.
+	void SendClearSubjectToConnections(FName SubjectName);
+
 	// Constructor for derived classes that allows specifying that no endpoint should be created.
 	FLiveLinkProvider(const FString& InProviderName, bool bInCreateEndpoint);
 
@@ -165,9 +168,9 @@ public:
 
 	FLiveLinkProvider(const FString& InProviderName, struct FMessageEndpointBuilder&& EndpointBuilder);
 
-	virtual ~FLiveLinkProvider();
+	virtual ~FLiveLinkProvider() override;
 
-	virtual void UpdateSubject(const FName& SubjectName, const TArray<FName>& BoneNames, const TArray<int32>& BoneParents) override;
+	virtual void UpdateSubject(const FName& SubjectName, const TArray<FName>& BoneNames, const TArray<int32>& BoneParents);
 
 	virtual bool UpdateSubjectStaticData(const FName SubjectName, TSubclassOf<ULiveLinkRole> Role, FLiveLinkStaticDataStruct&& StaticData) override;
 
@@ -176,17 +179,17 @@ public:
 	virtual void RemoveSubject(const FName SubjectName) override;
 
 PRAGMA_DISABLE_DEPRECATION_WARNINGS
-	virtual void UpdateSubjectFrame(const FName& SubjectName, const TArray<FTransform>& BoneTransforms, const TArray<FLiveLinkCurveElement>& CurveData, double Time) override;
+	virtual void UpdateSubjectFrame(const FName& SubjectName, const TArray<FTransform>& BoneTransforms, const TArray<FLiveLinkCurveElement>& CurveData, double Time);
 
 	virtual void UpdateSubjectFrame(const FName& SubjectName, const TArray<FTransform>& BoneTransforms, const TArray<FLiveLinkCurveElement>& CurveData,
-									const FLiveLinkMetaData& MetaData, double Time) override;
+									const FLiveLinkMetaData& MetaData, double Time);
 PRAGMA_ENABLE_DEPRECATION_WARNINGS
 
-	virtual bool UpdateSubjectFrameData(const FName SubjectName, FLiveLinkFrameDataStruct&& FrameData);
-
+	virtual bool UpdateSubjectFrameData(const FName SubjectName, FLiveLinkFrameDataStruct&& FrameData) override;
+	
 	virtual bool HasConnection() const override;
 
 	virtual FDelegateHandle RegisterConnStatusChangedHandle(const FLiveLinkProviderConnectionStatusChanged::FDelegate& ConnStatusChanged) override;
 
-	virtual void UnregisterConnStatusChangedHandle(FDelegateHandle Handle);
+	virtual void UnregisterConnStatusChangedHandle(FDelegateHandle Handle) override;
 };
