LiveLinkHub - Fix renaming a subject twice not working

diff --git a/Engine/Plugins/Animation/LiveLink/Source/LiveLinkHub/Private/Subjects/LiveLinkHubSubjectSessionConfig.cpp b/Engine/Plugins/Animation/LiveLink/Source/LiveLinkHub/Private/Subjects/LiveLinkHubSubjectSessionConfig.cpp
--- a/Engine/Plugins/Animation/LiveLink/Source/LiveLinkHub/Private/Subjects/LiveLinkHubSubjectSessionConfig.cpp
+++ b/Engine/Plugins/Animation/LiveLink/Source/LiveLinkHub/Private/Subjects/LiveLinkHubSubjectSessionConfig.cpp
@@ -144,10 +144,10 @@ void ULiveLinkHubSubjectProxy::NotifyRename()
 		{
 			bPendingOutboundNameChange = true;
 
-			// We need to keep the last subject static data to send it.
-			Provider->SendClearSubjectToConnections(SubjectKey.SubjectName.Name);
+			Provider->SendClearSubjectToConnections(PreviousOutboundName);
 
-			TPair<UClass*, FLiveLinkStaticDataStruct*> StaticData = Provider->GetLastSubjectStaticDataStruct(SubjectKey.SubjectName.Name);
+			// Re-send the last static data with the new name.
+			TPair<UClass*, FLiveLinkStaticDataStruct*> StaticData = Provider->GetLastSubjectStaticDataStruct(PreviousOutboundName);
 			if (StaticData.Key && StaticData.Value)
 			{
 				FLiveLinkStaticDataStruct StaticDataCopy;
@@ -156,6 +156,7 @@ void ULiveLinkHubSubjectProxy::NotifyRename()
 				Provider->UpdateSubjectStaticData(*OutboundName, StaticData.Key, MoveTemp(StaticDataCopy));
 			}
 
+			// Then clear the old static data entry in the provider.
 			Provider->RemoveSubject(PreviousOutboundName);
 
 			bPendingOutboundNameChange = false;
