Add tool menu entry for livelink hub

diff --git a/Engine/Plugins/Animation/LiveLink/Source/LiveLinkHub/LiveLinkHub.Build.cs b/Engine/Plugins/Animation/LiveLink/Source/LiveLinkHub/LiveLinkHub.Build.cs
--- a/Engine/Plugins/Animation/LiveLink/Source/LiveLinkHub/LiveLinkHub.Build.cs
+++ b/Engine/Plugins/Animation/LiveLink/Source/LiveLinkHub/LiveLinkHub.Build.cs
@@ -35,6 +35,7 @@ namespace UnrealBuildTool.Rules
 				"StandaloneRenderer",
 				"StructUtils",
 				"TimeManagement",
+				"ToolMenus",
 				"ToolWidgets",
 				"UnrealEd",
 				"OutputLog",
diff --git a/Engine/Plugins/Animation/LiveLink/Source/LiveLinkHub/Private/LiveLinkHubModule.cpp b/Engine/Plugins/Animation/LiveLink/Source/LiveLinkHub/Private/LiveLinkHubModule.cpp
--- a/Engine/Plugins/Animation/LiveLink/Source/LiveLinkHub/Private/LiveLinkHubModule.cpp
+++ b/Engine/Plugins/Animation/LiveLink/Source/LiveLinkHub/Private/LiveLinkHubModule.cpp
@@ -8,6 +8,14 @@
 #include "Recording/LiveLinkHubPlaybackController.h"
 #include "Recording/LiveLinkHubRecordingController.h"
 
+#if !WITH_LIVELINK_HUB
+#include "HAL/FileManager.h"
+#include "Misc/AsyncTaskNotification.h"
+#include "ToolMenus.h"
+#endif
+
+#define LOCTEXT_NAMESPACE "LiveLinkHubModule"
+
 void FLiveLinkHubModule::StartLiveLinkHub()
 {
 	// Nothing will get executed after this, so put everything before.
@@ -20,6 +28,28 @@ void FLiveLinkHubModule::StartLiveLinkHub()
 	LiveLinkHub.Reset();
 }
 
+void FLiveLinkHubModule::StartupModule()
+{
+#if !WITH_LIVELINK_HUB  // When running in the editor
+	FToolMenuOwnerScoped OwnerScoped(this);
+	UToolMenu* Menu = UToolMenus::Get()->ExtendMenu("LevelEditor.MainMenu.Tools");
+	FToolMenuSection& Section = Menu->AddSection("VirtualProductionSection", LOCTEXT("VirtualProductionSection", "Virtual Production"));
+
+	Section.AddMenuEntry("LiveLinkHub",
+		LOCTEXT("LiveLinkHubLabel", "LiveLink Hub"),
+		LOCTEXT("LiveLinkHubTooltip", "Launch the LiveLink Hub app."),
+		FSlateIcon("LiveLinkStyle", "LiveLinkClient.Common.Icon.Small"),
+		FUIAction(FExecuteAction::CreateRaw(this, &FLiveLinkHubModule::OpenLiveLinkHub)));
+#endif
+}
+
+void FLiveLinkHubModule::ShutdownModule()
+{
+#if !WITH_LIVELINK_HUB
+	UToolMenus::UnregisterOwner(this);
+#endif
+}
+
 TSharedPtr<FLiveLinkHub> FLiveLinkHubModule::GetLiveLinkHub() const
 {
 	return LiveLinkHub;
@@ -45,6 +75,71 @@ TSharedPtr<FLiveLinkHubPlaybackController> FLiveLinkHubModule::GetPlaybackContro
 	return LiveLinkHub->PlaybackController;
 }
 
+#if !WITH_LIVELINK_HUB
+void FLiveLinkHubModule::OpenLiveLinkHub() const
+{
+	FAsyncTaskNotificationConfig NotificationConfig;
+	NotificationConfig.bKeepOpenOnFailure = true;
+	NotificationConfig.TitleText = LOCTEXT("LaunchingLiveLinkHub", "Launching LiveLink Hub...");
+	NotificationConfig.LogCategory = &LogLiveLinkHub;
+
+	FAsyncTaskNotification Notification(NotificationConfig);
+
+	// Find livelink hub executable location for our build configuration
+	FString LiveLinkHubPath = FPlatformProcess::GenerateApplicationPath(TEXT("LiveLinkHub"), FApp::GetBuildConfiguration());
+
+	// Validate it exists and fall back to development if it doesn't.
+	if (!IFileManager::Get().FileExists(*LiveLinkHubPath) && FApp::GetBuildConfiguration() != EBuildConfiguration::Development)
+	{
+		LiveLinkHubPath = FPlatformProcess::GenerateApplicationPath(TEXT("LiveLinkHub"), EBuildConfiguration::Development);
+	}
+
+	const FText LaunchLiveLinkHubErrorTitle = LOCTEXT("LaunchLiveLinkHubErrorTitle", "Failed to Launch LiveLinkhub.");
+	if (!IFileManager::Get().FileExists(*LiveLinkHubPath))
+	{
+		Notification.SetComplete(
+			LaunchLiveLinkHubErrorTitle,
+			LOCTEXT("LaunchLiveLinkHubError_ExecutableMissing", "Could not find the executable. Have you compiled the LiveLink Hub app?"),
+			false
+			);
+
+		return;
+	}
+
+	// Validate we do not have it running locally
+	const FString AppName = FPaths::GetCleanFilename(LiveLinkHubPath);
+	if (FPlatformProcess::IsApplicationRunning(*AppName))
+	{
+		Notification.SetComplete(
+			LaunchLiveLinkHubErrorTitle,
+			LOCTEXT("LaunchLiveLinkHubError_AlreadyRunning", "A LiveLinkHub instance is already running."),
+			false
+			);
+		return;
+	}
+
+	constexpr bool bLaunchDetached = true;
+	constexpr bool bLaunchHidden = false;
+	constexpr bool bLaunchReallyHidden = false;
+
+	const FProcHandle ProcHandle = FPlatformProcess::CreateProc(*LiveLinkHubPath, TEXT(""), bLaunchDetached, bLaunchHidden, bLaunchReallyHidden, nullptr, 0, nullptr, nullptr, nullptr);
+	if (ProcHandle.IsValid())
+	{
+		Notification.SetComplete(
+			LOCTEXT("LaunchedLiveLinkHub", "Launched LiveLink Hub"), FText(), true);
+
+		return;
+	}
+	else // Very unlikely in practice, but possible in theory.
+	{
+		Notification.SetComplete(
+			LaunchLiveLinkHubErrorTitle,
+			LOCTEXT("LaunchLiveLinkHubError_InvalidHandle", "Failed to create the LiveLink Hub process."),
+			false);
+	}
+}
+#endif /** WITH_LIVELINK_HUB */
+
 TSharedPtr<FLiveLinkHubSubjectController> FLiveLinkHubModule::GetSubjectController() const
 {
 	return LiveLinkHub->SubjectController;
@@ -56,3 +151,4 @@ TSharedPtr<ILiveLinkHubSessionManager> FLiveLinkHubModule::GetSessionManager() c
 }
 
 IMPLEMENT_MODULE(FLiveLinkHubModule, LiveLinkHub);
+#undef LOCTEXT_NAMESPACE /* LiveLinkHubModule */
diff --git a/Engine/Plugins/Animation/LiveLink/Source/LiveLinkHub/Private/LiveLinkHubModule.h b/Engine/Plugins/Animation/LiveLink/Source/LiveLinkHub/Private/LiveLinkHubModule.h
--- a/Engine/Plugins/Animation/LiveLink/Source/LiveLinkHub/Private/LiveLinkHubModule.h
+++ b/Engine/Plugins/Animation/LiveLink/Source/LiveLinkHub/Private/LiveLinkHubModule.h
@@ -11,13 +11,20 @@ class FLiveLinkHubProvider;
 class FLiveLinkHubRecordingController;
 class FLiveLinkHubRecordingListController;
 class FLiveLinkHubSubjectController;
+class FUICommandList;
 class ILiveLinkHubSessionManager;
 
+#ifndef WITH_LIVELINK_HUB
+#define WITH_LIVELINK_HUB 0
+#endif
+
 class FLiveLinkHubModule : public ILiveLinkHubModule
 {
 public:
 	//~ Begin ILiveLinkHubModule interface
 	virtual void StartLiveLinkHub() override;
+	virtual void StartupModule() override;
+	virtual void ShutdownModule() override;
 	//~ End ILiveLinkHubModule interface
 
 	/** Get the livelink hub object. */
@@ -35,6 +42,10 @@ public:
 	/** Get the subject controller. */
     TSharedPtr<ILiveLinkHubSessionManager> GetSessionManager() const;
 
+#if !WITH_LIVELINK_HUB
+	/** Launch livelink hub. */
+	void OpenLiveLinkHub() const;
+#endif
 private:
 	/** LiveLinkHub object responsible for initializing the different controllers. */
 	TSharedPtr<FLiveLinkHub> LiveLinkHub;
