Add ReadOnly mode to livelink

diff --git a/Engine/Plugins/Animation/LiveLink/Source/LiveLinkEditor/Internal/LiveLinkClientPanelViews.h b/Engine/Plugins/Animation/LiveLink/Source/LiveLinkEditor/Internal/LiveLinkClientPanelViews.h
--- a/Engine/Plugins/Animation/LiveLink/Source/LiveLinkEditor/Internal/LiveLinkClientPanelViews.h
+++ b/Engine/Plugins/Animation/LiveLink/Source/LiveLinkEditor/Internal/LiveLinkClientPanelViews.h
@@ -32,8 +32,8 @@ typedef TSharedPtr<FLiveLinkSubjectUIEntry> FLiveLinkSubjectUIEntryPtr;
 
 namespace UE::LiveLink
 {
-	TSharedPtr<IDetailsView> CreateSourcesDetailsView(const TSharedPtr<FLiveLinkSourcesView>& InSourcesView);
-	TSharedPtr<SLiveLinkDataView> CreateSubjectsDetailsView(FLiveLinkClient* InLiveLinkClient);
+	TSharedPtr<IDetailsView> CreateSourcesDetailsView(const TSharedPtr<FLiveLinkSourcesView>& InSourcesView, const TAttribute<bool>& bInReadOnly);
+	TSharedPtr<SLiveLinkDataView> CreateSubjectsDetailsView(FLiveLinkClient* InLiveLinkClient, const TAttribute<bool>& bInReadOnly);
 }
 
 // Structure that defines a single entry in the subject UI
@@ -117,25 +117,45 @@ namespace UE::LiveLink
 		{
 			if (InKeyEvent.GetKey() == EKeys::Delete || InKeyEvent.GetKey() == EKeys::BackSpace)
 			{
-				TArray<ListElementType> SelectedItem = ListType::GetSelectedItems();
-				for (ListElementType Item : SelectedItem)
+				if (!bReadOnly.Get())
 				{
-					Item->RemoveFromClient();
+					TArray<ListElementType> SelectedItem = ListType::GetSelectedItems();
+					for (ListElementType Item : SelectedItem)
+					{
+						Item->RemoveFromClient();
+					}
 				}
 				return FReply::Handled();
 			}
 
 			return FReply::Unhandled();
 		}
+
+		// Whether the panel is in read-only mode or not.
+		TAttribute<bool> bReadOnly;
 	};
 }
 
 class SLiveLinkSourceListView : public UE::LiveLink::SLiveLinkListView<SListView<FLiveLinkSourceUIEntryPtr>, FLiveLinkSourceUIEntryPtr>
 {
+public:
+	void Construct(const FArguments& InArgs, TAttribute<bool> bInReadOnly)
+	{
+		bReadOnly = bInReadOnly;
+
+		UE::LiveLink::SLiveLinkListView<SListView<FLiveLinkSourceUIEntryPtr>, FLiveLinkSourceUIEntryPtr>::Construct(InArgs);
+	}
 };
 
 class SLiveLinkSubjectsTreeView : public UE::LiveLink::SLiveLinkListView<STreeView<FLiveLinkSubjectUIEntryPtr>, FLiveLinkSubjectUIEntryPtr>
 {
+public:
+	void Construct(const FArguments& InArgs, TAttribute<bool> bInReadOnly)
+	{
+		bReadOnly = bInReadOnly;
+
+		UE::LiveLink::SLiveLinkListView<STreeView<FLiveLinkSubjectUIEntryPtr>, FLiveLinkSubjectUIEntryPtr>::Construct(InArgs);
+	}
 };
 
 class FLiveLinkSubjectsView : public TSharedFromThis<FLiveLinkSubjectsView>
@@ -143,7 +163,7 @@ class FLiveLinkSubjectsView : public TSharedFromThis<FLiveLinkSubjectsView>
 public:
 	DECLARE_DELEGATE_TwoParams(FOnSubjectSelectionChanged, FLiveLinkSubjectUIEntryPtr, ESelectInfo::Type);
 
-	FLiveLinkSubjectsView(FOnSubjectSelectionChanged InOnSubjectSelectionChanged, const TSharedPtr<FUICommandList>& InCommandList);
+	FLiveLinkSubjectsView(FOnSubjectSelectionChanged InOnSubjectSelectionChanged, const TSharedPtr<FUICommandList>& InCommandList, TAttribute<bool> bInReadOnly);
 
 	// Helper functions for building the subject tree UI
 	TSharedRef<ITableRow> MakeTreeRowWidget(FLiveLinkSubjectUIEntryPtr InInfo, const TSharedRef<STableViewBase>& OwnerTable);
@@ -167,6 +187,8 @@ public:
 	TArray<FLiveLinkSubjectUIEntryPtr> SubjectData;
 	// Subject Selection Changed delegate
 	FOnSubjectSelectionChanged SubjectSelectionChangedDelegate;
+	// Returns whether the panel is in read-only mode.
+	TAttribute<bool> bReadOnly;
 };
 
 class FLiveLinkSourcesView : public TSharedFromThis<FLiveLinkSourcesView>
@@ -174,7 +196,7 @@ class FLiveLinkSourcesView : public TSharedFromThis<FLiveLinkSourcesView>
 public:
 	DECLARE_DELEGATE_TwoParams(FOnSourceSelectionChanged, FLiveLinkSourceUIEntryPtr, ESelectInfo::Type);
 
-	FLiveLinkSourcesView(FLiveLinkClient* InLiveLinkClient, TSharedPtr<FUICommandList> InCommandList, FOnSourceSelectionChanged InOnSourceSelectionChanged);
+	FLiveLinkSourcesView(FLiveLinkClient* InLiveLinkClient, TSharedPtr<FUICommandList> InCommandList, TAttribute<bool> bInReadOnly, FOnSourceSelectionChanged InOnSourceSelectionChanged);
 
 	// Gather information about all sources and update the list view 
 	void RefreshSourceData(bool bRefreshUI);
@@ -202,4 +224,6 @@ public:
 	FLiveLinkClient* Client;
 	// Source selection changed delegate
 	FOnSourceSelectionChanged OnSourceSelectionChangedDelegate;
+	// Returns whether the panel is in read-only mode.
+	TAttribute<bool> bReadOnly;
 };
diff --git a/Engine/Plugins/Animation/LiveLink/Source/LiveLinkEditor/Internal/LiveLinkPanelController.h b/Engine/Plugins/Animation/LiveLink/Source/LiveLinkEditor/Internal/LiveLinkPanelController.h
--- a/Engine/Plugins/Animation/LiveLink/Source/LiveLinkEditor/Internal/LiveLinkPanelController.h
+++ b/Engine/Plugins/Animation/LiveLink/Source/LiveLinkEditor/Internal/LiveLinkPanelController.h
@@ -5,6 +5,7 @@
 #include "Delegates/Delegate.h"
 #include "Delegates/IDelegateInstance.h"
 #include "Delegates/DelegateCombinations.h"
+#include "Misc/Attribute.h"
 #include "Templates/SharedPointer.h"
 
 class IDetailsView;
@@ -27,7 +28,7 @@ typedef TSharedPtr<FLiveLinkSubjectUIEntry> FLiveLinkSubjectUIEntryPtr;
 class LIVELINKEDITOR_API FLiveLinkPanelController : public TSharedFromThis<FLiveLinkPanelController>
 {
 public:
-	FLiveLinkPanelController();
+	FLiveLinkPanelController(TAttribute<bool> bInReadOnly = false);
 	~FLiveLinkPanelController();
 
 	DECLARE_MULTICAST_DELEGATE_OneParam(FOnSubjectSelectionChanged, const FLiveLinkSubjectKey&);
@@ -86,4 +87,4 @@ public:
 	TSharedPtr<FUICommandList> CommandList;
 	// Delegate called when the subject selection changes
 	FOnSubjectSelectionChanged SubjectSelectionChangedDelegate;
-};
\ No newline at end of file
+};
diff --git a/Engine/Plugins/Animation/LiveLink/Source/LiveLinkEditor/Internal/SLiveLinkDataView.h b/Engine/Plugins/Animation/LiveLink/Source/LiveLinkEditor/Internal/SLiveLinkDataView.h
--- a/Engine/Plugins/Animation/LiveLink/Source/LiveLinkEditor/Internal/SLiveLinkDataView.h
+++ b/Engine/Plugins/Animation/LiveLink/Source/LiveLinkEditor/Internal/SLiveLinkDataView.h
@@ -19,7 +19,10 @@ private:
 	using Super = SCompoundWidget;
 
 public:
-	SLATE_BEGIN_ARGS(SLiveLinkDataView){}
+	SLATE_BEGIN_ARGS(SLiveLinkDataView)
+		: _ReadOnly(false)
+		{}
+	SLATE_ATTRIBUTE(bool, ReadOnly)
 	SLATE_END_ARGS()
 
 	void Construct(const FArguments& Args, FLiveLinkClient* InClient);
diff --git a/Engine/Plugins/Animation/LiveLink/Source/LiveLinkEditor/Private/LiveLinkClientPanel.cpp b/Engine/Plugins/Animation/LiveLink/Source/LiveLinkEditor/Private/LiveLinkClientPanel.cpp
--- a/Engine/Plugins/Animation/LiveLink/Source/LiveLinkEditor/Private/LiveLinkClientPanel.cpp
+++ b/Engine/Plugins/Animation/LiveLink/Source/LiveLinkEditor/Private/LiveLinkClientPanel.cpp
@@ -9,6 +9,7 @@
 #include "LiveLinkClient.h"
 #include "LiveLinkClientPanelToolbar.h"
 #include "LiveLinkClientPanelViews.h"
+#include "LiveLinkEditorSettings.h"
 #include "LiveLinkLog.h"
 #include "LiveLinkRole.h"
 #include "LiveLinkSourceFactory.h"
@@ -37,7 +38,7 @@ void SLiveLinkClientPanel::Construct(const FArguments& Args, FLiveLinkClient* In
 
 	DetailWidgetIndex = 0;
 	
-	PanelController = MakeShared<FLiveLinkPanelController>();
+	PanelController = MakeShared<FLiveLinkPanelController>(TAttribute<bool>::CreateSP(this, &SLiveLinkClientPanel::IsInReadOnlyMode));
 
 	const FName LogName = "Live Link";
 	TSharedPtr<class IMessageLogListing> MessageLogListing;
@@ -67,6 +68,7 @@ void SLiveLinkClientPanel::Construct(const FArguments& Args, FLiveLinkClient* In
 			.AutoHeight()
 			[
 				SNew(SLiveLinkClientPanelToolbar, Client)
+				.Visibility(this, &SLiveLinkClientPanel::GetVisibilityBasedOnReadOnly)
 			]
 			+SVerticalBox::Slot()
 			.FillHeight(1.f)
@@ -181,6 +183,16 @@ int32 SLiveLinkClientPanel::GetDetailWidgetIndex() const
 	return PanelController->SubjectsDetailsView->GetSubjectKey().Source.IsValid() && !PanelController->SubjectsDetailsView->GetSubjectKey().SubjectName.IsNone() ? 1 : 0;
 }
 
+bool SLiveLinkClientPanel::IsInReadOnlyMode() const
+{
+	return GetDefault<ULiveLinkEditorSettings>()->bReadOnly;
+}
+
+EVisibility SLiveLinkClientPanel::GetVisibilityBasedOnReadOnly() const
+{
+	return IsInReadOnlyMode() ? EVisibility::Collapsed : EVisibility::Visible;
+}
+
 void SLiveLinkClientPanel::PostUndo(bool bSuccess)
 {
 	PanelController->SourcesDetailsView->ForceRefresh();
diff --git a/Engine/Plugins/Animation/LiveLink/Source/LiveLinkEditor/Private/LiveLinkClientPanel.h b/Engine/Plugins/Animation/LiveLink/Source/LiveLinkEditor/Private/LiveLinkClientPanel.h
--- a/Engine/Plugins/Animation/LiveLink/Source/LiveLinkEditor/Private/LiveLinkClientPanel.h
+++ b/Engine/Plugins/Animation/LiveLink/Source/LiveLinkEditor/Private/LiveLinkClientPanel.h
@@ -26,7 +26,7 @@ class STableViewBase;
 
 class LIVELINKEDITOR_API SLiveLinkClientPanel : public SCompoundWidget, public FGCObject, public FEditorUndoClient
 {
-	SLATE_BEGIN_ARGS(SLiveLinkClientPanel){}
+	SLATE_BEGIN_ARGS(SLiveLinkClientPanel) {}
 	SLATE_END_ARGS()
 
 	virtual ~SLiveLinkClientPanel() override;
@@ -58,12 +58,19 @@ private:
 
 	// Return the occurrence count and last time occurred text
 	FText GetSelectedMessageOccurrenceText() const;
-private:
+
 	int32 GetDetailWidgetIndex() const;
 	
 	// Handler for the subject tree context menu opening
 	TSharedPtr<SWidget> OnOpenVirtualSubjectContextMenu();
 
+	// Returns whether the panel is in read only mode or not.
+	bool IsInReadOnlyMode() const;
+
+	/** Get widget visibility according to whether the panel is in read-only mode. */
+	EVisibility GetVisibilityBasedOnReadOnly() const;
+
+private:
 	// Handles callback connections between the sources, subjects and details views
 	TSharedPtr<FLiveLinkPanelController> PanelController;
 
diff --git a/Engine/Plugins/Animation/LiveLink/Source/LiveLinkEditor/Private/LiveLinkClientPanelViews.cpp b/Engine/Plugins/Animation/LiveLink/Source/LiveLinkEditor/Private/LiveLinkClientPanelViews.cpp
--- a/Engine/Plugins/Animation/LiveLink/Source/LiveLinkEditor/Private/LiveLinkClientPanelViews.cpp
+++ b/Engine/Plugins/Animation/LiveLink/Source/LiveLinkEditor/Private/LiveLinkClientPanelViews.cpp
@@ -41,7 +41,7 @@ namespace SubjectTreeUI
 
 namespace UE::LiveLink
 {
-TSharedPtr<IDetailsView> CreateSourcesDetailsView(const TSharedPtr<FLiveLinkSourcesView>& InSourcesView)
+TSharedPtr<IDetailsView> CreateSourcesDetailsView(const TSharedPtr<FLiveLinkSourcesView>& InSourcesView, const TAttribute<bool>& bInReadOnly)
 {
 	FPropertyEditorModule& PropertyEditorModule = FModuleManager::GetModuleChecked<FPropertyEditorModule>("PropertyEditor");
 
@@ -54,13 +54,18 @@ TSharedPtr<IDetailsView> CreateSourcesDetailsView(const TSharedPtr<FLiveLinkSour
 	TSharedPtr<IDetailsView> SettingsDetailsView = PropertyEditorModule.CreateDetailView(DetailsViewArgs);
 	// todo: use controller here instead of view widget
 	SettingsDetailsView->OnFinishedChangingProperties().AddRaw(InSourcesView.Get(), &FLiveLinkSourcesView::OnPropertyChanged);
+	SettingsDetailsView->SetIsPropertyEditingEnabledDelegate(FIsPropertyEditingEnabled::CreateLambda(
+		[bInReadOnly](){
+		return !bInReadOnly.Get();
+	}));
 
 	return SettingsDetailsView;
 }
 
-TSharedPtr<SLiveLinkDataView> CreateSubjectsDetailsView(FLiveLinkClient* InLiveLinkClient)
+TSharedPtr<SLiveLinkDataView> CreateSubjectsDetailsView(FLiveLinkClient* InLiveLinkClient, const TAttribute<bool>& bInReadOnly)
 {
-	return SNew(SLiveLinkDataView, InLiveLinkClient);
+	return SNew(SLiveLinkDataView, InLiveLinkClient)
+		.ReadOnly(bInReadOnly);
 }
 } // namespace UE::LiveLink
 
@@ -177,12 +182,14 @@ public:
 	SLATE_BEGIN_ARGS(SLiveLinkClientPanelSubjectRow) {}
 	/** The list item for this row */
 	SLATE_ARGUMENT(FLiveLinkSubjectUIEntryPtr, Entry)
+	SLATE_ATTRIBUTE(bool, ReadOnly)
 	SLATE_END_ARGS()
 
 
 	void Construct(const FArguments& Args, const TSharedRef<STableViewBase>& OwnerTableView)
 	{
 		EntryPtr = Args._Entry;
+		bReadOnly = Args._ReadOnly;
 
 		SMultiColumnTableRow<FLiveLinkSubjectUIEntryPtr>::Construct(
 			FSuperRowType::FArguments()
@@ -199,6 +206,7 @@ public:
 			if (EntryPtr->IsSubject())
 			{
 				return SNew(SCheckBox)
+					.Visibility(this, &SLiveLinkClientPanelSubjectRow::GetVisibilityFromReadOnly)
 					.IsChecked(MakeAttributeSP(this, &SLiveLinkClientPanelSubjectRow::GetSubjectEnabled))
 					.OnCheckStateChanged(this, &SLiveLinkClientPanelSubjectRow::OnEnabledChanged);
 			}
@@ -238,6 +246,7 @@ public:
 					.ContentPadding(0.f)
 					.ForegroundColor(FSlateColor::UseForeground())
 					.IsFocusable(false)
+					.Visibility(this, &SLiveLinkClientPanelSubjectRow::GetVisibilityFromReadOnly)
 					[
 						SNew(SImage)
 						.Image(FAppStyle::GetBrush("Icons.Delete"))
@@ -249,6 +258,7 @@ public:
 				return SNew(SBox)
 					.HAlign(HAlign_Center)
 					.VAlign(VAlign_Center)
+					.Visibility(this, &SLiveLinkClientPanelSubjectRow::GetVisibilityFromReadOnly)
 					[
 						SNew(STextBlock)
 						.Font(FAppStyle::Get().GetFontStyle("FontAwesome.8"))
@@ -290,7 +300,16 @@ private:
 		return FLinearColor(0.f, 0.f, 0.f, 0.f);
 	}
 
+	/** Get widget visibility according to whether or not the panel is in read-only mode. */
+	EVisibility GetVisibilityFromReadOnly() const
+	{
+		return bReadOnly.Get() ? EVisibility::Collapsed : EVisibility::Visible;
+	}
+
 	FLiveLinkSubjectUIEntryPtr EntryPtr;
+
+	/** Returns whether the panel is in read-only mode. */
+	TAttribute<bool> bReadOnly;
 };
 
 class SLiveLinkClientPanelSourcesRow : public SMultiColumnTableRow<FLiveLinkSourceUIEntryPtr>
@@ -299,11 +318,13 @@ public:
 	SLATE_BEGIN_ARGS(SLiveLinkClientPanelSourcesRow) {}
 	/** The list item for this row */
 		SLATE_ARGUMENT(FLiveLinkSourceUIEntryPtr, Entry)
+		SLATE_ATTRIBUTE(bool, ReadOnly)
 	SLATE_END_ARGS()
 
 	void Construct(const FArguments& Args, const TSharedRef<STableViewBase>& OwnerTableView)
 	{
 		EntryPtr = Args._Entry;
+		bReadOnly = Args._ReadOnly;
 
 		SMultiColumnTableRow<FLiveLinkSourceUIEntryPtr>::Construct(
 			FSuperRowType::FArguments()
@@ -336,6 +357,7 @@ public:
 				.ButtonStyle(FAppStyle::Get(), "HoverHintOnly")
 				.HAlign(HAlign_Center)
 				.VAlign(VAlign_Center)
+				.Visibility(this, &SLiveLinkClientPanelSourcesRow::GetVisibilityFromReadOnly)
 				.OnClicked(this, &SLiveLinkClientPanelSourcesRow::OnRemoveClicked)
 				.ToolTipText(LOCTEXT("RemoveSource", "Remove selected live link source"))
 				.ContentPadding(0.f)
@@ -368,12 +390,23 @@ private:
 		return FReply::Handled();
 	}
 
+	/** Get widget visibility according to whether or not the panel is in read-only mode. */
+	EVisibility GetVisibilityFromReadOnly() const
+	{
+		return bReadOnly.Get() ? EVisibility::Collapsed : EVisibility::Visible;
+	}
+
+private:
 	FLiveLinkSourceUIEntryPtr EntryPtr;
+
+	/** Attribute used to query whether the panel is in read only mode or not. */
+	TAttribute<bool> bReadOnly;
 };
 
-FLiveLinkSourcesView::FLiveLinkSourcesView(FLiveLinkClient* InLiveLinkClient, TSharedPtr<FUICommandList> InCommandList, FOnSourceSelectionChanged InOnSourceSelectionChanged)
+FLiveLinkSourcesView::FLiveLinkSourcesView(FLiveLinkClient* InLiveLinkClient, TSharedPtr<FUICommandList> InCommandList, TAttribute<bool> bInReadOnly, FOnSourceSelectionChanged InOnSourceSelectionChanged)
 	: Client(InLiveLinkClient)
 	, OnSourceSelectionChangedDelegate(MoveTemp(InOnSourceSelectionChanged))
+	, bReadOnly(MoveTemp(bInReadOnly))
 {
 	CreateSourcesListView(InCommandList);
 }
@@ -381,7 +414,8 @@ FLiveLinkSourcesView::FLiveLinkSourcesView(FLiveLinkClient* InLiveLinkClient, TS
 TSharedRef<ITableRow> FLiveLinkSourcesView::MakeSourceListViewWidget(FLiveLinkSourceUIEntryPtr Entry, const TSharedRef<STableViewBase>& OwnerTable) const
 {
 	return SNew(SLiveLinkClientPanelSourcesRow, OwnerTable)
-		.Entry(Entry);
+		.Entry(Entry)
+		.ReadOnly(bReadOnly);
 }
 
 void FLiveLinkSourcesView::OnSourceListSelectionChanged(FLiveLinkSourceUIEntryPtr Entry, ESelectInfo::Type SelectionType) const
@@ -391,7 +425,7 @@ void FLiveLinkSourcesView::OnSourceListSelectionChanged(FLiveLinkSourceUIEntryPt
 
 void FLiveLinkSourcesView::CreateSourcesListView(const TSharedPtr<FUICommandList>& InCommandList)
 {
-	SAssignNew(SourcesListView, SLiveLinkSourceListView)
+	SAssignNew(SourcesListView, SLiveLinkSourceListView, bReadOnly)
 		.ListItemsSource(&SourceData)
 		.SelectionMode(ESelectionMode::Single)
 		.OnGenerateRow_Raw(this, &FLiveLinkSourcesView::MakeSourceListViewWidget)
@@ -417,6 +451,11 @@ void FLiveLinkSourcesView::CreateSourcesListView(const TSharedPtr<FUICommandList
 
 TSharedPtr<SWidget> FLiveLinkSourcesView::OnSourceConstructContextMenu(TSharedPtr<FUICommandList> InCommandList)
 {
+	if (bReadOnly.Get())
+	{
+		return nullptr;
+	}
+
 	const bool bShouldCloseWindowAfterMenuSelection = true;
 	FMenuBuilder MenuBuilder(bShouldCloseWindowAfterMenuSelection, InCommandList);
 
@@ -464,8 +503,9 @@ bool FLiveLinkSourcesView::CanRemoveSource()
 	return SourcesListView->GetNumItemsSelected() > 0;
 }
 
-FLiveLinkSubjectsView::FLiveLinkSubjectsView(FOnSubjectSelectionChanged InOnSubjectSelectionChanged, const TSharedPtr<FUICommandList>& InCommandList)
+FLiveLinkSubjectsView::FLiveLinkSubjectsView(FOnSubjectSelectionChanged InOnSubjectSelectionChanged, const TSharedPtr<FUICommandList>& InCommandList, TAttribute<bool> bInReadOnly)
 	: SubjectSelectionChangedDelegate(InOnSubjectSelectionChanged)
+	, bReadOnly(MoveTemp(bInReadOnly))
 {
 	CreateSubjectsTreeView(InCommandList);
 }
@@ -488,7 +528,8 @@ void FLiveLinkSourcesView::OnPropertyChanged(const FPropertyChangedEvent& InEven
 TSharedRef<ITableRow> FLiveLinkSubjectsView::MakeTreeRowWidget(FLiveLinkSubjectUIEntryPtr InInfo, const TSharedRef<STableViewBase>& OwnerTable)
 {
 	return SNew(SLiveLinkClientPanelSubjectRow, OwnerTable)
-		.Entry(InInfo);
+		.Entry(InInfo)
+		.ReadOnly(bReadOnly);
 }
 
 void FLiveLinkSubjectsView::GetChildrenForInfo(FLiveLinkSubjectUIEntryPtr InInfo, TArray< FLiveLinkSubjectUIEntryPtr >& OutChildren)
@@ -498,6 +539,11 @@ void FLiveLinkSubjectsView::GetChildrenForInfo(FLiveLinkSubjectUIEntryPtr InInfo
 
 TSharedPtr<SWidget> FLiveLinkSubjectsView::OnOpenVirtualSubjectContextMenu(TSharedPtr<FUICommandList> InCommandList)
 {
+	if (bReadOnly.Get())
+	{
+		return nullptr;
+	}
+
 	constexpr bool bShouldCloseWindowAfterMenuSelection = true;
 	FMenuBuilder MenuBuilder(bShouldCloseWindowAfterMenuSelection, InCommandList);
 
@@ -522,7 +568,7 @@ bool FLiveLinkSubjectsView::CanRemoveSubject() const
 
 void FLiveLinkSubjectsView::CreateSubjectsTreeView(const TSharedPtr<FUICommandList>& InCommandList)
 {
-	SAssignNew(SubjectsTreeView, SLiveLinkSubjectsTreeView)
+	SAssignNew(SubjectsTreeView, SLiveLinkSubjectsTreeView, bReadOnly)
 		.TreeItemsSource(&SubjectData)
 		.OnGenerateRow_Raw(this, &FLiveLinkSubjectsView::MakeTreeRowWidget)
 		.OnGetChildren_Raw(this, &FLiveLinkSubjectsView::GetChildrenForInfo)
@@ -548,4 +594,4 @@ void FLiveLinkSubjectsView::CreateSubjectsTreeView(const TSharedPtr<FUICommandLi
 }
 
 
-#undef LOCTEXT_NAMESPACE /**LiveLinkClientPanel*/
\ No newline at end of file
+#undef LOCTEXT_NAMESPACE /**LiveLinkClientPanel*/
diff --git a/Engine/Plugins/Animation/LiveLink/Source/LiveLinkEditor/Private/LiveLinkEditorModule.cpp b/Engine/Plugins/Animation/LiveLink/Source/LiveLinkEditor/Private/LiveLinkEditorModule.cpp
--- a/Engine/Plugins/Animation/LiveLink/Source/LiveLinkEditor/Private/LiveLinkEditorModule.cpp
+++ b/Engine/Plugins/Animation/LiveLink/Source/LiveLinkEditor/Private/LiveLinkEditorModule.cpp
@@ -27,6 +27,7 @@
 #include "LiveLinkClientPanel.h"
 #include "LiveLinkClientCommands.h"
 #include "LiveLinkEditorPrivate.h"
+#include "LiveLinkEditorSettings.h"
 #include "LiveLinkGraphPanelPinFactory.h"
 #include "LiveLinkSettings.h"
 #include "LiveLinkSourceSettingsDetailCustomization.h"
@@ -219,9 +220,14 @@ private:
 
 	void RegisterSettings()
 	{
-		ISettingsModule* SettingsModule = FModuleManager::GetModulePtr<ISettingsModule>("Settings");
-		if (SettingsModule != nullptr)
+		if (ISettingsModule* SettingsModule = FModuleManager::GetModulePtr<ISettingsModule>("Settings"))
 		{
+			SettingsModule->RegisterSettings("Editor", "Plugins", "LiveLink",
+				LOCTEXT("EditorSettingsName", "Live Link"),
+				LOCTEXT("EditorSettingsDescription", "Configure Live Link."),
+				GetMutableDefault<ULiveLinkEditorSettings>()
+			);
+
 			SettingsModule->RegisterSettings("Project", "Plugins", "LiveLink",
 				LOCTEXT("LiveLinkSettingsName", "Live Link"),
 				LOCTEXT("LiveLinkDescription", "Configure the Live Link plugin."),
@@ -244,6 +250,7 @@ private:
 		{
 			SettingsModule->UnregisterSettings("Project", "Plugins", "LiveLinkComponent");
 			SettingsModule->UnregisterSettings("Project", "Plugins", "LiveLink");
+			SettingsModule->UnregisterSettings("Editor", "Plugins", "LiveLink");
 		}
 	}
 
diff --git a/Engine/Plugins/Animation/LiveLink/Source/LiveLinkEditor/Private/LiveLinkPanelController.cpp b/Engine/Plugins/Animation/LiveLink/Source/LiveLinkEditor/Private/LiveLinkPanelController.cpp
--- a/Engine/Plugins/Animation/LiveLink/Source/LiveLinkEditor/Private/LiveLinkPanelController.cpp
+++ b/Engine/Plugins/Animation/LiveLink/Source/LiveLinkEditor/Private/LiveLinkPanelController.cpp
@@ -14,7 +14,7 @@
 #define WITH_LIVELINK_HUB 0
 #endif
 
-FLiveLinkPanelController::FLiveLinkPanelController()
+FLiveLinkPanelController::FLiveLinkPanelController(TAttribute<bool> bInReadOnly)
 {
 	Client = (FLiveLinkClient*)&IModularFeatures::Get().GetModularFeature<ILiveLinkClient>(ILiveLinkClient::ModularFeatureName);
 
@@ -24,10 +24,10 @@ FLiveLinkPanelController::FLiveLinkPanelController()
 	CommandList = MakeShared<FUICommandList>();
 	BindCommands();
 
-	SourcesView = MakeShared<FLiveLinkSourcesView>(Client, CommandList, FLiveLinkSourcesView::FOnSourceSelectionChanged::CreateRaw(this, &FLiveLinkPanelController::OnSourceSelectionChangedHandler));
-	SubjectsView = MakeShared<FLiveLinkSubjectsView>(FLiveLinkSubjectsView::FOnSubjectSelectionChanged::CreateRaw(this, &FLiveLinkPanelController::OnSubjectSelectionChangedHandler), CommandList);
-	SourcesDetailsView = UE::LiveLink::CreateSourcesDetailsView(SourcesView);
-	SubjectsDetailsView = UE::LiveLink::CreateSubjectsDetailsView(Client);
+	SourcesView = MakeShared<FLiveLinkSourcesView>(Client, CommandList, bInReadOnly, FLiveLinkSourcesView::FOnSourceSelectionChanged::CreateRaw(this, &FLiveLinkPanelController::OnSourceSelectionChangedHandler));
+	SubjectsView = MakeShared<FLiveLinkSubjectsView>(FLiveLinkSubjectsView::FOnSubjectSelectionChanged::CreateRaw(this, &FLiveLinkPanelController::OnSubjectSelectionChangedHandler), CommandList, bInReadOnly);
+	SourcesDetailsView = UE::LiveLink::CreateSourcesDetailsView(SourcesView, bInReadOnly);
+	SubjectsDetailsView = UE::LiveLink::CreateSubjectsDetailsView(Client, bInReadOnly);
 
 	RebuildSourceList();
 	RebuildSubjectList();
@@ -285,4 +285,4 @@ void FLiveLinkPanelController::HandleRemoveSubject()
 	{
 		Selected[0]->RemoveFromClient();
 	}
-}
\ No newline at end of file
+}
diff --git a/Engine/Plugins/Animation/LiveLink/Source/LiveLinkEditor/Private/SLiveLinkDataView.cpp b/Engine/Plugins/Animation/LiveLink/Source/LiveLinkEditor/Private/SLiveLinkDataView.cpp
--- a/Engine/Plugins/Animation/LiveLink/Source/LiveLinkEditor/Private/SLiveLinkDataView.cpp
+++ b/Engine/Plugins/Animation/LiveLink/Source/LiveLinkEditor/Private/SLiveLinkDataView.cpp
@@ -99,6 +99,7 @@ void SLiveLinkDataView::Construct(const FArguments& Args, FLiveLinkClient* InCli
 
 	SettingsDetailsView = PropertyEditorModule.CreateDetailView(DetailsViewArgs);
 	SettingsDetailsView->OnFinishedChangingProperties().AddSP(this, &SLiveLinkDataView::OnPropertyChanged);
+	SettingsDetailsView->SetIsPropertyEditingEnabledDelegate(FIsPropertyEditingEnabled::CreateLambda([bReadOnly = Args._ReadOnly](){ return !bReadOnly.Get(); }));
 
 	FStructureDetailsViewArgs StructViewArgs;
 	StructureDetailsView = PropertyEditorModule.CreateStructureDetailView(DetailsViewArgs, StructViewArgs, TSharedPtr<FStructOnScope>());
diff --git a/Engine/Plugins/Animation/LiveLink/Source/LiveLinkEditor/Public/LiveLinkEditorSettings.h b/Engine/Plugins/Animation/LiveLink/Source/LiveLinkEditor/Public/LiveLinkEditorSettings.h
new file mode 100644
--- /dev/null
+++ b/Engine/Plugins/Animation/LiveLink/Source/LiveLinkEditor/Public/LiveLinkEditorSettings.h
@@ -0,0 +1,21 @@
+// Copyright Epic Games, Inc. All Rights Reserved.
+
+#pragma once
+
+#include "UObject/Object.h"
+
+#include "LiveLinkEditorSettings.generated.h"
+
+/**
+ * Editor preferences settings for LiveLink.
+ */
+UCLASS(config = EditorPerProjectUserSettings)
+class LIVELINKEDITOR_API ULiveLinkEditorSettings : public UObject
+{
+	GENERATED_BODY()
+
+public:
+	/** Whether the livelink panel should be in read-only mode. */
+	UPROPERTY(config, EditAnywhere, Category = "LiveLink")
+	bool bReadOnly = false;
+};
