LiveLinkHub - Improve boot time by disabling shared DDC and autosdk
Add some cpu scopes
Add optional monolithic mode

diff --git a/Engine/Plugins/Animation/LiveLink/Source/LiveLinkHub/Private/LiveLinkHub.cpp b/Engine/Plugins/Animation/LiveLink/Source/LiveLinkHub/Private/LiveLinkHub.cpp
--- a/Engine/Plugins/Animation/LiveLink/Source/LiveLinkHub/Private/LiveLinkHub.cpp
+++ b/Engine/Plugins/Animation/LiveLink/Source/LiveLinkHub/Private/LiveLinkHub.cpp
@@ -29,6 +29,8 @@
 
 void FLiveLinkHub::Initialize()
 {
+	TRACE_CPUPROFILER_EVENT_SCOPE(FLiveLinkHub::Initialize);
+
 	// We must register the livelink client first since we might rely on the modular feature to initialize the controllers/managers.
 	LiveLinkHubClient = MakeShared<FLiveLinkHubClient>(AsShared());
 	IModularFeatures::Get().RegisterModularFeature(ILiveLinkClient::ModularFeatureName, LiveLinkHubClient.Get());
@@ -39,12 +41,15 @@ void FLiveLinkHub::Initialize()
 	CommandExecutor = MakeUnique<FConsoleCommandExecutor>();
 	IModularFeatures::Get().RegisterModularFeature(IConsoleCommandExecutor::ModularFeatureName(), CommandExecutor.Get());
 
-	RecordingController = MakeShared<FLiveLinkHubRecordingController>();
-	PlaybackController = MakeShared<FLiveLinkHubPlaybackController>();
-	RecordingListController = MakeShared<FLiveLinkHubRecordingListController>(AsShared());
-	ClientsController = MakeShared<FLiveLinkHubClientsController>(LiveLinkProvider.ToSharedRef());
-	CommandList = MakeShared<FUICommandList>();
-	SubjectController = MakeShared<FLiveLinkHubSubjectController>();
+	{
+		TRACE_CPUPROFILER_EVENT_SCOPE(FLiveLinkHub::InitializeControllers);
+		RecordingController = MakeShared<FLiveLinkHubRecordingController>();
+		PlaybackController = MakeShared<FLiveLinkHubPlaybackController>();
+		RecordingListController = MakeShared<FLiveLinkHubRecordingListController>(AsShared());
+		ClientsController = MakeShared<FLiveLinkHubClientsController>(LiveLinkProvider.ToSharedRef());
+		CommandList = MakeShared<FUICommandList>();
+		SubjectController = MakeShared<FLiveLinkHubSubjectController>();
+	}
 
 
 	FLiveLinkHubCommands::Register();
diff --git a/Engine/Plugins/Animation/LiveLink/Source/LiveLinkHub/Private/LiveLinkHubModule.cpp b/Engine/Plugins/Animation/LiveLink/Source/LiveLinkHub/Private/LiveLinkHubModule.cpp
--- a/Engine/Plugins/Animation/LiveLink/Source/LiveLinkHub/Private/LiveLinkHubModule.cpp
+++ b/Engine/Plugins/Animation/LiveLink/Source/LiveLinkHub/Private/LiveLinkHubModule.cpp
@@ -18,6 +18,7 @@
 
 void FLiveLinkHubModule::StartLiveLinkHub()
 {
+	TRACE_CPUPROFILER_EVENT_SCOPE(StartLiveLinkHub);
 	// Nothing will get executed after this, so put everything before.
 	LiveLinkHub = MakeShared<FLiveLinkHub>();
 	LiveLinkHub->Initialize();
diff --git a/Engine/Plugins/Animation/LiveLink/Source/LiveLinkHub/Private/UI/Window/LiveLinkHubWindowController.cpp b/Engine/Plugins/Animation/LiveLink/Source/LiveLinkHub/Private/UI/Window/LiveLinkHubWindowController.cpp
--- a/Engine/Plugins/Animation/LiveLink/Source/LiveLinkHub/Private/UI/Window/LiveLinkHubWindowController.cpp
+++ b/Engine/Plugins/Animation/LiveLink/Source/LiveLinkHub/Private/UI/Window/LiveLinkHubWindowController.cpp
@@ -62,6 +62,8 @@ TSharedRef<SWindow> FLiveLinkHubWindowController::CreateWindow()
 }
 void FLiveLinkHubWindowController::RestoreLayout()
 {
+	TRACE_CPUPROFILER_EVENT_SCOPE(FLiveLinkHubWindowController::RestoreLayout);
+
 	const TSharedRef<FTabManager::FLayout> DefaultLayout = FTabManager::NewLayout("LiveLinkHub_v1.0");
 	const TSharedRef<FTabManager::FArea> MainWindowArea = FTabManager::NewPrimaryArea()->SetOrientation(Orient_Horizontal);
 	const TSharedRef<FTabManager::FStack> MainStack = FTabManager::NewStack();
@@ -70,12 +72,20 @@ void FLiveLinkHubWindowController::RestoreLayout()
 
 	MainWindowArea->Split(MainStack);
 	DefaultLayout->AddArea(MainWindowArea);
+	PersistentLayout = DefaultLayout;
 
-	PersistentLayout = FLayoutSaveRestore::LoadFromConfig(LiveLinkHubLayoutIni, DefaultLayout);
+	{
+		TRACE_CPUPROFILER_EVENT_SCOPE(FLiveLinkHubWindowController::LoadFromConfig);
+		PersistentLayout = FLayoutSaveRestore::LoadFromConfig(LiveLinkHubLayoutIni, DefaultLayout);
+	}
 
-	constexpr bool bEmbedTitleAreaContent = true;
-	const TSharedPtr<SWidget> Content = FGlobalTabmanager::Get()->RestoreFrom(PersistentLayout.ToSharedRef(), RootWindow, bEmbedTitleAreaContent, EOutputCanBeNullptr::Never);
-	RootWindow->SetContent(Content.ToSharedRef());
+
+	{
+		TRACE_CPUPROFILER_EVENT_SCOPE(FLiveLinkHubWindowController::RestoreFrom);
+		constexpr bool bEmbedTitleAreaContent = true;
+		const TSharedPtr<SWidget> Content = FGlobalTabmanager::Get()->RestoreFrom(PersistentLayout.ToSharedRef(), RootWindow, bEmbedTitleAreaContent, EOutputCanBeNullptr::Never);
+		RootWindow->SetContent(Content.ToSharedRef());
+	}
 
 	RootWindow->ShowWindow();
 	constexpr bool bForceWindowToFront = true;
@@ -86,6 +96,7 @@ void FLiveLinkHubWindowController::RestoreLayout()
 
 TSharedPtr<FModalWindowManager> FLiveLinkHubWindowController::InitializeSlateApplication()
 {
+	TRACE_CPUPROFILER_EVENT_SCOPE(FLiveLinkHubWindowController::InitializeAsStandaloneApplication);
 	FSlateApplication::InitializeAsStandaloneApplication(GetStandardStandaloneRenderer());
 
 	const FText ApplicationTitle = LOCTEXT("AppTitle", "LiveLink Hub");
diff --git a/Engine/Source/Programs/LiveLinkHub/LiveLinkHub.Target.cs b/Engine/Source/Programs/LiveLinkHub/LiveLinkHub.Target.cs
--- a/Engine/Source/Programs/LiveLinkHub/LiveLinkHub.Target.cs
+++ b/Engine/Source/Programs/LiveLinkHub/LiveLinkHub.Target.cs
@@ -2,14 +2,17 @@
 
 using UnrealBuildTool;
 using System.Collections.Generic;
+using EpicGames.Core;
 
 [SupportedPlatforms(UnrealPlatformClass.Desktop)]
 public class LiveLinkHubTarget : TargetRules
 {
+	[CommandLine("-Monolithic")]
+	public bool bMonolithic = false;
 	public LiveLinkHubTarget(TargetInfo Target) : base(Target)
 	{
 		Type = TargetType.Program;
-		LinkType = TargetLinkType.Modular;
+		LinkType = bMonolithic ? TargetLinkType.Monolithic : TargetLinkType.Modular;
 		LaunchModuleName = "LiveLinkHubLauncher";
 
 		IncludeOrderVersion = EngineIncludeOrderVersion.Latest;
@@ -46,8 +49,10 @@ public class LiveLinkHubTarget : TargetRules
 		bIsBuildingConsoleApplication = false;
 
 		GlobalDefinitions.Add("WITH_LIVELINK_HUB=1");
+		GlobalDefinitions.Add("AUTOSDKS_ENABLED=0");
+		GlobalDefinitions.Add("NO_LOGGING=0");
 
-		bEnableTrace = false;
+		bEnableTrace = true;
 
 		OptedInModulePlatforms = new UnrealTargetPlatform[] { UnrealTargetPlatform.Win64, UnrealTargetPlatform.Mac,
 															  UnrealTargetPlatform.Linux, UnrealTargetPlatform.LinuxArm64 };
diff --git a/Engine/Source/Programs/LiveLinkHub/Private/LiveLinkHubRun.cpp b/Engine/Source/Programs/LiveLinkHub/Private/LiveLinkHubRun.cpp
--- a/Engine/Source/Programs/LiveLinkHub/Private/LiveLinkHubRun.cpp
+++ b/Engine/Source/Programs/LiveLinkHub/Private/LiveLinkHubRun.cpp
@@ -11,7 +11,7 @@ int32 RunLiveLinkHub(const TCHAR* CommandLine)
 	FTaskTagScope Scope(ETaskTag::EGameThread);
 	
 	// Start up the main loop, adding some extra command line arguments:
-	const int32 Result = GEngineLoop.PreInit(*FString::Printf(TEXT("%s %s"), CommandLine, TEXT("LiveLinkHubCommandlet -Messaging")));
+	const int32 Result = GEngineLoop.PreInit(*FString::Printf(TEXT("%s %s"), CommandLine, TEXT("LiveLinkHubCommandlet -Messaging -DDC=NoShared -NoShaderCompile")));
 
 	if (Result != 0)
 	{
