LiveLinkHub - Fix subject name always being none

diff --git a/Engine/Plugins/Animation/LiveLink/Source/LiveLinkHub/Private/Subjects/LiveLinkHubSubjectSessionConfig.cpp b/Engine/Plugins/Animation/LiveLink/Source/LiveLinkHub/Private/Subjects/LiveLinkHubSubjectSessionConfig.cpp
--- a/Engine/Plugins/Animation/LiveLink/Source/LiveLinkHub/Private/Subjects/LiveLinkHubSubjectSessionConfig.cpp
+++ b/Engine/Plugins/Animation/LiveLink/Source/LiveLinkHub/Private/Subjects/LiveLinkHubSubjectSessionConfig.cpp
@@ -98,9 +98,9 @@ void ULiveLinkHubSubjectSessionConfig::OnSubjectRemoved(const FLiveLinkSubjectKe
 
 void ULiveLinkHubSubjectProxy::Initialize(const FLiveLinkSubjectKey& InSubjectKey, FString InSource)
 {
-	SubjectName = SubjectKey.SubjectName.Name.ToString();
+	SubjectName = InSubjectKey.SubjectName.Name.ToString();
 	SubjectKey = InSubjectKey;
-	OutboundName = SubjectKey.SubjectName.Name.ToString();
+	OutboundName = InSubjectKey.SubjectName.Name.ToString();
 	Source = MoveTemp(InSource);
 }
 
