LiveLinkHub - Add ability to pick which clients to use in a given session
Persist subject outbound names
Fix subject name being none

diff --git a/Engine/Plugins/Animation/LiveLink/Source/LiveLinkHub/Private/Clients/LiveLinkHubClientsController.h b/Engine/Plugins/Animation/LiveLink/Source/LiveLinkHub/Private/Clients/LiveLinkHubClientsController.h
--- a/Engine/Plugins/Animation/LiveLink/Source/LiveLinkHub/Private/Clients/LiveLinkHubClientsController.h
+++ b/Engine/Plugins/Animation/LiveLink/Source/LiveLinkHub/Private/Clients/LiveLinkHubClientsController.h
@@ -7,16 +7,19 @@
 #include "IDetailsView.h"
 #include "IStructureDetailsView.h"
 #include "LiveLinkHub.h"
+#include "LiveLinkHubModule.h"
 #include "LiveLinkHubUEClientInfo.h"
 #include "LiveLinkHubProvider.h"
 #include "LiveLinkTypes.h"
+#include "Modules/ModuleManager.h"
 #include "PropertyEditorModule.h"
+#include "Session/LiveLinkHubSession.h"
+#include "Session/LiveLinkHubSessionManager.h"
 #include "SLiveLinkHubClientsView.h"
 #include "UObject/StructOnScope.h"
 #include "Widgets/SNullWidget.h"
 #include "Widgets/Text/STextBlock.h"
 
-
 /** Controller responsible for holding the list of connected clients and creating the clients view. */
 class FLiveLinkHubClientsController
 {
@@ -24,6 +27,9 @@ public:
 	FLiveLinkHubClientsController(const TSharedRef<ILiveLinkHubClientsModel>& InClientsModel)
 		: ClientsModel(InClientsModel)
 	{
+		const FLiveLinkHubModule& LiveLinkHubModule = FModuleManager::Get().GetModuleChecked<FLiveLinkHubModule>("LiveLinkHub");
+		LiveLinkHubModule.GetSessionManager()->OnActiveSessionChanged().AddRaw(this, &FLiveLinkHubClientsController::OnActiveSessionChanged);
+
 		ClientsModel->OnClientEvent().AddRaw(this, &FLiveLinkHubClientsController::OnClientEvent);
 		
 		UEClientInfo = MakeShared<TStructOnScope<FLiveLinkHubUEClientInfo>>();
@@ -31,6 +37,14 @@ public:
 
 	~FLiveLinkHubClientsController()
 	{
+		if (const FLiveLinkHubModule* LiveLinkHubModule = FModuleManager::Get().GetModulePtr<FLiveLinkHubModule>("LiveLinkHub"))
+		{
+			if (TSharedPtr<ILiveLinkHubSessionManager> SessionManager = LiveLinkHubModule->GetSessionManager())
+			{
+				SessionManager->OnActiveSessionChanged().RemoveAll(this);
+			}
+		}
+
 		ClientsModel->OnClientEvent().RemoveAll(this);
 	}
 
@@ -38,7 +52,9 @@ public:
 	TSharedRef<SWidget> MakeClientsView()
 	{
 		return SAssignNew(ClientsView, SLiveLinkHubClientsView, ClientsModel.ToSharedRef())
-			.OnClientSelected_Raw(this, &FLiveLinkHubClientsController::UpdateClientDetails);
+			.OnClientSelected_Raw(this, &FLiveLinkHubClientsController::UpdateClientDetails)
+			.OnDiscoveredClientPicked_Raw(this, &FLiveLinkHubClientsController::OnDiscoveredClientPicked)
+			.OnRemoveClientFromSession_Raw(this, &FLiveLinkHubClientsController::OnRemoveClientFromSession);
 	}
 
 	/** Create the widget that displays information about a given UE client. */
@@ -67,19 +83,14 @@ private:
 	{
 		switch (EventType)
 		{
-			case ILiveLinkHubClientsModel::EClientEventType::Connected:
+			case ILiveLinkHubClientsModel::EClientEventType::Discovered:
 				break;
-			case ILiveLinkHubClientsModel::EClientEventType::Removed:
+			case ILiveLinkHubClientsModel::EClientEventType::Disconnected:
 			{
-				if (TOptional<FLiveLinkHubClientId> SelectedClient = ClientsView->GetSelectedClient())
-				{
-					if (ClientIdentifier == *SelectedClient)
-					{
-						StructDetailsView->SetStructureData(nullptr);
-					}
-				}
 				break;
 			}
+			case ILiveLinkHubClientsModel::EClientEventType::Reestablished:
+				// fallthrough
 			case ILiveLinkHubClientsModel::EClientEventType::Modified:
 			{
 				if (TOptional<FLiveLinkHubClientId> SelectedClient = ClientsView->GetSelectedClient())
@@ -97,6 +108,16 @@ private:
 			}
 		}
 	}
+
+	void OnActiveSessionChanged(const TSharedRef<ILiveLinkHubSession>& ActiveSession)
+	{
+		if (ClientsView)
+		{
+			ClientsView->Reinitialize();
+		}
+
+		StructDetailsView->SetStructureData(nullptr);
+	}
 	
 	/** Handles updating client details in the client details panel. */
 	void UpdateClientDetails(FLiveLinkHubClientId Client)
@@ -113,6 +134,38 @@ private:
 		}
 	}
 
+	void OnDiscoveredClientPicked(FLiveLinkHubClientId Client)
+	{
+		if (TSharedPtr<ILiveLinkHubSessionManager> SessionManager = FModuleManager::Get().GetModuleChecked<FLiveLinkHubModule>("LiveLinkHub").GetLiveLinkHub()->GetSessionManager())
+		{
+			if (TSharedPtr<ILiveLinkHubSession> CurrentSession = SessionManager->GetCurrentSession())
+			{
+				CurrentSession->AddClient(Client);
+			}
+		}
+	}
+	
+	void OnRemoveClientFromSession(FLiveLinkHubClientId Client)
+	{
+		if (TSharedPtr<ILiveLinkHubSessionManager> SessionManager = FModuleManager::Get().GetModuleChecked<FLiveLinkHubModule>("LiveLinkHub").GetLiveLinkHub()->GetSessionManager())
+		{
+			if (TSharedPtr<ILiveLinkHubSession> CurrentSession = SessionManager->GetCurrentSession())
+			{
+				if (TOptional<FLiveLinkHubClientId> SelectedClient = ClientsView->GetSelectedClient())
+				{
+					// Clear client details if it was selected.
+					if (Client == *SelectedClient)
+					{
+						UEClientInfo->Reset();
+						StructDetailsView->SetStructureData(nullptr);
+					}
+				}
+
+				CurrentSession->RemoveClient(Client);
+			}
+		}
+	}
+
 private:
 	/** List view of UE clients the hub can connect to. */
 	TSharedPtr<SLiveLinkHubClientsView> ClientsView;
diff --git a/Engine/Plugins/Animation/LiveLink/Source/LiveLinkHub/Private/Clients/LiveLinkHubClientsModel.h b/Engine/Plugins/Animation/LiveLink/Source/LiveLinkHub/Private/Clients/LiveLinkHubClientsModel.h
--- a/Engine/Plugins/Animation/LiveLink/Source/LiveLinkHub/Private/Clients/LiveLinkHubClientsModel.h
+++ b/Engine/Plugins/Animation/LiveLink/Source/LiveLinkHub/Private/Clients/LiveLinkHubClientsModel.h
@@ -11,7 +11,6 @@ struct FLiveLinkClientInfoMessage;
 struct FLiveLinkHubUEClientInfo;
 struct FLiveLinkSubjectKey;
 
-
 /** Decouples the UI from the livelink hub functions. */
 class ILiveLinkHubClientsModel
 {
@@ -19,12 +18,13 @@ public:
 	/** Types of client updates. */
 	enum class EClientEventType
 	{
-		Connected,
-		Removed,
-		Modified
+		Discovered, // UE client was discovered by the hub.
+		Disconnected, // Connection was lost with a client.
+		Reestablished, // Connection was restored with a previously disconnected client.
+		Modified // Information about a given client has changed (ie. level has changed).
 	};
 
-	DECLARE_MULTICAST_DELEGATE_TwoParams(FOnClientEvent, FLiveLinkHubClientId, EClientEventType);
+	DECLARE_TS_MULTICAST_DELEGATE_TwoParams(FOnClientEvent, FLiveLinkHubClientId, EClientEventType);
 
 	virtual ~ILiveLinkHubClientsModel() = default;
 
@@ -32,7 +32,10 @@ public:
 	virtual FOnClientEvent& OnClientEvent() = 0;
 
 	/** Get the status text of a client. */
-	virtual FText GetClientStatus(FLiveLinkHubClientId InClient) const = 0;
+	virtual FText GetClientStatus(FLiveLinkHubClientId Client) const = 0;
+
+	/** Get the list of clients that should be displayed in the clients panel according to the current session config. */
+	virtual TArray<FLiveLinkHubClientId> GetSessionClients() const = 0;
 
 	/** Get the list of clients discovered by the hub. */
 	virtual TArray<FLiveLinkHubClientId> GetDiscoveredClients() const = 0;
@@ -43,9 +46,6 @@ public:
 	/** Get the name of a client. */
 	virtual FText GetClientDisplayName(FLiveLinkHubClientId InClient) const = 0;
 
-	/** Remove a client from the client list. */
-	virtual void RemoveClient(FLiveLinkHubClientId InClient) = 0;
-
 	/** Get whether a client should receive livelink data. */
 	virtual bool IsClientEnabled(FLiveLinkHubClientId InClient) const = 0;
 
diff --git a/Engine/Plugins/Animation/LiveLink/Source/LiveLinkHub/Private/Clients/LiveLinkHubProvider.cpp b/Engine/Plugins/Animation/LiveLink/Source/LiveLinkHub/Private/Clients/LiveLinkHubProvider.cpp
new file mode 100644
--- /dev/null
+++ b/Engine/Plugins/Animation/LiveLink/Source/LiveLinkHub/Private/Clients/LiveLinkHubProvider.cpp
@@ -0,0 +1,396 @@
+// Copyright Epic Games, Inc. All Rights Reserved.
+
+#include "LiveLinkHubProvider.h"
+
+#include "Algo/Transform.h"
+#include "Async/Async.h"
+#include "Clients/LiveLinkHubClientsModel.h"
+#include "Clients/LiveLinkHubUEClientInfo.h"
+#include "Containers/ObservableArray.h"
+#include "CoreMinimal.h"
+#include "Delegates/DelegateCombinations.h"
+#include "Editor.h"
+#include "HAL/CriticalSection.h"
+#include "LiveLinkHubLog.h"
+#include "LiveLinkHubMessages.h"
+#include "LiveLinkHubSettings.h"
+#include "LiveLinkSettings.h"
+#include "MessageEndpointBuilder.h"
+#include "Misc/ScopeLock.h"
+#include "Session/LiveLinkHubSession.h"
+#include "Session/LiveLinkHubSessionManager.h"
+#include "Subjects/LiveLinkHubSubjectSessionConfig.h"
+#include "TimerManager.h"
+
+
+
+#define LOCTEXT_NAMESPACE "LiveLinkHub.LiveLinkHubProvider"
+
+
+FLiveLinkHubProvider::FLiveLinkHubProvider(const TSharedRef<ILiveLinkHubSessionManager>& InSessionManager)
+	: FLiveLinkProvider(TEXT("LiveLink Hub"), false)
+	, SessionManager(InSessionManager)
+{
+	Annotations.Add(FLiveLinkHubMessageAnnotation::ProviderTypeAnnotation, UE::LiveLinkHub::Private::LiveLinkHubProviderType.ToString());
+
+	FMessageEndpointBuilder EndpointBuilder = FMessageEndpoint::Builder(*GetProviderName());
+	EndpointBuilder.WithHandler(MakeHandler<FLiveLinkClientInfoMessage>(&FLiveLinkHubProvider::HandleClientInfoMessage));
+	EndpointBuilder.WithHandler(MakeHandler<FLiveLinkHubConnectMessage>(&FLiveLinkHubProvider::HandleHubConnectMessage));
+
+	CreateMessageEndpoint(EndpointBuilder);
+
+	const double ValidateConnectionsRate = GetDefault<ULiveLinkSettings>()->MessageBusPingRequestFrequency;
+	GEditor->GetTimerManager()->SetTimer(ValidateConnectionsTimer, FTimerDelegate::CreateRaw(this, &FLiveLinkHubProvider::ValidateConnections), ValidateConnectionsRate, true);
+}
+
+FLiveLinkHubProvider::~FLiveLinkHubProvider()
+{
+	if (GEditor)
+	{
+		GEditor->GetTimerManager()->ClearTimer(ValidateConnectionsTimer);
+	}
+}
+
+bool FLiveLinkHubProvider::ShouldTransmitToSubject_AnyThread(FName SubjectName, FMessageAddress Address) const
+{
+	FReadScopeLock Locker(ClientsMapLock);
+
+	const FLiveLinkHubClientId ClientId = AddressToIdCache.FindRef(Address);
+	if (const FLiveLinkHubUEClientInfo* ClientInfoPtr = ClientsMap.Find(ClientId))
+	{
+		if (const TSharedPtr<ILiveLinkHubSessionManager> Manager = SessionManager.Pin())
+		{
+			if (const TSharedPtr<ILiveLinkHubSession> CurrentSession = Manager->GetCurrentSession())
+			{
+				if (!CurrentSession->IsClientInSession(ClientInfoPtr->Id))
+				{
+					return false;
+				}
+			}
+		}
+
+		if (!ClientInfoPtr->bEnabled)
+		{
+			return false;
+		}
+
+		return !ClientInfoPtr->DisabledSubjects.Contains(SubjectName);
+	}
+	else
+	{
+		UE_LOG(LogLiveLinkHub, Warning, TEXT("Attempted to transmit data to an invalid client."));
+	}
+
+	return true;
+}
+
+void FLiveLinkHubProvider::AddRestoredClient(FLiveLinkHubUEClientInfo& RestoredClientInfo)
+{
+	// If a client was already discovered with the same hostname, update it to match the restored client.
+	bool bMatchedExistingConnection = false;
+	if (const TSharedPtr<ILiveLinkHubSessionManager> Manager = SessionManager.Pin())
+	{
+		if (const TSharedPtr<ILiveLinkHubSession> ActiveSession = Manager->GetCurrentSession())
+		{
+			FReadScopeLock Locker(ClientsMapLock);
+
+			for (auto It = ClientsMap.CreateIterator(); It; ++It)
+			{
+				FLiveLinkHubUEClientInfo& IteratedClient = It->Value;
+				if (IteratedClient.Hostname == RestoredClientInfo.Hostname && !ActiveSession->IsClientInSession(It->Key))
+				{
+					bMatchedExistingConnection = true;
+
+					// Update Client info from the new connection.
+					RestoredClientInfo = It->Value;
+					break;
+				}
+			}
+		}
+	}
+
+	if (!bMatchedExistingConnection)
+	{
+		ClientsMap.Add(RestoredClientInfo.Id, RestoredClientInfo);
+	}
+
+	OnClientEventDelegate.Broadcast(RestoredClientInfo.Id, EClientEventType::Discovered);
+}
+
+TOptional<FLiveLinkHubUEClientInfo> FLiveLinkHubProvider::GetClientInfo(FLiveLinkHubClientId InClient) const
+{
+	FReadScopeLock Locker(ClientsMapLock);
+	TOptional<FLiveLinkHubUEClientInfo> ClientInfo;
+	if (const FLiveLinkHubUEClientInfo* ClientInfoPtr = ClientsMap.Find(InClient))
+	{
+		ClientInfo = *ClientInfoPtr;
+	}
+
+	return ClientInfo;
+}
+
+void FLiveLinkHubProvider::HandleHubConnectMessage(const FLiveLinkHubConnectMessage& Message, const TSharedRef<IMessageContext, ESPMode::ThreadSafe>& Context)
+{
+	FLiveLinkConnectMessage ConnectMessage;
+	ConnectMessage.LiveLinkVersion = Message.ClientInfo.LiveLinkVersion;
+	FLiveLinkProvider::HandleConnectMessage(ConnectMessage, Context);
+
+	const FMessageAddress ConnectionAddress = Context->GetSender();
+
+	TOptional<FLiveLinkHubClientId> UpdatedClient;
+	{
+		FWriteScopeLock Locker(ClientsMapLock);
+
+		// First check if there are multiple disconnected entries with the same host as the incoming client.
+		uint8 NumClientsForHost = 0;
+		for (auto It = ClientsMap.CreateIterator(); It; ++It)
+		{
+			FLiveLinkHubUEClientInfo& IteratedClient = It->Value;
+			if (IteratedClient.Hostname == Message.ClientInfo.Hostname && IteratedClient.Status == ELiveLinkClientStatus::Disconnected)
+			{
+				NumClientsForHost++;
+				if (NumClientsForHost > 1)
+				{
+					break;
+				}
+			}
+		}
+
+		// Remove old entries if one is found
+		for (auto It = ClientsMap.CreateIterator(); It; ++It)
+		{
+			FLiveLinkHubUEClientInfo& IteratedClient = It->Value;
+			// If there are multiple disconnected clients with the same hostname, try finding a client with the same project.
+			bool bFindWithMatchingProject = NumClientsForHost > 1;
+
+			// Only replace disconnected clients to support multiple UE instances on the same host.
+			if (IteratedClient.Status == ELiveLinkClientStatus::Disconnected && IteratedClient.Hostname == Message.ClientInfo.Hostname)
+			{
+				if (!bFindWithMatchingProject || IteratedClient.ProjectName == Message.ClientInfo.ProjectName)
+				{
+					IteratedClient.UpdateFromInfoMessage(Message.ClientInfo);
+					IteratedClient.Id = It->Key;
+					IteratedClient.Status = ELiveLinkClientStatus::Connected;
+
+					AddressToIdCache.FindOrAdd(ConnectionAddress) = IteratedClient.Id;
+
+					UpdatedClient = IteratedClient.Id;
+					break;
+				}
+            }
+        }
+	}
+
+	if (UpdatedClient)
+	{
+		// Just updated the map
+		OnClientEventDelegate.Broadcast(*UpdatedClient, EClientEventType::Reestablished);
+	}
+	else
+	{
+		// Actually added a new entry in the map.
+		FLiveLinkHubUEClientInfo NewClient{Message.ClientInfo};
+		const FLiveLinkHubClientId NewClientId = NewClient.Id;
+		{
+			FWriteScopeLock Locker(ClientsMapLock);
+			AddressToIdCache.FindOrAdd(ConnectionAddress) = NewClientId;
+			ClientsMap.Add(NewClient.Id, MoveTemp(NewClient));
+		}
+
+		if (GetDefault<ULiveLinkHubSettings>()->bAutoAddDiscoveredClients)
+		{
+			AsyncTask(ENamedThreads::GameThread, [WeakSessionManager = SessionManager, NewClientId]()
+			{
+				if (const TSharedPtr<ILiveLinkHubSessionManager> Manager = WeakSessionManager.Pin())
+				{
+					if (const TSharedPtr<ILiveLinkHubSession> CurrentSession = Manager->GetCurrentSession())
+					{
+						CurrentSession->AddClient(NewClientId);
+					}
+				}
+			});
+		}
+		else
+		{
+			OnClientEventDelegate.Broadcast(NewClientId, EClientEventType::Discovered);
+		}
+	}
+}
+	
+void FLiveLinkHubProvider::HandleClientInfoMessage(const FLiveLinkClientInfoMessage& Message, const TSharedRef<IMessageContext, ESPMode::ThreadSafe>& Context)
+{
+	FMessageAddress Address = Context->GetSender();
+
+	FLiveLinkHubClientId ClientId = AddressToIdCache.FindRef(Address);
+	{
+		FWriteScopeLock Locker(ClientsMapLock);
+		if (FLiveLinkHubUEClientInfo* ClientInfo = ClientsMap.Find(ClientId))
+		{
+			ClientInfo->UpdateFromInfoMessage(Message);
+		}
+	}
+
+	if (ClientId.IsValid())
+	{
+		OnClientEventDelegate.Broadcast(ClientId, EClientEventType::Modified);
+	}
+}
+
+void FLiveLinkHubProvider::OnConnectionsClosed(const TArray<FMessageAddress>& ClosedAddresses)
+{
+	// List of OldId -> NewId
+	TArray<FLiveLinkHubClientId> Notifications;
+	TArray<FMessageAddress> AddressesToRemove;
+	{
+		FWriteScopeLock Locker(ClientsMapLock);
+
+		for (FMessageAddress TrackedAddress : ClosedAddresses)
+		{
+			FLiveLinkHubClientId ClientId = AddressToIdCache.FindRef(TrackedAddress);
+			if (FLiveLinkHubUEClientInfo* FoundInfo = ClientsMap.Find(ClientId))
+			{
+				FoundInfo->Status = ELiveLinkClientStatus::Disconnected;
+				Notifications.Add(ClientId);
+			}
+		}
+	}
+
+	for (const FLiveLinkHubClientId& Client : Notifications)
+	{
+		OnClientEventDelegate.Broadcast(Client, EClientEventType::Disconnected);
+	}
+
+	for (FMessageAddress TrackedAddress : ClosedAddresses)
+	{
+		FWriteScopeLock Locker(ClientsMapLock);
+		AddressToIdCache.Remove(TrackedAddress);
+	}
+}
+
+TArray<FLiveLinkHubClientId> FLiveLinkHubProvider::GetSessionClients() const
+{
+	TArray<FLiveLinkHubClientId> SessionClients;
+
+	if (const TSharedPtr<ILiveLinkHubSessionManager> Manager = SessionManager.Pin())
+	{
+		if (const TSharedPtr<ILiveLinkHubSession> CurrentSession = Manager->GetCurrentSession())
+		{
+			SessionClients = CurrentSession->GetSessionClients();
+		}
+	}
+
+	return SessionClients;
+}
+
+TMap<FName, FString> FLiveLinkHubProvider::GetAnnotations() const
+{
+	return Annotations;
+}
+
+TArray<FLiveLinkHubClientId> FLiveLinkHubProvider::GetDiscoveredClients() const
+{
+	TArray<FLiveLinkHubClientId> DiscoveredClients;
+
+	if (const TSharedPtr<ILiveLinkHubSessionManager> Manager = SessionManager.Pin())
+	{
+		if (const TSharedPtr<ILiveLinkHubSession> CurrentSession = Manager->GetCurrentSession())
+		{
+			FReadScopeLock Locker(ClientsMapLock);
+			TArray<FLiveLinkHubClientId> SessionClients = CurrentSession->GetSessionClients();
+			for (const TPair<FLiveLinkHubClientId, FLiveLinkHubUEClientInfo>& ClientPair : ClientsMap)
+			{
+				if (ClientPair.Value.Status != ELiveLinkClientStatus::Disconnected && !SessionClients.Contains(ClientPair.Key))
+				{
+					DiscoveredClients.Add(ClientPair.Key);
+				}
+			}
+		}
+	}
+
+	return DiscoveredClients;
+}
+
+FText FLiveLinkHubProvider::GetClientDisplayName(FLiveLinkHubClientId InAddress) const
+{
+	FReadScopeLock Locker(ClientsMapLock);
+	FText DisplayName;
+
+	if (const FLiveLinkHubUEClientInfo* ClientInfoPtr = ClientsMap.Find(InAddress))
+	{
+		DisplayName = FText::FromString(FString::Format(TEXT("{0} ({1})"), {*ClientInfoPtr->Hostname, *ClientInfoPtr->CurrentLevel}));
+	}
+	else
+	{
+		DisplayName = LOCTEXT("InvalidClientLabel", "Invalid Client");
+	}
+
+	return DisplayName;
+}
+
+FText FLiveLinkHubProvider::GetClientStatus(FLiveLinkHubClientId Client) const
+{
+	FReadScopeLock Locker(ClientsMapLock);
+	if (const FLiveLinkHubUEClientInfo* ClientInfoPtr = ClientsMap.Find(Client))
+	{
+		return StaticEnum<ELiveLinkClientStatus>()->GetDisplayNameTextByValue(static_cast<int64>(ClientInfoPtr->Status));
+	}
+		
+	return LOCTEXT("InvalidStatus", "Disconnected");
+}
+
+bool FLiveLinkHubProvider::IsClientEnabled(FLiveLinkHubClientId Client) const
+{
+	FReadScopeLock Locker(ClientsMapLock);
+	if (const FLiveLinkHubUEClientInfo* ClientInfoPtr = ClientsMap.Find(Client))
+	{
+		return ClientInfoPtr->bEnabled;
+	}
+	return false;
+}
+
+bool FLiveLinkHubProvider::IsClientConnected(FLiveLinkHubClientId Client) const
+{
+	FReadScopeLock Locker(ClientsMapLock);
+	if (const FLiveLinkHubUEClientInfo* ClientInfoPtr = ClientsMap.Find(Client))
+	{
+		return ClientInfoPtr->Status == ELiveLinkClientStatus::Connected;
+	}
+	return false;
+}
+
+void FLiveLinkHubProvider::SetClientEnabled(FLiveLinkHubClientId Client, bool bInEnable)
+{
+	FWriteScopeLock Locker(ClientsMapLock);
+	if (FLiveLinkHubUEClientInfo* ClientInfoPtr = ClientsMap.Find(Client))
+	{
+		ClientInfoPtr->bEnabled = bInEnable;
+	}
+}
+
+bool FLiveLinkHubProvider::IsSubjectEnabled(FLiveLinkHubClientId Client, const FLiveLinkSubjectKey& Subject) const
+{
+	FReadScopeLock Locker(ClientsMapLock);
+	if (const FLiveLinkHubUEClientInfo* ClientInfoPtr = ClientsMap.Find(Client))
+	{
+		return !ClientInfoPtr->DisabledSubjects.Contains(Subject.SubjectName);
+	}
+	return false;
+}
+
+void FLiveLinkHubProvider::SetSubjectEnabled(FLiveLinkHubClientId Client, const FLiveLinkSubjectKey& Subject, bool bInEnable)
+{
+	FWriteScopeLock Locker(ClientsMapLock);
+	if (FLiveLinkHubUEClientInfo* ClientInfoPtr = ClientsMap.Find(Client))
+	{
+		if (bInEnable)
+		{
+			ClientInfoPtr->DisabledSubjects.Remove(Subject.SubjectName);
+		}
+		else
+		{
+			ClientInfoPtr->DisabledSubjects.Add(Subject.SubjectName);
+		}
+	}
+}
+
+#undef LOCTEXT_NAMESPACE /*LiveLinkHub.LiveLinkHubProvider*/
diff --git a/Engine/Plugins/Animation/LiveLink/Source/LiveLinkHub/Private/Clients/LiveLinkHubProvider.h b/Engine/Plugins/Animation/LiveLink/Source/LiveLinkHub/Private/Clients/LiveLinkHubProvider.h
--- a/Engine/Plugins/Animation/LiveLink/Source/LiveLinkHub/Private/Clients/LiveLinkHubProvider.h
+++ b/Engine/Plugins/Animation/LiveLink/Source/LiveLinkHub/Private/Clients/LiveLinkHubProvider.h
@@ -2,29 +2,17 @@
 
 #pragma once
 
-#include "Algo/Transform.h"
-#include "Async/Async.h"
+#include "Clients/GameThreadMessageHandler.h"
 #include "Clients/LiveLinkHubClientsModel.h"
 #include "Clients/LiveLinkHubUEClientInfo.h"
-#include "Containers/ObservableArray.h"
-#include "CoreMinimal.h"
-#include "Delegates/DelegateCombinations.h"
-#include "Editor.h"
 #include "Engine/TimerHandle.h"
-#include "GameThreadMessageHandler.h"
-#include "HAL/CriticalSection.h"
 #include "IMessageContext.h"
-#include "LiveLinkHubLog.h"
-#include "LiveLinkHubMessages.h"
 #include "LiveLinkProviderImpl.h"
-#include "LiveLinkSettings.h"
-#include "MessageEndpointBuilder.h"
-#include "Misc/ScopeLock.h"
-#include "Subjects/LiveLinkHubSubjectSessionConfig.h"
-#include "Session/LiveLinkHubSession.h"
-#include "TimerManager.h"
 
-#define LOCTEXT_NAMESPACE "LiveLinkHub.LiveLinkHubProvider"
+class ILiveLinkHubSessionManager;
+struct FLiveLinkHubClientId;
+struct FLiveLinkHubConnectMessage;
+struct FLiveLinkHubUEClientInfo;
 
 /** 
  * LiveLink Provider that allows getting more information about a UE client by communicating with a LiveLinkHub MessageBus Source.
@@ -45,297 +33,48 @@ public:
 		return MakeShared<TGameThreadMessageHandler<MessageType, FLiveLinkHubProvider>>(this, Func);
 	}
 
-	FLiveLinkHubProvider(const TSharedRef<ILiveLinkHubSessionManager>& InSessionManager)
-		: FLiveLinkProvider(TEXT("LiveLink Hub"), false)
-		, SessionManager(InSessionManager)
-	{
-		Annotations.Add(FLiveLinkHubMessageAnnotation::ProviderTypeAnnotation, UE::LiveLinkHub::Private::LiveLinkHubProviderType.ToString());
-
-		FMessageEndpointBuilder EndpointBuilder = FMessageEndpoint::Builder(*GetProviderName());
-		EndpointBuilder.WithHandler(MakeHandler<FLiveLinkClientInfoMessage>(&FLiveLinkHubProvider::HandleClientInfoMessage));
-		EndpointBuilder.WithHandler(MakeHandler<FLiveLinkHubConnectMessage>(&FLiveLinkHubProvider::HandleHubConnectMessage));
-
-		CreateMessageEndpoint(EndpointBuilder);
-
-		const double ValidateConnectionsRate = GetDefault<ULiveLinkSettings>()->MessageBusPingRequestFrequency;
-		GEditor->GetTimerManager()->SetTimer(ValidateConnectionsTimer, FTimerDelegate::CreateRaw(this, &FLiveLinkHubProvider::ValidateConnections), ValidateConnectionsRate, true);
-	}
-
-	virtual ~FLiveLinkHubProvider() override
-	{
-		if (GEditor)
-		{
-			GEditor->GetTimerManager()->ClearTimer(ValidateConnectionsTimer);
-		}
-	}
-
-	virtual bool ShouldTransmitToSubject_AnyThread(FName SubjectName, FMessageAddress Address) const override
-	{
-		FReadScopeLock Locker(ClientsMapLock);
-		if (const FLiveLinkHubUEClientInfo* ClientInfoPtr = ClientsMap.FindByHash(GetTypeHash(Address), Address))
-		{
-			if (!ClientInfoPtr->bEnabled)
-			{
-				return false;
-			}
-
-			return !ClientInfoPtr->DisabledSubjects.Contains(SubjectName);
-		}
-		else
-		{
-			UE_LOG(LogLiveLinkHub, Warning, TEXT("Attempted to transmit data to an invalid client."));
-		}
+	FLiveLinkHubProvider(const TSharedRef<ILiveLinkHubSessionManager>& InSessionManager);
 
-		return true;
-	}
-
-	/** Manually add a client to the client map. */
-	void AddClient(const FLiveLinkHubUEClientInfo& InClientInfo)
-	{
-		{
-			FWriteScopeLock Locker(ClientsMapLock);
-			ClientsMap.Add(InClientInfo.Id, InClientInfo);
-		}
-
-		OnClientEventDelegate.Broadcast(InClientInfo.Id, EClientEventType::Connected);
-	}
-
-	/** Manually remove a client from the client map. */
-	virtual void RemoveClient(FLiveLinkHubClientId InClientId) override
-	{
-		{
-			// Remove the client from the map so that the connection closed callback doesn't set it to disconnected status.
-			FWriteScopeLock Locker(ClientsMapLock);
-			ClientsMap.Remove(InClientId);
-		}
+	virtual ~FLiveLinkHubProvider() override;
 
-		CloseConnection(InClientId.GetAddress());
+	//~ Begin LiveLinkProvider interface
+	virtual bool ShouldTransmitToSubject_AnyThread(FName SubjectName, FMessageAddress Address) const override;
+	virtual TOptional<FLiveLinkHubUEClientInfo> GetClientInfo(FLiveLinkHubClientId InClient) const override;
+	//~ End LiveLinkProvider interface
 
-		OnClientEventDelegate.Broadcast(InClientId, EClientEventType::Removed);
-	}
+	/**
+	 * Restore a client, calling this will modify the client ID if it matches an existing connection.
+	 */
+	void AddRestoredClient(FLiveLinkHubUEClientInfo& InOutRestoredClientInfo);
 
 	/** Retrieve the existing client map. */
 	const TMap<FLiveLinkHubClientId, FLiveLinkHubUEClientInfo>& GetClientsMap() const { return ClientsMap; }
 
 private:
 	/** Handle a connection message resulting from a livelink hub message bus source connecting to this provider. */
-	void HandleHubConnectMessage(const FLiveLinkHubConnectMessage& Message, const TSharedRef<IMessageContext, ESPMode::ThreadSafe>& Context)
-	{
-		FLiveLinkConnectMessage ConnectMessage;
-		ConnectMessage.LiveLinkVersion = Message.ClientInfo.LiveLinkVersion;
-		FLiveLinkProvider::HandleConnectMessage(ConnectMessage, Context);
-
-		const FMessageAddress ConnectionAddress = Context->GetSender();
-
-		TOptional<FLiveLinkHubClientId> UpdatedClient;
-		{
-			FWriteScopeLock Locker(ClientsMapLock);
-			// Remove old entries if one is found
-            for (auto It = ClientsMap.CreateIterator(); It; ++It)
-            {
-            	FLiveLinkHubUEClientInfo& IteratedClient = It->Value;
-            	if (IteratedClient.Hostname == Message.ClientInfo.Hostname && IteratedClient.LongName == Message.ClientInfo.LongName
-            		&& IteratedClient.ProjectName == Message.ClientInfo.ProjectName && IteratedClient.CurrentLevel == Message.ClientInfo.CurrentLevel)
-            	{
-            		// Only replace disconnected clients to support multiple UE instances on the same host.
-					if (IteratedClient.Status == ELiveLinkClientStatus::Disconnected)
-					{
-						FLiveLinkHubUEClientInfo RemovedInfo = MoveTemp(It->Value);
-						RemovedInfo.Id.InvalidateAddress();
-						RemovedInfo.Status = ELiveLinkClientStatus::Connected;
-						ClientsMap.Remove(It->Key);
-
-						FLiveLinkHubClientId NewId = RemovedInfo.Id;
-						ClientsMap.Add(NewId, MoveTemp(RemovedInfo));
-						UpdatedClient = NewId;
-            			break;
-					}
-            	}
-            }
-		}
-
-		if (UpdatedClient)
-		{
-			OnClientEventDelegate.Broadcast(*UpdatedClient, EClientEventType::Modified);
-		}
-		else
-		{
-			FLiveLinkHubUEClientInfo NewClient{Message.ClientInfo, ConnectionAddress};
-			const FLiveLinkHubClientId NewClientId = NewClient.Id;
-			{
-				FWriteScopeLock Locker(ClientsMapLock);
-				ClientsMap.Add(NewClient.Id, MoveTemp(NewClient));
-			}
-
-			OnClientEventDelegate.Broadcast(NewClientId, EClientEventType::Connected);
-		}
-	}
+	void HandleHubConnectMessage(const FLiveLinkHubConnectMessage& Message, const TSharedRef<IMessageContext, ESPMode::ThreadSafe>& Context);
 	
 	/** Handle a client info message being received. Happens when new information about a client is received (ie. Client has changed map) */
-	void HandleClientInfoMessage(const FLiveLinkClientInfoMessage& Message, const TSharedRef<IMessageContext, ESPMode::ThreadSafe>& Context)
-	{
-		FMessageAddress Address = Context->GetSender();
-		FLiveLinkHubClientId Id(Address);
-
-		{
-			FWriteScopeLock Locker(ClientsMapLock);
-			if (FLiveLinkHubUEClientInfo* ClientInfo = ClientsMap.Find(Id))
-			{
-				*ClientInfo = FLiveLinkHubUEClientInfo(Message, Address);
-			}
-		}
-		OnClientEventDelegate.Broadcast(Id, EClientEventType::Modified);
-	}
+	void HandleClientInfoMessage(const FLiveLinkClientInfoMessage& Message, const TSharedRef<IMessageContext, ESPMode::ThreadSafe>& Context);
 
 protected:
-	//~ Begin FLiveLinkProvider interface
-	virtual void OnConnectionsClosed(const TArray<FMessageAddress>& ClosedAddresses) override
-	{
-		TArray<FLiveLinkHubClientId> Notifications;
-		{
-			FWriteScopeLock Locker(ClientsMapLock);
-
-			for (FMessageAddress TrackedAddress : ClosedAddresses)
-			{
-				const FLiveLinkHubClientId Id{ TrackedAddress };
-
-				if (FLiveLinkHubUEClientInfo* FoundInfo = ClientsMap.Find(Id))
-				{
-					FLiveLinkHubUEClientInfo RemovedInfo = MoveTemp(*FoundInfo);
-					RemovedInfo.Id.InvalidateAddress();
-					RemovedInfo.Status = ELiveLinkClientStatus::Disconnected;
-					ClientsMap.Remove(Id);
-
-					FLiveLinkHubClientId NewId = RemovedInfo.Id;
-					ClientsMap.Add(NewId, MoveTemp(RemovedInfo));
-					Notifications.Add(Id);
-				}
-			}
-		}
-
-		for (FLiveLinkHubClientId Id : Notifications)
-		{
-			OnClientEventDelegate.Broadcast(Id, EClientEventType::Modified);
-		}
-	}
-
-	virtual TMap<FName, FString> GetAnnotations() const override
-	{
-		return Annotations;
-	}
-
-	virtual TArray<FLiveLinkHubClientId> GetDiscoveredClients() const override
-	{
-		FReadScopeLock Locker(ClientsMapLock);
-		TArray<FLiveLinkHubClientId> ClientAddresses;
-		ClientsMap.GenerateKeyArray(ClientAddresses);
-		return ClientAddresses;
-	}
-
-	virtual TOptional<FLiveLinkHubUEClientInfo> GetClientInfo(FLiveLinkHubClientId InAddress) const override
-	{
-		FReadScopeLock Locker(ClientsMapLock);
-		TOptional<FLiveLinkHubUEClientInfo> ClientInfo;
-		if (const FLiveLinkHubUEClientInfo* ClientInfoPtr = ClientsMap.Find(InAddress))
-		{
-			ClientInfo = *ClientInfoPtr;
-		}
-
-		return ClientInfo;
-	}
-
-	virtual FText GetClientDisplayName(FLiveLinkHubClientId InAddress) const override
-	{
-		FReadScopeLock Locker(ClientsMapLock);
-		FText DisplayName;
-
-		if (const FLiveLinkHubUEClientInfo* ClientInfoPtr = ClientsMap.Find(InAddress))
-       	{
-			DisplayName = FText::FromString(ClientInfoPtr->LongName);
-       	}
-       	else
-       	{
-       		DisplayName = LOCTEXT("InvalidSourceLabel", "Invalid Source");
-       	}
-
-       	return DisplayName;
-	}
-
+	//~ Begin ILiveLinkHubClientsModel interface
+	virtual void OnConnectionsClosed(const TArray<FMessageAddress>& ClosedAddresses) override;
+	virtual TArray<FLiveLinkHubClientId> GetSessionClients() const override;
+	virtual TMap<FName, FString> GetAnnotations() const override;
+	virtual TArray<FLiveLinkHubClientId> GetDiscoveredClients() const override;
+	virtual FText GetClientDisplayName(FLiveLinkHubClientId InAddress) const override;
 	virtual FOnClientEvent& OnClientEvent() override
 	{
 		return OnClientEventDelegate;
 	}
-
-	virtual FText GetClientStatus(FLiveLinkHubClientId Client) const override
-	{
-		FReadScopeLock Locker(ClientsMapLock);
-		if (const FLiveLinkHubUEClientInfo* ClientInfoPtr = ClientsMap.Find(Client))
-		{
-			return StaticEnum<ELiveLinkClientStatus>()->GetDisplayNameTextByValue(static_cast<int64>(ClientInfoPtr->Status));
-		}
-		
-		return LOCTEXT("InvalidStatus", "Disconnected");
-	}
-
-	/** Get whether a client should receive livelink data. */
-	virtual bool IsClientEnabled(FLiveLinkHubClientId Client) const override
-	{
-		FReadScopeLock Locker(ClientsMapLock);
-		if (const FLiveLinkHubUEClientInfo* ClientInfoPtr = ClientsMap.Find(Client))
-		{
-			return ClientInfoPtr->bEnabled;
-		}
-		return false;
-	}
-
-	virtual bool IsClientConnected(FLiveLinkHubClientId Client) const override
-	{
-		FReadScopeLock Locker(ClientsMapLock);
-		if (const FLiveLinkHubUEClientInfo* ClientInfoPtr = ClientsMap.Find(Client))
-		{
-			return ClientInfoPtr->Status == ELiveLinkClientStatus::Connected;
-		}
-		return false;
-	}
-
-	/** Set whether a client should receive livelink data. */
-	virtual void SetClientEnabled(FLiveLinkHubClientId Client, bool bInEnable) override
-	{
-		FWriteScopeLock Locker(ClientsMapLock);
-		if (FLiveLinkHubUEClientInfo* ClientInfoPtr = ClientsMap.Find(Client))
-		{
-			ClientInfoPtr->bEnabled = bInEnable;
-		}
-	}
-
-	/** Get whether a subject is enabled on a given client. */
-	virtual bool IsSubjectEnabled(FLiveLinkHubClientId Client, const FLiveLinkSubjectKey& Subject) const override
-	{
-		FReadScopeLock Locker(ClientsMapLock);
-		if (const FLiveLinkHubUEClientInfo* ClientInfoPtr = ClientsMap.Find(Client))
-		{
-			return !ClientInfoPtr->DisabledSubjects.Contains(Subject.SubjectName);
-		}
-		return false;
-	}
-
-	/** Set whether a subject should receive livelink data. */
-	virtual void SetSubjectEnabled(FLiveLinkHubClientId Client, const FLiveLinkSubjectKey& Subject, bool bInEnable) override
-	{
-		FWriteScopeLock Locker(ClientsMapLock);
-		if (FLiveLinkHubUEClientInfo* ClientInfoPtr = ClientsMap.Find(Client))
-		{
-			if (bInEnable)
-			{
-				ClientInfoPtr->DisabledSubjects.Remove(Subject.SubjectName);
-			}
-			else
-			{
-				ClientInfoPtr->DisabledSubjects.Add(Subject.SubjectName);
-			}
-		}
-	}
-	//~ End FLiveLinkProvider interface
+	virtual FText GetClientStatus(FLiveLinkHubClientId Client) const override;
+	virtual bool IsClientEnabled(FLiveLinkHubClientId Client) const override;
+	virtual bool IsClientConnected(FLiveLinkHubClientId Client) const override;
+	virtual void SetClientEnabled(FLiveLinkHubClientId Client, bool bInEnable) override;
+	virtual bool IsSubjectEnabled(FLiveLinkHubClientId Client, const FLiveLinkSubjectKey& Subject) const override;
+	virtual void SetSubjectEnabled(FLiveLinkHubClientId Client, const FLiveLinkSubjectKey& Subject, bool bInEnable) override;
+	//~ End ILiveLinkHubClientsModel interface
 
 private:
 	/** Handle to the timer responsible for validating the livelinkprovider's connections.*/
@@ -348,8 +87,8 @@ private:
 	TMap<FName, FString> Annotations;
 	/** LiveLinkHub session manager. */
 	TWeakPtr<ILiveLinkHubSessionManager> SessionManager;
+	/** Cache used to retrieve the client id from a message bus address. */
+	TMap<FMessageAddress, FLiveLinkHubClientId> AddressToIdCache;
 	/** Lock used to access the clients map from different threads. */
 	mutable FRWLock ClientsMapLock;
 };
-
-#undef LOCTEXT_NAMESPACE /*LiveLinkHub.LiveLinkHubProvider*/
diff --git a/Engine/Plugins/Animation/LiveLink/Source/LiveLinkHub/Private/Clients/LiveLinkHubUEClientInfo.h b/Engine/Plugins/Animation/LiveLink/Source/LiveLinkHub/Private/Clients/LiveLinkHubUEClientInfo.h
--- a/Engine/Plugins/Animation/LiveLink/Source/LiveLinkHub/Private/Clients/LiveLinkHubUEClientInfo.h
+++ b/Engine/Plugins/Animation/LiveLink/Source/LiveLinkHub/Private/Clients/LiveLinkHubUEClientInfo.h
@@ -18,66 +18,30 @@ struct FLiveLinkHubClientId
 	/** Default constructor, should only be used by the reflection code. */
 	FLiveLinkHubClientId() = default;
 
-	explicit FLiveLinkHubClientId(FMessageAddress InAddress)
-		: Address(MoveTemp(InAddress))
-		, Guid(FGuid::NewGuid())
+	static FLiveLinkHubClientId NewId()
 	{
+		FLiveLinkHubClientId Id{};
+		Id.Guid = FGuid::NewGuid();
+		return Id;
 	}
 
-	/** Update the address contained in this ID. */
-	void UpdateAddress(FMessageAddress InAddress)
+	bool IsValid() const
 	{
-		Address = InAddress;
-	}
-
-	/** Invalidate the address contained in this ID. */
-	void InvalidateAddress()
-	{
-		Address.Invalidate();
-	}
-
-	/** Get the message bus address of this client. */
-	FMessageAddress GetAddress() const
-	{
-		return Address;
+		return Guid.IsValid();
 	}
 
 	/** Get the hash of this ID. */
 	friend uint32 GetTypeHash(FLiveLinkHubClientId Id)
 	{
-		if (Id.Address.IsValid())
-		{
-			return GetTypeHash(Id.Address);
-		}
-
 		return GetTypeHash(Id.Guid);
 	}
 
 	bool operator==(const FLiveLinkHubClientId& Other) const
 	{
-		if (Guid == Other.Guid)
-		{
-			return true;
-		}
-
-		if (Address.IsValid() && Other.Address.IsValid())
-		{
-			return Address == Other.Address;
-		}
-
-		// Make sure we return false so that 2 clients with invalid addresses are considered different.
-		return false;
-	}
-
-	bool operator==(const FMessageAddress& Other) const
-	{
-		return Address == Other;
+		return Guid == Other.Guid;
 	}
 
 private:
-	/** MessageBus address for this client. Invalid when the client is disconnected. */
-	FMessageAddress Address;
-
 	/** Unique identifier for this client. */
 	UPROPERTY()
 	FGuid Guid;
@@ -91,8 +55,8 @@ struct FLiveLinkHubUEClientInfo
 
 	FLiveLinkHubUEClientInfo() = default;
 
-	explicit FLiveLinkHubUEClientInfo(const FLiveLinkClientInfoMessage& InClientInfo, FMessageAddress Address)
-		: Id(FLiveLinkHubClientId(MoveTemp(Address)))
+	explicit FLiveLinkHubUEClientInfo(const FLiveLinkClientInfoMessage& InClientInfo)
+		: Id(FLiveLinkHubClientId::NewId())
 		, LongName(InClientInfo.LongName)
 		, Status(InClientInfo.Status)
 		, IPAddress(TEXT("192.168.0.1 (Placeholder)"))
@@ -102,17 +66,15 @@ struct FLiveLinkHubUEClientInfo
 	{
 	}
 
-	void UpdateFromClient(const FLiveLinkClientInfoMessage& InClientInfo, FMessageAddress Address)
-    {
-		// Preserve old id but update its connection address.
-		FLiveLinkHubClientId PreviousId = Id;
-		PreviousId.UpdateAddress(Address);
-
-		FLiveLinkHubUEClientInfo NewInfo{InClientInfo, Address};
-		NewInfo.Id = PreviousId;
-
-		*this = MoveTemp(NewInfo);
-    }
+	void UpdateFromInfoMessage(const FLiveLinkClientInfoMessage& InClientInfo)
+	{
+		LongName = InClientInfo.LongName;
+		Status = InClientInfo.Status;
+		IPAddress = TEXT("192.168.0.1 (Placeholder)");
+		Hostname = InClientInfo.Hostname;
+		ProjectName = InClientInfo.ProjectName;
+		CurrentLevel = InClientInfo.CurrentLevel;
+	}
 
 	/** Identifier for this client. */
 	UPROPERTY()
diff --git a/Engine/Plugins/Animation/LiveLink/Source/LiveLinkHub/Private/Clients/SLiveLinkHubClientsView.h b/Engine/Plugins/Animation/LiveLink/Source/LiveLinkHub/Private/Clients/SLiveLinkHubClientsView.h
--- a/Engine/Plugins/Animation/LiveLink/Source/LiveLinkHub/Private/Clients/SLiveLinkHubClientsView.h
+++ b/Engine/Plugins/Animation/LiveLink/Source/LiveLinkHub/Private/Clients/SLiveLinkHubClientsView.h
@@ -8,13 +8,16 @@
 #include "Delegates/Delegate.h"
 #include "Delegates/DelegateCombinations.h"
 #include "Features/IModularFeatures.h"
+#include "Framework/Application/SlateApplication.h"
 #include "Misc/Guid.h"
 #include "LiveLinkClient.h"
 #include "LiveLinkHubClientsModel.h"
 #include "LiveLinkHubUEClientInfo.h"
+#include "Session/LiveLinkHubSessionManager.h"
 #include "Styling/SlateTypes.h"
 #include "Widgets/Input/SButton.h"
 #include "Widgets/Input/SCheckBox.h"
+#include "Widgets/Input/SComboButton.h"
 #include "Widgets/Text/STextBlock.h"
 #include "Widgets/Views/STreeView.h"
 
@@ -24,6 +27,9 @@ class FLiveLinkHub;
 #define LOCTEXT_NAMESPACE "LiveLinkHub.ClientsView"
 
 DECLARE_DELEGATE_OneParam(FOnClientSelected, FLiveLinkHubClientId/*ClientIdentifier*/);
+DECLARE_DELEGATE_OneParam(FOnDiscoveredClientPicked, FLiveLinkHubClientId/*ClientIdentifier*/);
+DECLARE_DELEGATE_OneParam(FOnRemoveClientFromSession, FLiveLinkHubClientId/*ClientIdentifier*/);
+
 
 static const FName NameColumnId = "Name";
 static const FName StatusColumnId = "Status";
@@ -32,14 +38,14 @@ static const FName EnabledIconColumnId = "EnabledIcon";
 /** Tree view item that represents either a client or a livelink subject. */
 struct FClientTreeViewItem
 {
-	virtual ~FClientTreeViewItem() = default;
-
 	FClientTreeViewItem(FLiveLinkHubClientId InClientId, TSharedRef<ILiveLinkHubClientsModel> InClientsModel)
 		: ClientId(MoveTemp(InClientId))
 		, ClientsModel(MoveTemp(InClientsModel))
 	{
 	}
 
+	virtual ~FClientTreeViewItem() = default;
+
 	/** Get the subject key for this tree item (Invalid key for client rows). */
 	virtual const FLiveLinkSubjectKey& GetSubjectKey() const
 	{
@@ -64,7 +70,8 @@ struct FClientTreeViewItem
 
 	/** Get status text for the row. */
 	virtual FText GetStatusText() const = 0;
-
+	
+public:
 	/** This item's children, in the case of client rows, these represent the livelink subjects. */
 	TArray<TSharedPtr<FClientTreeViewItem>> Children;
 	/** Name of the tree item (client or subject name). */
@@ -156,7 +163,7 @@ struct FClientTreeViewSubjectItem : public FClientTreeViewItem
 	{
 		if (const TSharedPtr<ILiveLinkHubClientsModel> ClientsModelPtr = ClientsModel.Pin())
 		{
-			return !ClientsModelPtr->IsClientEnabled(ClientId);
+			return !ClientsModelPtr->IsClientEnabled(ClientId) || !ClientsModelPtr->IsClientConnected(ClientId);
 		}
 
 		return false;
@@ -180,7 +187,7 @@ struct FClientTreeViewSubjectItem : public FClientTreeViewItem
 	FLiveLinkSubjectKey LiveLinkSubjectKey;
 };
 
-
+/** Holds a client row's data. */
 class SLiveLinkHubClientsRow : public SMultiColumnTableRow<TSharedPtr<FClientTreeViewItem>>
 {
 public:
@@ -275,6 +282,8 @@ class SLiveLinkHubClientsView : public SCompoundWidget
 public:
 	SLATE_BEGIN_ARGS(SLiveLinkHubClientsView) {}
 	SLATE_EVENT(FOnClientSelected, OnClientSelected)
+	SLATE_EVENT(FOnDiscoveredClientPicked, OnDiscoveredClientPicked)
+	SLATE_EVENT(FOnRemoveClientFromSession, OnRemoveClientFromSession)
 	SLATE_END_ARGS()
 
 	using FClientTreeItemPtr = TSharedPtr<FClientTreeViewItem>;
@@ -283,17 +292,55 @@ public:
 	void Construct(const FArguments& InArgs, TSharedRef<ILiveLinkHubClientsModel> InClientsModel)
 	{
 		OnClientSelectedDelegate = InArgs._OnClientSelected;
+		OnDiscoveredClientPickedDelegate = InArgs._OnDiscoveredClientPicked;
+		OnRemoveClientFromSessionDelegate = InArgs._OnRemoveClientFromSession;
+
 		ClientsModel = MoveTemp(InClientsModel);
 		 
 		ClientsModel->OnClientEvent().AddSP(this, &SLiveLinkHubClientsView::OnClientEvent);
 
+		const FLiveLinkHubModule& LiveLinkHubModule = FModuleManager::Get().GetModuleChecked<FLiveLinkHubModule>("LiveLinkHub");
+		LiveLinkHubModule.GetSessionManager()->OnClientAddedToSession().AddSP(this, &SLiveLinkHubClientsView::OnClientAddedToSession);
+		LiveLinkHubModule.GetSessionManager()->OnClientRemovedFromSession().AddSP(this, &SLiveLinkHubClientsView::OnClientRemovedFromSession);
+		LiveLinkHubModule.GetSessionManager()->OnActiveSessionChanged().AddSP(this, &SLiveLinkHubClientsView::OnActiveSessionChanged);
+
 		FLiveLinkClient& LiveLinkClient = IModularFeatures::Get().GetModularFeature<FLiveLinkClient>(ILiveLinkClient::ModularFeatureName);
 		LiveLinkClient.OnLiveLinkSubjectAdded().AddSP(this, &SLiveLinkHubClientsView::OnSubjectAdded_AnyThread);
 		LiveLinkClient.OnLiveLinkSubjectRemoved().AddSP(this, &SLiveLinkHubClientsView::OnSubjectRemoved_AnyThread);
 
 		ChildSlot
 		[
-			SAssignNew(TreeView, STreeView<FClientTreeItemPtr>)
+			SNew(SVerticalBox)
+			+ SVerticalBox::Slot()
+			.HAlign(HAlign_Right)
+			.AutoHeight()
+			[
+				SNew(SComboButton)
+				.CollapseMenuOnParentFocus(true)
+				.ButtonContent()
+				[
+					SNew(STextBlock)
+					.Text(LOCTEXT("AddClientLabel", "Add Client"))
+				]
+				.MenuContent()
+				[
+					SNew(SBox)
+					.HeightOverride(200.0)
+					.WidthOverride(200.0)
+					[
+						SAssignNew(DiscoveredClientsListView, SListView<TSharedPtr<FLiveLinkHubClientId>>)
+						.ListItemsSource(&DiscoveredClients)
+						.ItemHeight(20.0f)
+						.OnSelectionChanged(this, &SLiveLinkHubClientsView::OnDiscoveredClientPicked)
+						.OnGenerateRow(this, &SLiveLinkHubClientsView::OnGenerateDiscoveredClientsRow)
+					]
+				]
+				
+			]
+			+ SVerticalBox::Slot()
+			.AutoHeight()
+			[
+				SAssignNew(TreeView, STreeView<FClientTreeItemPtr>)
 				.TreeItemsSource(&Clients)
 				.ItemHeight(20.0f)
 				.OnSelectionChanged(this, &SLiveLinkHubClientsView::OnSelectionChanged)
@@ -313,9 +360,10 @@ public:
 					.ManualWidth(20.f)
 					.DefaultLabel(LOCTEXT("EnabledIconEmpty", ""))
 				)
+			]
 		];
 
-		PopulateClients();
+		Reinitialize();
 	}
 	//~ End SWidget interface
 
@@ -328,12 +376,33 @@ public:
 			LiveLinkClient.OnLiveLinkSubjectAdded().RemoveAll(this);
 		}
 
+		if (FLiveLinkHubModule* LiveLinkHubModule = FModuleManager::Get().GetModulePtr<FLiveLinkHubModule>("LiveLinkHub"))
+		{
+			if (TSharedPtr<ILiveLinkHubSessionManager> SessionManager = LiveLinkHubModule->GetSessionManager())
+			{
+				SessionManager->OnActiveSessionChanged().RemoveAll(this);
+				SessionManager->OnClientRemovedFromSession().RemoveAll(this);
+				SessionManager->OnClientAddedToSession().RemoveAll(this);
+			}
+		}
+
 		if (ClientsModel)
 		{
 			ClientsModel->OnClientEvent().RemoveAll(this);
 		}
 	}
 
+	/** Refresh the data and widgets. */
+	void Reinitialize()
+	{
+		PopulateClients();
+		TreeView->RequestTreeRefresh();
+		if (DiscoveredClientsListView)
+		{
+			DiscoveredClientsListView->RequestListRefresh();
+		}
+	}
+
 	/** Get the currently selected client, or an invalid optional when nothing is selected. */
 	TOptional<FLiveLinkHubClientId> GetSelectedClient() const
 	{
@@ -349,6 +418,22 @@ public:
 	}
 
 private:
+	/** Handler used to generate a widget for a given client row. */
+	TSharedRef<ITableRow> OnGenerateDiscoveredClientsRow(TSharedPtr<FLiveLinkHubClientId> Item, const TSharedRef<STableViewBase>& OwnerTable)
+	{
+		FText ClientName = LOCTEXT("InvalidClient", "Invalid Client");
+		if (ClientsModel && Item)
+		{
+			ClientName = ClientsModel->GetClientDisplayName(*Item);
+		}
+
+		return SNew(STableRow<TSharedPtr<FLiveLinkHubClientId>>, OwnerTable)
+		[
+			SNew(STextBlock)
+			.Text(ClientName)
+		];
+	}
+
 	/** Handler used to generate a widget for a given client row. */
 	TSharedRef<ITableRow> OnGenerateClientRow(FClientTreeItemPtr Item, const TSharedRef<STableViewBase>& OwnerTable)
 	{
@@ -369,7 +454,7 @@ private:
     	{
 			if (TOptional<FLiveLinkHubClientId> Id = GetSelectedClient())
 			{
-				ClientsModel->RemoveClient(*Id);
+				OnRemoveClientFromSessionDelegate.ExecuteIfBound(*Id);
 			}
     		return FReply::Handled();
     	}
@@ -386,28 +471,108 @@ private:
 		}
 	}
 
+	/** Handler called when a client is picked in the Add Client menu. */
+	void OnDiscoveredClientPicked(TSharedPtr<FLiveLinkHubClientId> InItem, const ESelectInfo::Type InSelectInfoType)
+	{
+		if (InItem && TreeView)
+		{
+			OnDiscoveredClientPickedDelegate.ExecuteIfBound(*InItem);
+			FSlateApplication::Get().SetUserFocus(0, TreeView);
+		}
+	}
+
 	/** Handler called when the client list has changed. */
 	void OnClientEvent(FLiveLinkHubClientId ClientId, ILiveLinkHubClientsModel::EClientEventType EventType)
 	{
 		switch (EventType)
 		{
-		case ILiveLinkHubClientsModel::EClientEventType::Connected:
+		case ILiveLinkHubClientsModel::EClientEventType::Discovered:
 		{
-			if (!Clients.ContainsByPredicate([ClientId](const FClientTreeItemPtr& InClient) { return InClient->ClientId == ClientId; }))
+			if (!DiscoveredClients.ContainsByPredicate([ClientId](const TSharedPtr<FLiveLinkHubClientId>& InClient) { return *InClient == ClientId; }))
 			{
-				TSharedPtr<FClientTreeViewClientItem> ClientItem = MakeShared<FClientTreeViewClientItem>(ClientId, ClientsModel.ToSharedRef());
-				ClientItem->ClientId = ClientId;
-				InitializeClientItem(*ClientItem);
-
-				Clients.Add(ClientItem);
-				TreeView->RequestTreeRefresh();
+				DiscoveredClients.Add(MakeShared<FLiveLinkHubClientId>(MoveTemp(ClientId)));
+
+				TWeakPtr<SListView<TSharedPtr<FLiveLinkHubClientId>>> WeakList = DiscoveredClientsListView;
+				AsyncTask(ENamedThreads::GameThread, [WeakList]() 
+				{
+					if (TSharedPtr<SListView<TSharedPtr<FLiveLinkHubClientId>>> List = WeakList.Pin())
+					{
+						List->RequestListRefresh();
+					}
+				});
+			}
+			break;
+		}
+		case ILiveLinkHubClientsModel::EClientEventType::Disconnected:
+		{
+			int32 ClientIndex = DiscoveredClients.IndexOfByPredicate([ClientId](const TSharedPtr<FLiveLinkHubClientId>& InClient) { return *InClient == ClientId; });
+			if (ClientIndex != INDEX_NONE)
+			{
+				DiscoveredClients.RemoveAt(ClientIndex);
+
+				TWeakPtr<SListView<TSharedPtr<FLiveLinkHubClientId>>> WeakList = DiscoveredClientsListView;
+				AsyncTask(ENamedThreads::GameThread, [WeakList]()
+					{
+						if (TSharedPtr<SListView<TSharedPtr<FLiveLinkHubClientId>>> List = WeakList.Pin())
+						{
+							List->RequestListRefresh();
+						}
+					});
 			}
 			break;
 		}
-		case ILiveLinkHubClientsModel::EClientEventType::Removed:
+		case ILiveLinkHubClientsModel::EClientEventType::Reestablished:
 		{
-			Clients.RemoveAll([ClientId](const FClientTreeItemPtr& InClient) { return InClient->ClientId == ClientId; });
-			TreeView->RequestTreeRefresh();
+
+			if (FLiveLinkHubModule* LiveLinkHubModule = FModuleManager::Get().GetModulePtr<FLiveLinkHubModule>("LiveLinkHub"))
+			{
+				if (TSharedPtr<ILiveLinkHubSessionManager> SessionManager = LiveLinkHubModule->GetSessionManager())
+				{
+					if (SessionManager->GetCurrentSession()->IsClientInSession(ClientId))
+					{
+						// Client is in session, add it to the clients, remove from discovered list
+						int32 DiscoveredClientIndex = DiscoveredClients.IndexOfByPredicate([ClientId](const TSharedPtr<FLiveLinkHubClientId>& InClient) { return *InClient == ClientId; });
+						if (DiscoveredClientIndex != INDEX_NONE)
+						{
+							DiscoveredClients.RemoveAt(DiscoveredClientIndex);
+						}
+
+						int32 ClientIndex = Clients.IndexOfByPredicate([ClientId](const FClientTreeItemPtr& InClient) { return InClient->ClientId == ClientId; });
+						if (ClientIndex == INDEX_NONE)
+						{
+							TSharedPtr<FClientTreeViewClientItem> ClientItem = MakeShared<FClientTreeViewClientItem>(ClientId, ClientsModel.ToSharedRef());
+							ClientItem->ClientId = ClientId;
+							InitializeClientItem(*ClientItem);
+							Clients.Add(ClientItem);
+						}
+					}
+					else
+					{
+						// Client is not in session, add it to discovered clients if not already there
+						int32 ClientIndex = DiscoveredClients.IndexOfByPredicate([ClientId](const TSharedPtr<FLiveLinkHubClientId>& InClient) { return *InClient == ClientId; });
+						if (ClientIndex == INDEX_NONE)
+						{
+							DiscoveredClients.Add(MakeShared<FLiveLinkHubClientId>(ClientId));
+						}
+					}
+				}
+			}
+
+
+			TWeakPtr<SListView<TSharedPtr<FLiveLinkHubClientId>>> WeakDiscoveredList = DiscoveredClientsListView;
+			TWeakPtr<STreeView<FClientTreeItemPtr>> WeakClients = TreeView;
+			AsyncTask(ENamedThreads::GameThread, [WeakDiscoveredList, WeakClients]()
+			{
+				if (TSharedPtr<SListView<TSharedPtr<FLiveLinkHubClientId>>> List = WeakDiscoveredList.Pin())
+				{
+					List->RequestListRefresh();
+				}
+
+				if (TSharedPtr<STreeView<FClientTreeItemPtr>> ClientsList = WeakClients.Pin())
+				{
+					ClientsList->RequestTreeRefresh();
+				}
+			});
 			break;
 		}
 		case ILiveLinkHubClientsModel::EClientEventType::Modified:
@@ -421,6 +586,48 @@ private:
 		}
 	}
 
+	/** Handler called when a client is added to a session. */
+	void OnClientAddedToSession(FLiveLinkHubClientId ClientId)
+	{
+		int32 ClientIndex = DiscoveredClients.IndexOfByPredicate([ClientId](const TSharedPtr<FLiveLinkHubClientId>& InClient) { return *InClient == ClientId; });
+		if (ClientIndex != INDEX_NONE)
+		{
+			DiscoveredClients.RemoveAt(ClientIndex);
+		}
+
+		TSharedPtr<FClientTreeViewClientItem> ClientItem = MakeShared<FClientTreeViewClientItem>(ClientId, ClientsModel.ToSharedRef());
+		ClientItem->ClientId = ClientId;
+		InitializeClientItem(*ClientItem);
+		Clients.Add(ClientItem);
+
+		TreeView->RequestTreeRefresh();
+		DiscoveredClientsListView->RequestListRefresh();
+	}
+
+	/** Handler called when a client is removed from a session. */
+	void OnClientRemovedFromSession(FLiveLinkHubClientId ClientId)
+	{
+		int32 ClientIndex = Clients.IndexOfByPredicate([ClientId](const FClientTreeItemPtr& InClient) { return InClient->ClientId == ClientId; });
+		if (ClientIndex != INDEX_NONE)
+		{
+			Clients.RemoveAt(ClientIndex);
+		}
+
+		if (ClientsModel->IsClientConnected(ClientId))
+		{
+			DiscoveredClients.Add(MakeShared<FLiveLinkHubClientId>(ClientId));
+			DiscoveredClientsListView->RequestListRefresh();
+		}
+
+		TreeView->RequestTreeRefresh();
+	}
+
+	/** Refreshes the data when the current session changes. */
+	void OnActiveSessionChanged(const TSharedRef<ILiveLinkHubSession>& ActiveSession)
+	{
+		Reinitialize();
+	}
+
 	/** Populate a client item row with its data and children. */
 	void InitializeClientItem(FClientTreeViewClientItem& ClientItem)
 	{
@@ -502,9 +709,17 @@ private:
 	/** Build the client list. */
 	void PopulateClients()
 	{
-		TArray<FLiveLinkHubClientId> ClientList = ClientsModel->GetDiscoveredClients();
+		TArray<FLiveLinkHubClientId> ClientList = ClientsModel->GetSessionClients();
 		Clients.Reset(ClientList.Num());
 
+		TArray<FLiveLinkHubClientId> DiscoveredClientList = ClientsModel->GetDiscoveredClients();
+		DiscoveredClients.Reset(DiscoveredClientList.Num());
+
+		for (const FLiveLinkHubClientId& Client : DiscoveredClientList)
+		{
+			DiscoveredClients.Add(MakeShared<FLiveLinkHubClientId>(Client));
+		}
+
 		for (const FLiveLinkHubClientId& Client : ClientList)
 		{
 			TSharedPtr<FClientTreeViewClientItem> ClientItem = MakeShared<FClientTreeViewClientItem>(Client, ClientsModel.ToSharedRef());
@@ -517,10 +732,18 @@ private:
 private:
 	/** Delegate called when a client is selected. */
 	FOnClientSelected OnClientSelectedDelegate;
+	/** Delegate called when a client is selected. */
+	FOnDiscoveredClientPicked OnDiscoveredClientPickedDelegate;
+	/** Delegate called when a client is deleted. */
+	FOnRemoveClientFromSession OnRemoveClientFromSessionDelegate;
 	/** TreeView widget that displays the clients. */
 	TSharedPtr<STreeView<FClientTreeItemPtr>> TreeView;
-	/** List of message addresses used to identify UE clients and populate the list view. */
+	/** ListView widget that displays clients discovered by the hub that aren't in the current session. */
+	TSharedPtr<SListView<TSharedPtr<FLiveLinkHubClientId>>> DiscoveredClientsListView;
+	/** List of clients in the current session. */
 	TArray<FClientTreeItemPtr> Clients;
+	/** Clients discovered by the hub. */
+	TArray<TSharedPtr<FLiveLinkHubClientId>> DiscoveredClients;
 	/** Model that holds the client data we are displaying. */
 	TSharedPtr<ILiveLinkHubClientsModel> ClientsModel;
 };
diff --git a/Engine/Plugins/Animation/LiveLink/Source/LiveLinkHub/Private/Config/LiveLinkHubConfigData.h b/Engine/Plugins/Animation/LiveLink/Source/LiveLinkHub/Private/Config/LiveLinkHubConfigData.h
deleted file mode 100644
diff --git a/Engine/Plugins/Animation/LiveLink/Source/LiveLinkHub/Private/Config/LiveLinkHubFileUtilities.cpp b/Engine/Plugins/Animation/LiveLink/Source/LiveLinkHub/Private/Config/LiveLinkHubFileUtilities.cpp
--- a/Engine/Plugins/Animation/LiveLink/Source/LiveLinkHub/Private/Config/LiveLinkHubFileUtilities.cpp
+++ b/Engine/Plugins/Animation/LiveLink/Source/LiveLinkHub/Private/Config/LiveLinkHubFileUtilities.cpp
@@ -4,10 +4,10 @@
 
 #include "HAL/FileManager.h"
 #include "JsonObjectConverter.h"
-#include "LiveLinkHubConfigData.h"
 #include "LiveLinkHubLog.h"
+#include "Session/LiveLinkHubSessionData.h"
 
-void UE::LiveLinkHub::FileUtilities::Private::SaveConfig(const FLiveLinkHubConfigData& InConfigData, const FString& InFilePath)
+void UE::LiveLinkHub::FileUtilities::Private::SaveConfig(const FLiveLinkHubPersistedSessionData& InConfigData, const FString& InFilePath)
 {
 	if (ensure(!InFilePath.IsEmpty()))
 	{
@@ -24,7 +24,7 @@ void UE::LiveLinkHub::FileUtilities::Private::SaveConfig(const FLiveLinkHubConfi
 	}
 }
 
-TSharedPtr<FLiveLinkHubConfigData> UE::LiveLinkHub::FileUtilities::Private::LoadConfig(const FString& InFilePath)
+TSharedPtr<FLiveLinkHubPersistedSessionData> UE::LiveLinkHub::FileUtilities::Private::LoadConfig(const FString& InFilePath)
 {
 	if (IFileManager::Get().FileExists(*InFilePath))
 	{
@@ -52,31 +52,31 @@ TSharedPtr<FLiveLinkHubConfigData> UE::LiveLinkHub::FileUtilities::Private::Load
 	return nullptr;
 }
 
-TSharedPtr<FJsonObject> UE::LiveLinkHub::FileUtilities::Private::ToJson(const FLiveLinkHubConfigData& InConfigData)
+TSharedPtr<FJsonObject> UE::LiveLinkHub::FileUtilities::Private::ToJson(const FLiveLinkHubPersistedSessionData& InConfigData)
 {
 	TSharedRef<FJsonObject> JsonObject = MakeShared<FJsonObject>();
 	const TSharedPtr<FJsonValueNumber> NumberValue = MakeShared<FJsonValueNumber>(LiveLinkHubVersion);
 	JsonObject->SetField(JsonVersionKey, NumberValue);
 	
-	FJsonObjectConverter::UStructToJsonObject(FLiveLinkHubConfigData::StaticStruct(), &InConfigData, JsonObject);
+	FJsonObjectConverter::UStructToJsonObject(FLiveLinkHubPersistedSessionData::StaticStruct(), &InConfigData, JsonObject);
 	
 	return JsonObject;
 }
 
-TSharedPtr<FLiveLinkHubConfigData> UE::LiveLinkHub::FileUtilities::Private::FromJson(const TSharedPtr<FJsonObject>& InJsonObject)
+TSharedPtr<FLiveLinkHubPersistedSessionData> UE::LiveLinkHub::FileUtilities::Private::FromJson(const TSharedPtr<FJsonObject>& InJsonObject)
 {
 	if (InJsonObject.IsValid())
 	{
-		FLiveLinkHubConfigData OutConfigData;
+		FLiveLinkHubPersistedSessionData OutConfigData;
 		const bool bResult =
 			FJsonObjectConverter::JsonObjectToUStruct(InJsonObject.ToSharedRef(),
-				FLiveLinkHubConfigData::StaticStruct(), &OutConfigData);
+				FLiveLinkHubPersistedSessionData::StaticStruct(), &OutConfigData);
 
 		if (!bResult)
 		{
-			UE_LOG(LogLiveLinkHub, Error, TEXT("Could not convert from json to LiveLinkHubConfigData."))
+			UE_LOG(LogLiveLinkHub, Error, TEXT("Could not convert from json to LiveLinkHubSessionData."))
 		}
-		return MakeShared<FLiveLinkHubConfigData>(OutConfigData);
+		return MakeShared<FLiveLinkHubPersistedSessionData>(MoveTemp(OutConfigData));
 	}
 
 	return nullptr;
diff --git a/Engine/Plugins/Animation/LiveLink/Source/LiveLinkHub/Private/Config/LiveLinkHubFileUtilities.h b/Engine/Plugins/Animation/LiveLink/Source/LiveLinkHub/Private/Config/LiveLinkHubFileUtilities.h
--- a/Engine/Plugins/Animation/LiveLink/Source/LiveLinkHub/Private/Config/LiveLinkHubFileUtilities.h
+++ b/Engine/Plugins/Animation/LiveLink/Source/LiveLinkHub/Private/Config/LiveLinkHubFileUtilities.h
@@ -5,7 +5,7 @@
 #include "Containers/UnrealString.h"
 #include "Templates/SharedPointer.h"
 
-struct FLiveLinkHubConfigData;
+struct FLiveLinkHubPersistedSessionData;
 class FJsonObject;
 
 namespace UE::LiveLinkHub::FileUtilities::Private
@@ -23,14 +23,14 @@ namespace UE::LiveLinkHub::FileUtilities::Private
 	const FString ConfigDescription = TEXT("Live Link Hub Config");
 
 	/** Save config data to disk. */
-	void SaveConfig(const FLiveLinkHubConfigData& InConfigData, const FString& InFilePath);
+	void SaveConfig(const FLiveLinkHubPersistedSessionData& InConfigData, const FString& InFilePath);
 
 	/** Load config data from disk. */
-	TSharedPtr<FLiveLinkHubConfigData> LoadConfig(const FString& InFilePath);
+	TSharedPtr<FLiveLinkHubPersistedSessionData> LoadConfig(const FString& InFilePath);
 
 	/** Convert config data to json. */
-	TSharedPtr<FJsonObject> ToJson(const FLiveLinkHubConfigData& InConfigData);
+	TSharedPtr<FJsonObject> ToJson(const FLiveLinkHubPersistedSessionData& InConfigData);
 
 	/** Convert config data from json. */
-	TSharedPtr<FLiveLinkHubConfigData> FromJson(const TSharedPtr<FJsonObject>& InJsonObject);
+	TSharedPtr<FLiveLinkHubPersistedSessionData> FromJson(const TSharedPtr<FJsonObject>& InJsonObject);
 }
\ No newline at end of file
diff --git a/Engine/Plugins/Animation/LiveLink/Source/LiveLinkHub/Private/LiveLinkHub.cpp b/Engine/Plugins/Animation/LiveLink/Source/LiveLinkHub/Private/LiveLinkHub.cpp
--- a/Engine/Plugins/Animation/LiveLink/Source/LiveLinkHub/Private/LiveLinkHub.cpp
+++ b/Engine/Plugins/Animation/LiveLink/Source/LiveLinkHub/Private/LiveLinkHub.cpp
@@ -4,7 +4,6 @@
 
 #include "Clients/LiveLinkHubClientsController.h"
 #include "Clients/LiveLinkHubProvider.h"
-#include "Config/LiveLinkHubConfigData.h"
 #include "Config/LiveLinkHubFileUtilities.h"
 #include "DesktopPlatformModule.h"
 #include "EditorDirectories.h"
@@ -13,15 +12,15 @@
 #include "IDesktopPlatform.h"
 #include "LiveLinkHubClient.h"
 #include "LiveLinkProvider.h"
-#include "LiveLinkSubject.h"
 #include "LiveLinkHubCommands.h"
-#include "LiveLinkSubject.h"
 #include "LiveLinkProviderImpl.h"
+#include "LiveLinkSubject.h"
 #include "LiveLinkSubjectSettings.h"
 #include "Misc/ConfigCacheIni.h"
 #include "Recording/LiveLinkHubPlaybackController.h"
 #include "Recording/LiveLinkHubRecordingController.h"
 #include "Recording/LiveLinkHubRecordingListController.h"
+#include "Session/LiveLinkHubSessionManager.h"
 #include "Subjects/LiveLinkHubSubjectController.h"
 #include "UI/Window/LiveLinkHubWindowController.h"
 
@@ -76,6 +75,8 @@ FLiveLinkHub::~FLiveLinkHub()
 	LiveLinkHubClient->OnLiveLinkSubjectAdded().RemoveAll(this);
 	LiveLinkHubClient->OnFrameDataReceived_AnyThread().RemoveAll(this);
 	LiveLinkHubClient->OnStaticDataReceived_AnyThread().RemoveAll(this);
+
+	IModularFeatures::Get().UnregisterModularFeature(ILiveLinkClient::ModularFeatureName, LiveLinkHubClient.Get());
 }
 
 bool FLiveLinkHub::IsInPlayback() const
@@ -188,149 +189,29 @@ void FLiveLinkHub::BindCommands()
 		FCanExecuteAction::CreateSP(this, &FLiveLinkHub::CanSaveConfig));
 }
 
-void FLiveLinkHub::ClearClient()
-{
-	check(LiveLinkProvider.IsValid());
-	check(LiveLinkHubClient.IsValid());
-	
-	TMap<FLiveLinkHubClientId, FLiveLinkHubUEClientInfo> Clients = LiveLinkProvider->GetClientsMap();
-	
-	for (const TTuple<FLiveLinkHubClientId, FLiveLinkHubUEClientInfo>& Client : Clients)
-	{
-		LiveLinkProvider->RemoveClient(Client.Key);
-	}
-	
-	LiveLinkHubClient->RemoveAllSources();
-
-	// Removing sources only sets bPendingKill to true, need to tick to make sure they are removed.
-	LiveLinkHubClient->ForceTick();
-}
-
 void FLiveLinkHub::NewConfig()
 {
-	ClearClient();
-	LastConfigPath.Empty();
+	SessionManager->NewSession();
 }
 
 void FLiveLinkHub::SaveConfigAs()
 {
-	const FString FileDescription = UE::LiveLinkHub::FileUtilities::Private::ConfigDescription;
-	const FString Extensions = UE::LiveLinkHub::FileUtilities::Private::ConfigExtension;
-	const FString FileTypes = FString::Printf(TEXT("%s (*.%s)|*.%s"), *FileDescription, *Extensions, *Extensions);
-
-	const FString DefaultFile = UE::LiveLinkHub::FileUtilities::Private::ConfigDefaultFileName;
-	
-	TArray<FString> SaveFileNames;
-
-	IDesktopPlatform* DesktopPlatform = FDesktopPlatformModule::Get();
-	const void* ParentWindowWindowHandle = FSlateApplication::Get().FindBestParentWindowHandleForDialogs(nullptr);
-	
-	const bool bFileSelected = DesktopPlatform->SaveFileDialog(
-		ParentWindowWindowHandle,
-		LOCTEXT("LiveLinkHubSaveAsTitle", "Save As").ToString(),
-		FEditorDirectories::Get().GetLastDirectory(ELastDirectory::GENERIC_SAVE),
-		DefaultFile,
-		FileTypes,
-		EFileDialogFlags::None,
-		SaveFileNames);
-
-	if (bFileSelected && SaveFileNames.Num() > 0)
-	{
-		LastConfigPath = SaveFileNames[0];
-		FEditorDirectories::Get().SetLastDirectory(ELastDirectory::GENERIC_SAVE, FPaths::GetPath(LastConfigPath));
-		SaveConfig();
-	}
+	SessionManager->SaveSessionAs();
 }
 
 bool FLiveLinkHub::CanSaveConfig() const
 {
-	return !LastConfigPath.IsEmpty();
+	return SessionManager->CanSaveCurrentSession();
 }
 
 void FLiveLinkHub::SaveConfig()
 {
-	if (LastConfigPath.IsEmpty())
-	{
-		return;
-	}
-	
-	check(LiveLinkProvider.IsValid())
-	check(LiveLinkHubClient.IsValid());
-	
-	FLiveLinkHubConfigData LiveLinkHubConfigData;
-	
-	TArray<FGuid> SourceGuids = LiveLinkHubClient->GetSources();
-	for (const FGuid& SourceGuid : SourceGuids)
-	{
-		LiveLinkHubConfigData.Sources.Add(LiveLinkHubClient->GetSourcePreset(SourceGuid, nullptr));
-	}
-
-	TArray<FLiveLinkSubjectKey> Subjects = LiveLinkHubClient->GetSubjects(true, true);
-	for (const FLiveLinkSubjectKey& Subject : Subjects)
-	{
-		LiveLinkHubConfigData.Subjects.Add(LiveLinkHubClient->GetSubjectPreset(Subject, nullptr));
-	}
-
-	const TMap<FLiveLinkHubClientId, FLiveLinkHubUEClientInfo>& ClientMap = LiveLinkProvider->GetClientsMap();
-	
-	for (const TTuple<FLiveLinkHubClientId, FLiveLinkHubUEClientInfo>& ClientKeyVal : ClientMap)
-	{
-		LiveLinkHubConfigData.Clients.Add(ClientKeyVal.Value);
-	}
-
-	UE::LiveLinkHub::FileUtilities::Private::SaveConfig(LiveLinkHubConfigData, LastConfigPath);
+	SessionManager->SaveCurrentSession();
 }
 
 void FLiveLinkHub::OpenConfig()
 {
-	const FString FileDescription = UE::LiveLinkHub::FileUtilities::Private::ConfigDescription;
-	const FString Extensions = UE::LiveLinkHub::FileUtilities::Private::ConfigExtension;
-	const FString FileTypes = FString::Printf(TEXT("%s (*.%s)|*.%s"), *FileDescription, *Extensions, *Extensions);
-
-	const FString DefaultFile = UE::LiveLinkHub::FileUtilities::Private::ConfigDefaultFileName;
-
-	check(LiveLinkProvider.IsValid())
-	check(LiveLinkHubClient.IsValid());
-	
-	ClearClient();
-	
-	TArray<FString> OpenFileNames;
-
-	IDesktopPlatform* DesktopPlatform = FDesktopPlatformModule::Get();
-	const void* ParentWindowWindowHandle = FSlateApplication::Get().FindBestParentWindowHandleForDialogs(nullptr);
-	const bool bFileSelected = DesktopPlatform->OpenFileDialog(
-		ParentWindowWindowHandle,
-		LOCTEXT("LiveLinkHubOpenTitle", "Open").ToString(),
-		FEditorDirectories::Get().GetLastDirectory(ELastDirectory::GENERIC_OPEN),
-		DefaultFile,
-		FileTypes,
-		EFileDialogFlags::None,
-		OpenFileNames);
-
-	if (bFileSelected && OpenFileNames.Num() > 0)
-	{
-		LastConfigPath = OpenFileNames[0];
-		FEditorDirectories::Get().SetLastDirectory(ELastDirectory::GENERIC_OPEN, FPaths::GetPath(LastConfigPath));
-
-		const TSharedPtr<FLiveLinkHubConfigData> ConfigData = UE::LiveLinkHub::FileUtilities::Private::LoadConfig(LastConfigPath);
-		if (ConfigData.IsValid())
-		{
-			for (const FLiveLinkSourcePreset& SourcePreset : ConfigData->Sources)
-			{
-				LiveLinkHubClient->CreateSource(SourcePreset);
-			}
-
-			for (const FLiveLinkSubjectPreset& SubjectPreset : ConfigData->Subjects)
-			{
-				LiveLinkHubClient->CreateSubject(SubjectPreset);	
-			}
-
-			for (const FLiveLinkHubUEClientInfo& Client : ConfigData->Clients)
-			{
-				LiveLinkProvider->AddClient(Client);
-			}
-		}
-	}
+	SessionManager->RestoreSession();
 }
 
 FName FLiveLinkHub::GetSubjectNameOverride(const FLiveLinkSubjectKey& InSubjectKey) const
@@ -339,7 +220,7 @@ FName FLiveLinkHub::GetSubjectNameOverride(const FLiveLinkSubjectKey& InSubjectK
 	{
 		if (const TSharedPtr<ILiveLinkHubSession> CurrentSession = Manager->GetCurrentSession())
 		{
-			if (const ULiveLinkHubSubjectProxy* SubjectProxy = CurrentSession->GetSubjectConfig(InSubjectKey))
+			if (TOptional<FLiveLinkHubSubjectProxy> SubjectProxy = CurrentSession->GetSubjectConfig(InSubjectKey))
 			{
 				return SubjectProxy->GetOutboundName();
 			}
diff --git a/Engine/Plugins/Animation/LiveLink/Source/LiveLinkHub/Private/LiveLinkHub.h b/Engine/Plugins/Animation/LiveLink/Source/LiveLinkHub/Private/LiveLinkHub.h
--- a/Engine/Plugins/Animation/LiveLink/Source/LiveLinkHub/Private/LiveLinkHub.h
+++ b/Engine/Plugins/Animation/LiveLink/Source/LiveLinkHub/Private/LiveLinkHub.h
@@ -82,7 +82,6 @@ private:
 	void BindCommands();
 
 	/** Clear all client settings. */
-	void ClearClient();
 	FName GetSubjectNameOverride(const FLiveLinkSubjectKey& InSubjectKey) const;
 
 	/** Create a new config. */
diff --git a/Engine/Plugins/Animation/LiveLink/Source/LiveLinkHub/Private/LiveLinkHubModule.cpp b/Engine/Plugins/Animation/LiveLink/Source/LiveLinkHub/Private/LiveLinkHubModule.cpp
--- a/Engine/Plugins/Animation/LiveLink/Source/LiveLinkHub/Private/LiveLinkHubModule.cpp
+++ b/Engine/Plugins/Animation/LiveLink/Source/LiveLinkHub/Private/LiveLinkHubModule.cpp
@@ -4,6 +4,7 @@
 
 #include "Clients/LiveLinkHubProvider.h"
 #include "LiveLinkHubApplication.h"
+#include "LiveLinkHubLog.h"
 #include "Modules/ModuleManager.h"
 #include "Recording/LiveLinkHubPlaybackController.h"
 #include "Recording/LiveLinkHubRecordingController.h"
@@ -58,22 +59,22 @@ TSharedPtr<FLiveLinkHub> FLiveLinkHubModule::GetLiveLinkHub() const
 
 TSharedPtr<FLiveLinkHubProvider> FLiveLinkHubModule::GetLiveLinkProvider() const
 {
-	return LiveLinkHub->LiveLinkProvider;
+	return LiveLinkHub ? LiveLinkHub->LiveLinkProvider : nullptr;
 }
 
 TSharedPtr<FLiveLinkHubRecordingController> FLiveLinkHubModule::GetRecordingController() const
 {
-	return LiveLinkHub->RecordingController;
+	return LiveLinkHub ? LiveLinkHub->RecordingController : nullptr;
 }
 
 TSharedPtr<FLiveLinkHubRecordingListController> FLiveLinkHubModule::GetRecordingListController() const
 {
-	return LiveLinkHub->RecordingListController;
+	return LiveLinkHub ? LiveLinkHub->RecordingListController : nullptr;
 }
 
 TSharedPtr<FLiveLinkHubPlaybackController> FLiveLinkHubModule::GetPlaybackController() const
 {
-	return LiveLinkHub->PlaybackController;
+	return LiveLinkHub ? LiveLinkHub->PlaybackController : nullptr;
 }
 
 #if !WITH_LIVELINK_HUB
@@ -143,12 +144,12 @@ void FLiveLinkHubModule::OpenLiveLinkHub() const
 
 TSharedPtr<FLiveLinkHubSubjectController> FLiveLinkHubModule::GetSubjectController() const
 {
-	return LiveLinkHub->SubjectController;
+	return LiveLinkHub ? LiveLinkHub->SubjectController : nullptr;
 }
 
 TSharedPtr<ILiveLinkHubSessionManager> FLiveLinkHubModule::GetSessionManager() const
 {
-	return LiveLinkHub->SessionManager;
+	return LiveLinkHub ? LiveLinkHub->SessionManager : nullptr;
 }
 
 IMPLEMENT_MODULE(FLiveLinkHubModule, LiveLinkHub);
diff --git a/Engine/Plugins/Animation/LiveLink/Source/LiveLinkHub/Private/LiveLinkHubSettings.h b/Engine/Plugins/Animation/LiveLink/Source/LiveLinkHub/Private/LiveLinkHubSettings.h
new file mode 100644
--- /dev/null
+++ b/Engine/Plugins/Animation/LiveLink/Source/LiveLinkHub/Private/LiveLinkHubSettings.h
@@ -0,0 +1,21 @@
+// Copyright Epic Games, Inc. All Rights Reserved.
+
+#pragma once
+
+#include "UObject/Object.h"
+
+#include "LiveLinkHubSettings.generated.h"
+
+/**
+ * Settings for LiveLinkHub.
+ */
+UCLASS(config=Game, defaultconfig)
+class LIVELINKHUB_API ULiveLinkHubSettings : public UObject
+{
+	GENERATED_BODY()
+
+public:
+	/** If enabled, discovered clients will be automatically added to the current session. */
+	UPROPERTY(config, EditAnywhere, Category="LiveLinkHub")
+	bool bAutoAddDiscoveredClients = true;
+};
diff --git a/Engine/Plugins/Animation/LiveLink/Source/LiveLinkHub/Private/Session/LiveLinkHubSession.h b/Engine/Plugins/Animation/LiveLink/Source/LiveLinkHub/Private/Session/LiveLinkHubSession.h
--- a/Engine/Plugins/Animation/LiveLink/Source/LiveLinkHub/Private/Session/LiveLinkHubSession.h
+++ b/Engine/Plugins/Animation/LiveLink/Source/LiveLinkHub/Private/Session/LiveLinkHubSession.h
@@ -2,63 +2,233 @@
 
 #pragma once
 
+#include "Clients/LiveLinkHubProvider.h"
+#include "Clients/LiveLinkHubUEClientInfo.h"
+#include "Features/IModularFeatures.h"
+#include "Framework/Application/SlateApplication.h"
+#include "LiveLinkClient.h"
+#include "LiveLinkHubModule.h"
+#include "LiveLinkHubSessionData.h"
 #include "LiveLinkTypes.h"
+#include "Misc/ScopeLock.h"
+#include "Modules/ModuleManager.h"
 #include "Subjects/LiveLinkHubSubjectSessionConfig.h"
 #include "UObject/StrongObjectPtr.h"
 
+#define LOCTEXT_NAMESPACE "LiveLinkHubSession"
+
+DECLARE_MULTICAST_DELEGATE_OneParam(FOnClientAddedToSession, FLiveLinkHubClientId);
+DECLARE_MULTICAST_DELEGATE_OneParam(FOnClientRemovedFromSession, FLiveLinkHubClientId);
+
 /**
- * Note: This is just a stub implementation until the session manager work is submitted.
+ * Holds the state of the hub for an active session, can be swapped out with a different session using the session manager.
  */
 class ILiveLinkHubSession
 {
 public:
 	virtual ~ILiveLinkHubSession() = default;
-	virtual ULiveLinkHubSubjectProxy* GetSubjectConfig(const FLiveLinkSubjectKey& SubjectKey) const = 0;
+
+	/** Get the configuration for a given subject. */
+	virtual TOptional<FLiveLinkHubSubjectProxy> GetSubjectConfig(const FLiveLinkSubjectKey& SubjectKey) const = 0;
+
+	/** Change the outbound name of a subject for this session. */
+	virtual void RenameSubject(const FLiveLinkSubjectKey& SubjectKey, FName NewName) = 0;
+
+	/** Add a client to this session. Note: Must be called from game thread. */
+	virtual void AddClient(const FLiveLinkHubClientId& Client) = 0;
+
+	/** Remove a client from this session. Note: Must be called from game thread. */
+	virtual void RemoveClient(const FLiveLinkHubClientId& Client) = 0;
+
+	/** Returns whether a client is in this session. */
+	virtual bool IsClientInSession(const FLiveLinkHubClientId& Client) = 0;
+
+	/** Get the list of clients in this session (The list of clients that can receive data from the hub) */
+	virtual TArray<FLiveLinkHubClientId> GetSessionClients() const = 0;
 };
 
-class FLiveLinkHubSession : public ILiveLinkHubSession
+class FLiveLinkHubSession : public ILiveLinkHubSession, public TSharedFromThis<FLiveLinkHubSession>
 {
 public:
-	FLiveLinkHubSession()
+	FLiveLinkHubSession(FOnClientAddedToSession& OnClientAddedToSession, FOnClientRemovedFromSession& OnClientRemovedFromSession)
+		: OnClientAddedToSessionDelegate(OnClientAddedToSession)
+		, OnClientRemovedFromSessionDelegate(OnClientRemovedFromSession)
+	{
+		RegisterDelegates();
+
+		SessionData.SubjectsConfig.Initialize();
+	}
+
+	FLiveLinkHubSession(FLiveLinkHubSessionData InSessionData, FOnClientAddedToSession& OnClientAddedToSession, FOnClientRemovedFromSession& OnClientRemovedFromSession)
+		: SessionData(MoveTemp(InSessionData))
+		, OnClientAddedToSessionDelegate(OnClientAddedToSession)
+		, OnClientRemovedFromSessionDelegate(OnClientRemovedFromSession)
+	{
+		RegisterDelegates();
+		SessionData.SubjectsConfig.Initialize();
+	}
+
+	virtual ~FLiveLinkHubSession() override
+	{
+		UnregisterDelegates();
+	}
+
+	virtual void RenameSubject(const FLiveLinkSubjectKey& SubjectKey, FName NewName) override
+	{
+		FWriteScopeLock Locker(SessionDataLock);
+		SessionData.SubjectsConfig.RenameSubject(SubjectKey, NewName);
+	}
+
+	virtual TOptional<FLiveLinkHubSubjectProxy> GetSubjectConfig(const FLiveLinkSubjectKey& SubjectKey) const override
+	{
+		FWriteScopeLock Locker(SessionDataLock);
+		return SessionData.SubjectsConfig.GetSubjectConfig(SubjectKey);
+	}
+
+	virtual TArray<FLiveLinkHubClientId> GetSessionClients() const override
+	{
+		FReadScopeLock Locker(SessionDataLock);
+		return CachedSessionClients.Array();
+	}
+
+	virtual void AddClient(const FLiveLinkHubClientId& Client) override
+	{
+		check(IsInGameThread());
+
+		{
+			if (TSharedPtr<FLiveLinkHubProvider> LiveLinkProvider = FModuleManager::Get().GetModuleChecked<FLiveLinkHubModule>("LiveLinkHub").GetLiveLinkProvider())
+			{
+				if (TOptional<FLiveLinkHubUEClientInfo> ClientInfo = LiveLinkProvider->GetClientInfo(Client))
+				{
+					FWriteScopeLock Locker(SessionDataLock);
+					CachedSessionClients.Add(Client);
+				}
+			}
+		}
+
+		OnClientAddedToSessionDelegate.Broadcast(Client);
+	}
+
+	virtual void RemoveClient(const FLiveLinkHubClientId& Client) override
 	{
-		SessionConfig = TStrongObjectPtr<ULiveLinkHubSubjectSessionConfig>(NewObject<ULiveLinkHubSubjectSessionConfig>());
+		check(IsInGameThread());
+
+		{
+			FWriteScopeLock Locker(SessionDataLock);
+			CachedSessionClients.Remove(Client);
+		}
+
+		OnClientRemovedFromSessionDelegate.Broadcast(Client);
 	}
 
-	virtual ~FLiveLinkHubSession() override = default;
+	virtual bool IsClientInSession(const FLiveLinkHubClientId& Client) override
+	{
+		FReadScopeLock Locker(SessionDataLock);
+		return CachedSessionClients.Contains(Client);
+	}
 
-	virtual ULiveLinkHubSubjectProxy* GetSubjectConfig(const FLiveLinkSubjectKey& SubjectKey) const override
+	void AddRestoredClient(FLiveLinkHubUEClientInfo& InOutRestoredClientInfo)
 	{
-		return SessionConfig->GetSubjectConfig(SubjectKey);
+		if (const TSharedPtr<FLiveLinkHubProvider> LiveLinkProvider = FModuleManager::Get().GetModuleChecked<FLiveLinkHubModule>("LiveLinkHub").GetLiveLinkProvider())
+		{
+			LiveLinkProvider->AddRestoredClient(InOutRestoredClientInfo);
+
+			{
+				FWriteScopeLock Locker(SessionDataLock);
+				CachedSessionClients.Add(InOutRestoredClientInfo.Id);
+			}
+		}
+
+		OnClientAddedToSessionDelegate.Broadcast(InOutRestoredClientInfo.Id);
 	}
 
 private:
-	TStrongObjectPtr<ULiveLinkHubSubjectSessionConfig> SessionConfig;
-};
+	/** Register livelink client delegates used to update the config's subject data. */
+	void RegisterDelegates()
+	{
+		FLiveLinkClient& LiveLinkClient = IModularFeatures::Get().GetModularFeature<FLiveLinkClient>(ILiveLinkClient::ModularFeatureName);
+		LiveLinkClient.OnLiveLinkSubjectAdded().AddRaw(this, &FLiveLinkHubSession::OnSubjectAdded_AnyThread);
+		LiveLinkClient.OnLiveLinkSubjectRemoved().AddRaw(this, &FLiveLinkHubSession::OnSubjectRemoved_AnyThread);
+	}
 
-class ILiveLinkHubSessionManager
-{
-public:
-	virtual ~ILiveLinkHubSessionManager() = default;
+	/** Unregister livelink client delegates. */
+	void UnregisterDelegates()
+	{
+		if (IModularFeatures::Get().IsModularFeatureAvailable(ILiveLinkClient::ModularFeatureName))
+		{
+			FLiveLinkClient& LiveLinkClient = IModularFeatures::Get().GetModularFeature<FLiveLinkClient>(ILiveLinkClient::ModularFeatureName);
+			LiveLinkClient.OnLiveLinkSubjectRemoved().RemoveAll(this);
+			LiveLinkClient.OnLiveLinkSubjectAdded().RemoveAll(this);
+		}
+	}
 
-	// Todo: Real implementation must be thread safe
-	virtual TSharedPtr<ILiveLinkHubSession> GetCurrentSession() const = 0;
-};
+	/** AnyThread handler for the SubjectAdded delegate, dispatches handling on the game thread to avoid asserts in Slate. */
+	void OnSubjectAdded_AnyThread(FLiveLinkSubjectKey SubjectKey)
+	{
+		TWeakPtr<FLiveLinkHubSession> WeakSession = AsShared();
+		AsyncTask(ENamedThreads::GameThread, [WeakSession, Key = MoveTemp(SubjectKey)]
+			{
+				if (TSharedPtr<FLiveLinkHubSession> Session = WeakSession.Pin())
+				{
+					Session->OnSubjectAdded(Key);
+				}
+			});
+	}
 
-class FLiveLinkHubSessionManager : public ILiveLinkHubSessionManager
-{
-public:
-	FLiveLinkHubSessionManager()
+	/** Handles updating the tree view when a subject is added. */
+	void OnSubjectAdded(const FLiveLinkSubjectKey& SubjectKey)
 	{
-		CurrentSession = MakeShared<FLiveLinkHubSession>();
+		if (!SessionData.SubjectsConfig.SubjectProxies.Contains(SubjectKey))
+		{
+			ILiveLinkClient& LiveLinkClient = IModularFeatures::Get().GetModularFeature<ILiveLinkClient>(ILiveLinkClient::ModularFeatureName);
+
+			FLiveLinkHubSubjectProxy SubjectSettings;
+			SubjectSettings.Initialize(SubjectKey, LiveLinkClient.GetSourceType(SubjectKey.Source).ToString());
+
+			{
+				FWriteScopeLock Locker(SessionDataLock);
+				SessionData.SubjectsConfig.SubjectProxies.FindOrAdd(SubjectKey) = MoveTemp(SubjectSettings);
+			}
+		}
 	}
 
-	virtual ~FLiveLinkHubSessionManager() override = default;
+	/** AnyThread handler for the SubjectRemoved delegate, dispatches handling on the game thread to avoid asserts in Slate. */
+	void OnSubjectRemoved_AnyThread(FLiveLinkSubjectKey SubjectKey)
+	{
+		TWeakPtr<FLiveLinkHubSession> WeakSession = AsShared();
+		AsyncTask(ENamedThreads::GameThread, [WeakSession, Key = MoveTemp(SubjectKey)]
+			{
+				if (TSharedPtr<FLiveLinkHubSession> Session = WeakSession.Pin())
+				{
+					Session->OnSubjectRemoved(Key);
+				}
+			});
+	}
 
-	virtual TSharedPtr<ILiveLinkHubSession> GetCurrentSession() const override
+	/** Handles updating the tree view when a subject is removed. */
+	void OnSubjectRemoved(const FLiveLinkSubjectKey& SubjectKey)
 	{
-		return CurrentSession;
+		FWriteScopeLock Locker(SessionDataLock);
+		SessionData.SubjectsConfig.SubjectProxies.Remove(SubjectKey);
 	}
 
 private:
-	TSharedPtr<ILiveLinkHubSession> CurrentSession;
+	/** List of clients in the current session. These represent the unreal instances than can receive data from the hub. */
+	TSet<FLiveLinkHubClientId> CachedSessionClients;
+
+	/** Holds data for this session. */
+	FLiveLinkHubSessionData SessionData;
+
+	/** Delegate used to notice the hub about clients being added to this session. */
+	FOnClientAddedToSession& OnClientAddedToSessionDelegate;
+
+	/** Delegate used to notice the hub about clients being removed from this session. */
+	FOnClientRemovedFromSession& OnClientRemovedFromSessionDelegate;
+
+	/** Lock used to access the client config from different threads. */
+	mutable FRWLock SessionDataLock;
+
+	friend class FLiveLinkHubSessionManager;
 };
+
+#undef LOCTEXT_NAMESPACE
diff --git a/Engine/Plugins/Animation/LiveLink/Source/LiveLinkHub/Private/Session/LiveLinkHubSessionData.h b/Engine/Plugins/Animation/LiveLink/Source/LiveLinkHub/Private/Session/LiveLinkHubSessionData.h
new file mode 100644
--- /dev/null
+++ b/Engine/Plugins/Animation/LiveLink/Source/LiveLinkHub/Private/Session/LiveLinkHubSessionData.h
@@ -0,0 +1,49 @@
+// Copyright Epic Games, Inc. All Rights Reserved.
+
+#pragma once
+
+#include "Clients/LiveLinkHubUEClientInfo.h"
+#include "CoreTypes.h"
+#include "LiveLinkPresetTypes.h"
+#include "Misc/Guid.h"
+#include "Subjects/LiveLinkHubSubjectSessionConfig.h"
+
+
+#include "LiveLinkHubSessionData.generated.h"
+
+/** Live link hub session data which can be serialized to disk. */
+USTRUCT()
+struct FLiveLinkHubSessionData
+{
+	GENERATED_BODY()
+
+	UPROPERTY()
+	/** Subject configs for this session. */
+	FLiveLinkHubSubjectSessionConfig SubjectsConfig;
+};
+
+/** Live link hub session data that can be saved to disk. */
+USTRUCT()
+struct FLiveLinkHubPersistedSessionData : public FLiveLinkHubSessionData
+{
+	GENERATED_BODY()
+
+	FLiveLinkHubPersistedSessionData() = default;
+
+	FLiveLinkHubPersistedSessionData(FLiveLinkHubSessionData SessionData)
+		: FLiveLinkHubSessionData(MoveTemp(SessionData))
+	{
+	}
+
+	/** Live link hub sources. */
+	UPROPERTY()
+	TArray<FLiveLinkSourcePreset> Sources;
+
+	/** Live link hub subjects. */
+	UPROPERTY()
+	TArray<FLiveLinkSubjectPreset> Subjects;
+
+	/** Live link hub client info. */
+	UPROPERTY()
+	TArray<FLiveLinkHubUEClientInfo> Clients;
+};
diff --git a/Engine/Plugins/Animation/LiveLink/Source/LiveLinkHub/Private/Session/LiveLinkHubSessionManager.h b/Engine/Plugins/Animation/LiveLink/Source/LiveLinkHub/Private/Session/LiveLinkHubSessionManager.h
new file mode 100644
--- /dev/null
+++ b/Engine/Plugins/Animation/LiveLink/Source/LiveLinkHub/Private/Session/LiveLinkHubSessionManager.h
@@ -0,0 +1,282 @@
+// Copyright Epic Games, Inc. All Rights Reserved.
+
+#pragma once
+
+#include "LiveLinkHubSession.h"
+
+#include "Config/LiveLinkHubFileUtilities.h"
+#include "DesktopPlatformModule.h"
+#include "EditorDirectories.h"
+#include "HAL/CriticalSection.h"
+#include "IDesktopPlatform.h"
+#include "LiveLinkHubClient.h"
+
+#define LOCTEXT_NAMESPACE "LiveLinkHub.SessionManager"
+
+DECLARE_MULTICAST_DELEGATE_OneParam(FOnActiveSessionChanged, const TSharedRef<ILiveLinkHubSession>& /**Active Session*/);
+
+class ILiveLinkHubSessionManager
+{
+public:
+	virtual ~ILiveLinkHubSessionManager() = default;
+
+	/** Delegate called when a UE client is added to the current session, enabling it to receive data from the hub. */
+	virtual FOnClientAddedToSession& OnClientAddedToSession() = 0;
+
+	/** Delegate called when a UE client is removed from the current session, returning it to the list of discovered clients. */
+	virtual FOnClientRemovedFromSession& OnClientRemovedFromSession() = 0;
+
+	/** Delegate called wehen the active session changes, which will change the list of sources, subjects and clients. */
+	virtual FOnActiveSessionChanged& OnActiveSessionChanged() = 0;
+
+	/** Get the current session, which holds information about which sources, subjects and clients that should be enabled in the hub at the moment. */
+	virtual TSharedPtr<ILiveLinkHubSession> GetCurrentSession() const = 0;
+
+	/** Clear out the current session data and stat a new empty session. */
+	virtual void NewSession() = 0;
+
+	/** Prompt the user to save the current session in a given directory. */
+	virtual void SaveSessionAs() = 0;
+
+	/** Prompt the user to pick a session file to restore. */
+	virtual void RestoreSession() = 0;
+
+	/** Save the current session. If not path is specified, the last save path will be used. */
+	virtual void SaveCurrentSession(const FString& SavePath = TEXT("")) = 0;
+
+	/** Returns whether the current session has as already been saved to disk before. */
+	virtual bool CanSaveCurrentSession() const = 0;
+};
+
+class FLiveLinkHubSessionManager : public ILiveLinkHubSessionManager
+{
+public:
+	FLiveLinkHubSessionManager()
+	{
+		FScopeLock Lock(&CurrentSessionCS);
+		CurrentSession = MakeShared<FLiveLinkHubSession>(OnClientAddedToSessionDelegate, OnClientRemovedFromSessionDelegate);
+	}
+
+	virtual ~FLiveLinkHubSessionManager() override = default;
+
+	//~ Begin LiveLinkHubSessionManager
+	virtual FOnClientAddedToSession& OnClientAddedToSession() override
+	{
+		check(IsInGameThread());
+		return OnClientAddedToSessionDelegate;
+	}
+
+	virtual FOnClientRemovedFromSession& OnClientRemovedFromSession() override
+	{
+		check(IsInGameThread());
+		return OnClientRemovedFromSessionDelegate;
+	}
+
+	virtual FOnActiveSessionChanged& OnActiveSessionChanged() override
+	{
+		check(IsInGameThread());
+		return OnActiveSessionChangedDelegate;
+	}
+
+	virtual void NewSession() override
+	{
+		ClearSession();
+		LastConfigPath.Empty();
+	}
+
+	virtual void SaveSessionAs() override
+	{
+		const FString FileDescription = UE::LiveLinkHub::FileUtilities::Private::ConfigDescription;
+		const FString Extensions = UE::LiveLinkHub::FileUtilities::Private::ConfigExtension;
+		const FString FileTypes = FString::Printf(TEXT("%s (*.%s)|*.%s"), *FileDescription, *Extensions, *Extensions);
+
+		const FString DefaultFile = UE::LiveLinkHub::FileUtilities::Private::ConfigDefaultFileName;
+
+		TArray<FString> SaveFileNames;
+
+		IDesktopPlatform* DesktopPlatform = FDesktopPlatformModule::Get();
+		const void* ParentWindowWindowHandle = FSlateApplication::Get().FindBestParentWindowHandleForDialogs(nullptr);
+
+		const bool bFileSelected = DesktopPlatform->SaveFileDialog(
+			ParentWindowWindowHandle,
+			LOCTEXT("LiveLinkHubSaveAsTitle", "Save As").ToString(),
+			FEditorDirectories::Get().GetLastDirectory(ELastDirectory::GENERIC_SAVE),
+			DefaultFile,
+			FileTypes,
+			EFileDialogFlags::None,
+			SaveFileNames);
+
+		if (bFileSelected && SaveFileNames.Num() > 0)
+		{
+			SaveCurrentSession(SaveFileNames[0]);
+		}
+	}
+
+	virtual TSharedPtr<ILiveLinkHubSession> GetCurrentSession() const override
+	{
+		FScopeLock Lock(&CurrentSessionCS);
+		return CurrentSession;
+	}
+
+	virtual void SaveCurrentSession(const FString& SavePath = TEXT("")) override
+	{
+		if (SavePath.IsEmpty() && LastConfigPath.IsEmpty())
+		{
+			return;
+		}
+
+		FLiveLinkHubModule& LiveLinkHubModule = FModuleManager::Get().GetModuleChecked<FLiveLinkHubModule>("LiveLinkHub");
+		TSharedPtr<FLiveLinkHubProvider> LiveLinkProvider = LiveLinkHubModule.GetLiveLinkProvider();
+		FLiveLinkHubClient* LiveLinkHubClient = static_cast<FLiveLinkHubClient*>(&IModularFeatures::Get().GetModularFeature<ILiveLinkClient>(ILiveLinkClient::ModularFeatureName));
+
+		check(LiveLinkProvider);
+		check(LiveLinkHubClient);
+
+		TSharedPtr<FLiveLinkHubSession> CurrentSessionPtr;
+		{
+			FScopeLock Lock(&CurrentSessionCS);
+			CurrentSessionPtr =  CurrentSession;
+		}
+
+		FLiveLinkHubPersistedSessionData LiveLinkHubSessionData = StaticCastSharedPtr<FLiveLinkHubSession>(CurrentSessionPtr)->SessionData;
+
+		TArray<FGuid> SourceGuids = LiveLinkHubClient->GetSources();
+		for (const FGuid& SourceGuid : SourceGuids)
+		{
+			LiveLinkHubSessionData.Sources.Add(LiveLinkHubClient->GetSourcePreset(SourceGuid, nullptr));
+		}
+
+		TArray<FLiveLinkSubjectKey> Subjects = LiveLinkHubClient->GetSubjects(true, true);
+		for (const FLiveLinkSubjectKey& Subject : Subjects)
+		{
+			LiveLinkHubSessionData.Subjects.Add(LiveLinkHubClient->GetSubjectPreset(Subject, nullptr));
+		}
+
+		const TMap<FLiveLinkHubClientId, FLiveLinkHubUEClientInfo>& ClientMap = LiveLinkProvider->GetClientsMap();
+
+		for (const TTuple<FLiveLinkHubClientId, FLiveLinkHubUEClientInfo>& ClientKeyVal : ClientMap)
+		{
+			LiveLinkHubSessionData.Clients.Add(ClientKeyVal.Value);
+		}
+
+		LastConfigPath = SavePath;
+		FEditorDirectories::Get().SetLastDirectory(ELastDirectory::GENERIC_SAVE, FPaths::GetPath(LastConfigPath));
+
+		UE::LiveLinkHub::FileUtilities::Private::SaveConfig(LiveLinkHubSessionData, LastConfigPath);
+	}
+
+	virtual void RestoreSession() override
+	{
+		const FString FileDescription = UE::LiveLinkHub::FileUtilities::Private::ConfigDescription;
+		const FString Extensions = UE::LiveLinkHub::FileUtilities::Private::ConfigExtension;
+		const FString FileTypes = FString::Printf(TEXT("%s (*.%s)|*.%s"), *FileDescription, *Extensions, *Extensions);
+
+		const FString DefaultFile = UE::LiveLinkHub::FileUtilities::Private::ConfigDefaultFileName;
+
+		ClearSession();
+
+		TArray<FString> OpenFileNames;
+
+		IDesktopPlatform* DesktopPlatform = FDesktopPlatformModule::Get();
+		const void* ParentWindowWindowHandle = FSlateApplication::Get().FindBestParentWindowHandleForDialogs(nullptr);
+		const bool bFileSelected = DesktopPlatform->OpenFileDialog(
+			ParentWindowWindowHandle,
+			LOCTEXT("LiveLinkHubOpenTitle", "Open").ToString(),
+			FEditorDirectories::Get().GetLastDirectory(ELastDirectory::GENERIC_OPEN),
+			DefaultFile,
+			FileTypes,
+			EFileDialogFlags::None,
+			OpenFileNames);
+
+		if (bFileSelected && OpenFileNames.Num() > 0)
+		{
+			RestoreSession(OpenFileNames[0]);
+		}
+	}
+
+	virtual bool CanSaveCurrentSession() const override
+	{
+		return !LastConfigPath.IsEmpty();
+	}
+	//~ End LiveLinkHubSessionManager
+
+private:
+	/** Load a session from disk and restore its content. */
+	void RestoreSession(const FString& Path)
+	{
+		LastConfigPath = Path;
+		FEditorDirectories::Get().SetLastDirectory(ELastDirectory::GENERIC_OPEN, FPaths::GetPath(LastConfigPath));
+
+		const TSharedPtr<FLiveLinkHubPersistedSessionData> SessionData = UE::LiveLinkHub::FileUtilities::Private::LoadConfig(LastConfigPath);
+
+		FLiveLinkHubClient* LiveLinkHubClient = static_cast<FLiveLinkHubClient*>(&IModularFeatures::Get().GetModularFeature<ILiveLinkClient>(ILiveLinkClient::ModularFeatureName));
+
+		const FLiveLinkHubModule& LiveLinkHubModule = FModuleManager::Get().GetModuleChecked<FLiveLinkHubModule>("LiveLinkHub");
+		const TSharedPtr<FLiveLinkHubProvider> LiveLinkProvider = LiveLinkHubModule.GetLiveLinkProvider();
+
+		check(LiveLinkHubClient);
+		check(LiveLinkProvider);
+
+		if (SessionData.IsValid())
+		{
+			for (const FLiveLinkSourcePreset& SourcePreset : SessionData->Sources)
+			{
+				LiveLinkHubClient->CreateSource(SourcePreset);
+			}
+
+			for (const FLiveLinkSubjectPreset& SubjectPreset : SessionData->Subjects)
+			{
+				LiveLinkHubClient->CreateSubject(SubjectPreset);
+			}
+		}
+
+		TSharedPtr<FLiveLinkHubSession> CurrentSessionPtr;
+		{
+			FScopeLock Lock(&CurrentSessionCS);
+			CurrentSessionPtr = CurrentSession = MakeShared<FLiveLinkHubSession>(*SessionData, OnClientAddedToSessionDelegate, OnClientRemovedFromSessionDelegate);
+		}
+
+		for (FLiveLinkHubUEClientInfo& Client : SessionData->Clients)
+		{
+			CurrentSessionPtr->AddRestoredClient(Client);
+		}
+
+		OnActiveSessionChangedDelegate.Broadcast(CurrentSessionPtr.ToSharedRef());
+	}
+
+	/** Clear the hub data contained in the current session, resetting the hub to its default state. */
+	void ClearSession()
+	{
+		FLiveLinkHubClient* LiveLinkHubClient = static_cast<FLiveLinkHubClient*>(&IModularFeatures::Get().GetModularFeature<ILiveLinkClient>(ILiveLinkClient::ModularFeatureName));
+		check(LiveLinkHubClient);
+
+		LiveLinkHubClient->RemoveAllSources();
+
+		// Removing sources only sets bPendingKill to true, need to tick to make sure they are removed.
+		LiveLinkHubClient->ForceTick();
+
+		TSharedPtr<FLiveLinkHubSession> CurrentSessionPtr;
+		{
+			FScopeLock Lock(&CurrentSessionCS);
+			CurrentSessionPtr = CurrentSession = MakeShared<FLiveLinkHubSession>(OnClientAddedToSessionDelegate, OnClientRemovedFromSessionDelegate);
+		}
+
+		OnActiveSessionChangedDelegate.Broadcast(CurrentSessionPtr.ToSharedRef());
+	}
+
+private:
+	/** Session that holds the current configuration of the hub (Clients, sources, subjects). */
+	TSharedPtr<FLiveLinkHubSession> CurrentSession;
+	/** Last path where we saved a session config file. */
+	FString LastConfigPath;
+	/** Delegate triggered when a client is added to the current session. */
+	FOnClientAddedToSession OnClientAddedToSessionDelegate;
+	/** Delegate triggered when a client is removed from the current session. */
+	FOnClientRemovedFromSession OnClientRemovedFromSessionDelegate;
+	/** Delegate triggered when the current session is changed. */
+	FOnActiveSessionChanged OnActiveSessionChangedDelegate;
+	/** Critical section used to synchronize access to the current session. */
+	mutable FCriticalSection CurrentSessionCS;
+};
+
+#undef LOCTEXT_NAMESPACE /*LiveLinkHub.SessionManager*/
diff --git a/Engine/Plugins/Animation/LiveLink/Source/LiveLinkHub/Private/Subjects/LiveLinkHubSubjectController.h b/Engine/Plugins/Animation/LiveLink/Source/LiveLinkHub/Private/Subjects/LiveLinkHubSubjectController.h
--- a/Engine/Plugins/Animation/LiveLink/Source/LiveLinkHub/Private/Subjects/LiveLinkHubSubjectController.h
+++ b/Engine/Plugins/Animation/LiveLink/Source/LiveLinkHub/Private/Subjects/LiveLinkHubSubjectController.h
@@ -2,8 +2,10 @@
 
 #pragma once
 
-#include "Clients/LiveLinkHubProvider.h"
+#include "LiveLinkHubModule.h"
 #include "LiveLinkTypes.h"
+#include "Modules/ModuleManager.h"
+#include "Session/LiveLinkHubSession.h"
 #include "SLiveLinkHubSubjectView.h"
 
 /** Controller responsible for handling the hub's subjects and creating the subject view. */
@@ -13,12 +15,27 @@ public:
 	FLiveLinkHubSubjectController()
 	{
 		SubjectModel = MakeShared<FLiveLinkHubSubjectModel>();
+
+		const FLiveLinkHubModule& LiveLinkHubModule = FModuleManager::Get().GetModuleChecked<FLiveLinkHubModule>("LiveLinkHub");
+		LiveLinkHubModule.GetSessionManager()->OnActiveSessionChanged().AddRaw(this, &FLiveLinkHubSubjectController::OnActiveSessionChanged);
+	}
+
+	~FLiveLinkHubSubjectController()
+	{
+		if (const FLiveLinkHubModule* LiveLinkHubModule = FModuleManager::Get().GetModulePtr<FLiveLinkHubModule>("LiveLinkHub"))
+		{
+			if (const TSharedPtr<ILiveLinkHubSessionManager> SessionManager = LiveLinkHubModule->GetSessionManager())
+			{
+				SessionManager->OnActiveSessionChanged().RemoveAll(this);
+			}
+		}
 	}
 
 	/** Create the widget for displaying a subject's settings. */
 	TSharedRef<SWidget> MakeSubjectView()
 	{
-		return SAssignNew(SubjectsView, SLiveLinkHubSubjectView, SubjectModel.ToSharedRef());
+		return SAssignNew(SubjectsView, SLiveLinkHubSubjectView, SubjectModel.ToSharedRef())
+			.OnRenameSubject_Raw(this, &FLiveLinkHubSubjectController::OnSubjectRenamed);
 	}
 
 	/** Set the displayed subject in the subject view. */
@@ -27,6 +44,24 @@ public:
 		SubjectsView->SetSubject(Subject);
 	}
 
+	/** Handle modifying the session config for the specified subject. */
+	void OnSubjectRenamed(const FLiveLinkSubjectKey& SubjectKey, FName NewName)
+	{
+		if (TSharedPtr<ILiveLinkHubSessionManager> SessionManager = FModuleManager::Get().GetModuleChecked<FLiveLinkHubModule>("LiveLinkHub").GetLiveLinkHub()->GetSessionManager())
+		{
+			SessionManager->GetCurrentSession()->RenameSubject(SubjectKey, NewName);
+		}
+	}
+
+	/** Handle updating the subject details when the session has been swapped out for a different one. */
+	void OnActiveSessionChanged(const TSharedRef<ILiveLinkHubSession>& ActiveSession)
+	{
+		if (SubjectsView)
+		{
+			SubjectsView->ClearSubjectDetails();
+		}
+	}
+
 private:
 	/** View widget for the selected subject. */
 	TSharedPtr<SLiveLinkHubSubjectView> SubjectsView;
diff --git a/Engine/Plugins/Animation/LiveLink/Source/LiveLinkHub/Private/Subjects/LiveLinkHubSubjectModel.h b/Engine/Plugins/Animation/LiveLink/Source/LiveLinkHub/Private/Subjects/LiveLinkHubSubjectModel.h
--- a/Engine/Plugins/Animation/LiveLink/Source/LiveLinkHub/Private/Subjects/LiveLinkHubSubjectModel.h
+++ b/Engine/Plugins/Animation/LiveLink/Source/LiveLinkHub/Private/Subjects/LiveLinkHubSubjectModel.h
@@ -15,7 +15,7 @@ public:
 	virtual ~ILiveLinkHubSubjectModel() = default;
 
 	/** Get the settings for a livelinkhub subject. */
-	virtual ULiveLinkHubSubjectProxy* GetSubjectConfig(const FLiveLinkSubjectKey& InSubject) const = 0;
+	virtual TOptional<FLiveLinkHubSubjectProxy> GetSubjectConfig(const FLiveLinkSubjectKey& InSubject) const = 0;
 };
 
 /** Implementation of the ILiveLinkHubSubjectModel. */
@@ -23,13 +23,13 @@ class FLiveLinkHubSubjectModel : public ILiveLinkHubSubjectModel
 {
 public:
 	/** Get the settings for a livelinkhub subject. */
-	virtual ULiveLinkHubSubjectProxy* GetSubjectConfig(const FLiveLinkSubjectKey& InSubject) const override
+	virtual TOptional<FLiveLinkHubSubjectProxy> GetSubjectConfig(const FLiveLinkSubjectKey& InSubject) const override
 	{
 		const FLiveLinkHubModule& LiveLinkHubModule = FModuleManager::Get().GetModuleChecked<FLiveLinkHubModule>("LiveLinkHub");
 		if (const TSharedPtr<ILiveLinkHubSession> Session = LiveLinkHubModule.GetSessionManager()->GetCurrentSession())
 		{
 			return Session->GetSubjectConfig(InSubject);
 		}
-		return nullptr;
+		return {};
 	}
 };
diff --git a/Engine/Plugins/Animation/LiveLink/Source/LiveLinkHub/Private/Subjects/LiveLinkHubSubjectSessionConfig.cpp b/Engine/Plugins/Animation/LiveLink/Source/LiveLinkHub/Private/Subjects/LiveLinkHubSubjectSessionConfig.cpp
--- a/Engine/Plugins/Animation/LiveLink/Source/LiveLinkHub/Private/Subjects/LiveLinkHubSubjectSessionConfig.cpp
+++ b/Engine/Plugins/Animation/LiveLink/Source/LiveLinkHub/Private/Subjects/LiveLinkHubSubjectSessionConfig.cpp
@@ -5,98 +5,49 @@
 #include "Async/Async.h"
 #include "Clients/LiveLinkHubProvider.h"
 #include "Features/IModularFeatures.h"
+#include "LiveLinkClient.h"
 #include "LiveLinkHubModule.h"
 #include "Modules/ModuleManager.h"
-#include "LiveLinkClient.h"
 #include "UObject/StrongObjectPtr.h"
 
-ULiveLinkHubSubjectSessionConfig::ULiveLinkHubSubjectSessionConfig()
+
+void FLiveLinkHubSubjectSessionConfig::Initialize()
 {
-	if (!IsTemplate())
-	{
-		FLiveLinkClient& LiveLinkClient = IModularFeatures::Get().GetModularFeature<FLiveLinkClient>(ILiveLinkClient::ModularFeatureName);
-		LiveLinkClient.OnLiveLinkSubjectAdded().AddUObject(this, &ULiveLinkHubSubjectSessionConfig::OnSubjectAdded_AnyThread);
-		LiveLinkClient.OnLiveLinkSubjectRemoved().AddUObject(this, &ULiveLinkHubSubjectSessionConfig::OnSubjectRemoved_AnyThread);
+	FLiveLinkClient& LiveLinkClient = IModularFeatures::Get().GetModularFeature<FLiveLinkClient>(ILiveLinkClient::ModularFeatureName);
 
-		constexpr bool bIncludeDisabledSubject = true;
-		constexpr bool bIncludeVirtualSubject = true;
-		for (const FLiveLinkSubjectKey& SubjectKey : LiveLinkClient.GetSubjects(bIncludeDisabledSubject, bIncludeVirtualSubject))
-		{
-			ULiveLinkHubSubjectProxy* SubjectSettings = NewObject<ULiveLinkHubSubjectProxy>();
-			SubjectSettings->Initialize(SubjectKey, LiveLinkClient.GetSourceType(SubjectKey.Source).ToString());
+	constexpr bool bIncludeDisabledSubject = true;
+	constexpr bool bIncludeVirtualSubject = true;
 
-			SubjectProxies.Add(SubjectKey, SubjectSettings);
-		}
-	}
-}
+	for (const FLiveLinkSubjectKey& SubjectKey : LiveLinkClient.GetSubjects(bIncludeDisabledSubject, bIncludeVirtualSubject))
+	{
+		FLiveLinkHubSubjectProxy SubjectSettings;
+		SubjectSettings.Initialize(SubjectKey, LiveLinkClient.GetSourceType(SubjectKey.Source).ToString());
 
-ULiveLinkHubSubjectSessionConfig::~ULiveLinkHubSubjectSessionConfig()
-{
-	if (IModularFeatures::Get().IsModularFeatureAvailable(ILiveLinkClient::ModularFeatureName))
-    {
-    	FLiveLinkClient& LiveLinkClient = IModularFeatures::Get().GetModularFeature<FLiveLinkClient>(ILiveLinkClient::ModularFeatureName);
-    	LiveLinkClient.OnLiveLinkSubjectRemoved().RemoveAll(this);
-    	LiveLinkClient.OnLiveLinkSubjectAdded().RemoveAll(this);
-    }
+		SubjectProxies.Add(SubjectKey, MoveTemp(SubjectSettings));
+	}
 }
 
-ULiveLinkHubSubjectProxy* ULiveLinkHubSubjectSessionConfig::GetSubjectConfig(const FLiveLinkSubjectKey& InSubject)
+TOptional<FLiveLinkHubSubjectProxy> FLiveLinkHubSubjectSessionConfig::GetSubjectConfig(const FLiveLinkSubjectKey& InSubject) const
 {
-    ULiveLinkHubSubjectProxy* Settings = nullptr;
+	TOptional<FLiveLinkHubSubjectProxy> Settings;
 
-    if (TObjectPtr<ULiveLinkHubSubjectProxy>* SettingsPtr = SubjectProxies.Find(InSubject))
+    if (const FLiveLinkHubSubjectProxy* SettingsPtr = SubjectProxies.Find(InSubject))
     {
-    	Settings = SettingsPtr->Get();
+    	Settings = *SettingsPtr;
     }
 
     return Settings;
 }
 
-void ULiveLinkHubSubjectSessionConfig::OnSubjectAdded_AnyThread(FLiveLinkSubjectKey SubjectKey)
+void FLiveLinkHubSubjectSessionConfig::RenameSubject(const FLiveLinkSubjectKey& SubjectKey, FName NewName)
 {
-	TWeakObjectPtr<ULiveLinkHubSubjectSessionConfig> Self = this;
-	AsyncTask(ENamedThreads::GameThread, [Self, Key = MoveTemp(SubjectKey)]
+	if (FLiveLinkHubSubjectProxy* Proxy = SubjectProxies.Find(SubjectKey))
 	{
-		ULiveLinkHubSubjectSessionConfig* Container = Self.Get();
-		if (UObjectInitialized() && Container)
-		{
-			Container->OnSubjectAdded(Key);
-		}
-	});
-}
-
-void ULiveLinkHubSubjectSessionConfig::OnSubjectAdded(const FLiveLinkSubjectKey& SubjectKey)
-{
-	if (!SubjectProxies.Contains(SubjectKey))
-	{
-		ILiveLinkClient& LiveLinkClient = IModularFeatures::Get().GetModularFeature<ILiveLinkClient>(ILiveLinkClient::ModularFeatureName);
-
-		ULiveLinkHubSubjectProxy* SubjectSettings = NewObject<ULiveLinkHubSubjectProxy>();
-		SubjectSettings->Initialize(SubjectKey, LiveLinkClient.GetSourceType(SubjectKey.Source).ToString());
-
-		SubjectProxies.FindOrAdd(SubjectKey) = SubjectSettings;
+		Proxy->SetOutboundName(NewName);
 	}
 }
 
-void ULiveLinkHubSubjectSessionConfig::OnSubjectRemoved_AnyThread(FLiveLinkSubjectKey SubjectKey)
-{
-	TWeakObjectPtr<ULiveLinkHubSubjectSessionConfig> Self = this;
-	AsyncTask(ENamedThreads::GameThread, [Self, Key = MoveTemp(SubjectKey)]
-	{
-		ULiveLinkHubSubjectSessionConfig* Container = Self.Get();
-		if (UObjectInitialized() && Container)
-		{
-			Container->OnSubjectAdded(Key);
-		}
-	});
-}
-
-void ULiveLinkHubSubjectSessionConfig::OnSubjectRemoved(const FLiveLinkSubjectKey& SubjectKey)
-{
-	SubjectProxies.Remove(SubjectKey);
-}
-
-void ULiveLinkHubSubjectProxy::Initialize(const FLiveLinkSubjectKey& InSubjectKey, FString InSource)
+void FLiveLinkHubSubjectProxy::Initialize(const FLiveLinkSubjectKey& InSubjectKey, FString InSource)
 {
 	SubjectName = InSubjectKey.SubjectName.Name.ToString();
 	SubjectKey = InSubjectKey;
@@ -104,37 +55,25 @@ void ULiveLinkHubSubjectProxy::Initialize(const FLiveLinkSubjectKey& InSubjectKe
 	Source = MoveTemp(InSource);
 }
 
-void ULiveLinkHubSubjectProxy::PreEditChange(FProperty* PropertyAboutToChange)
+FName FLiveLinkHubSubjectProxy::GetOutboundName() const
 {
-	Super::PreEditChange(PropertyAboutToChange);
-
-	if (PropertyAboutToChange && PropertyAboutToChange->GetFName() == GET_MEMBER_NAME_CHECKED(ULiveLinkHubSubjectProxy, OutboundName))
+	if (bPendingOutboundNameChange)
 	{
-		PreviousOutboundName = *OutboundName;
+		return PreviousOutboundName;
 	}
-}
-
-void ULiveLinkHubSubjectProxy::PostEditChangeProperty(struct FPropertyChangedEvent& PropertyChangedEvent)
-{
-	Super::PostEditChangeProperty(PropertyChangedEvent);
 
-	if (PropertyChangedEvent.Property && PropertyChangedEvent.Property->GetName() == GET_MEMBER_NAME_CHECKED(ULiveLinkHubSubjectProxy, OutboundName))
-	{
-		NotifyRename();
-	}
+	return *OutboundName;
 }
 
-FName ULiveLinkHubSubjectProxy::GetOutboundName() const
+void FLiveLinkHubSubjectProxy::SetOutboundName(FName NewName)
 {
-	if (bPendingOutboundNameChange)
-	{
-		return PreviousOutboundName;
-	}
+	PreviousOutboundName = *OutboundName;
+	OutboundName = *NewName.ToString();
 
-	return *OutboundName;
+	NotifyRename();
 }
 
-void ULiveLinkHubSubjectProxy::NotifyRename()
+void FLiveLinkHubSubjectProxy::NotifyRename()
 {
 	FLiveLinkHubModule& LiveLinkHubModule = FModuleManager::Get().GetModuleChecked<FLiveLinkHubModule>("LiveLinkHub");
 
diff --git a/Engine/Plugins/Animation/LiveLink/Source/LiveLinkHub/Private/Subjects/LiveLinkHubSubjectSessionConfig.h b/Engine/Plugins/Animation/LiveLink/Source/LiveLinkHub/Private/Subjects/LiveLinkHubSubjectSessionConfig.h
--- a/Engine/Plugins/Animation/LiveLink/Source/LiveLinkHub/Private/Subjects/LiveLinkHubSubjectSessionConfig.h
+++ b/Engine/Plugins/Animation/LiveLink/Source/LiveLinkHub/Private/Subjects/LiveLinkHubSubjectSessionConfig.h
@@ -9,39 +9,9 @@
 
 #include "LiveLinkHubSubjectSessionConfig.generated.h"
 
-/** Config pertaining to livelink hub subjects for a given session. */
-UCLASS()
-class ULiveLinkHubSubjectSessionConfig : public UObject
-{
-public:
-	GENERATED_BODY()
-
-	ULiveLinkHubSubjectSessionConfig();
-	virtual ~ULiveLinkHubSubjectSessionConfig() override;
-
-	/** Get the config for a livelinkhub subject. */
-	ULiveLinkHubSubjectProxy* GetSubjectConfig(const FLiveLinkSubjectKey& InSubject);
-
-private:
-	/** AnyThread handler for the SubjectAdded delegate, dispatches handling on the game thread to avoid asserts in Slate. */
-	void OnSubjectAdded_AnyThread(FLiveLinkSubjectKey SubjectKey);
-	/** Handles updating the tree view when a subject is added. */
-	void OnSubjectAdded(const FLiveLinkSubjectKey& SubjectKey);
-	/** AnyThread handler for the SubjectRemoved delegate, dispatches handling on the game thread to avoid asserts in Slate. */
-	void OnSubjectRemoved_AnyThread(FLiveLinkSubjectKey SubjectKey);
-	/** Handles updating the tree view when a subject is removed. */
-	void OnSubjectRemoved(const FLiveLinkSubjectKey& SubjectKey);
-
-private:
-	/** Settings for subjects displayed in the livelink hub. */
-	UPROPERTY(Instanced)
-	TMap<FLiveLinkSubjectKey, TObjectPtr<ULiveLinkHubSubjectProxy>> SubjectProxies;
-};
-
-
 /** Holds information and config for a subject for the duration of a livelink hub session.  */
-UCLASS()
-class ULiveLinkHubSubjectProxy : public UObject
+USTRUCT()
+struct FLiveLinkHubSubjectProxy
 {
 public:
 	GENERATED_BODY()
@@ -52,10 +22,15 @@ public:
 	/** Get the outbound name for this subject, allows  */
 	FName GetOutboundName() const;
 
-	//~ Begin UObject interface
-	virtual void PreEditChange(FProperty* PropertyAboutToChange) override;
-	virtual void PostEditChangeProperty(struct FPropertyChangedEvent& PropertyChangedEvent) override;
-	//~ End UObject interface
+	/** Change the outbound name for this subject proxy. */
+	void SetOutboundName(FName NewName);
+
+	/** Get the outbound name property name. */
+	static FName GetOutboundNamePropertyName()
+	{
+		return GET_MEMBER_NAME_CHECKED(FLiveLinkHubSubjectProxy, OutboundName);
+	}
+
 private:
 	/** Notify clients that the old subject should be deleted and replaced with a new one with the updated name. */
 	void NotifyRename();
@@ -86,3 +61,26 @@ private:
 	/* Previous outbound name to be used for noticing clients to remove this entry from their subject list. */
 	FName PreviousOutboundName;
 };
+
+/** Config pertaining to livelink hub subjects for a given session. */
+USTRUCT()
+struct FLiveLinkHubSubjectSessionConfig
+{
+public:
+	GENERATED_BODY()
+
+	void Initialize();
+
+	/** Get the config for a livelinkhub subject. */
+	TOptional<FLiveLinkHubSubjectProxy> GetSubjectConfig(const FLiveLinkSubjectKey& InSubject) const;
+
+	/** Change the outbound name of a subject for the current session. */
+	void RenameSubject(const FLiveLinkSubjectKey& SubjectKey, FName NewName);
+
+private:
+	/** Settings for subjects displayed in the livelink hub. */
+	UPROPERTY()
+	TMap<FLiveLinkSubjectKey, FLiveLinkHubSubjectProxy> SubjectProxies;
+
+	friend class FLiveLinkHubSession;
+};
diff --git a/Engine/Plugins/Animation/LiveLink/Source/LiveLinkHub/Private/Subjects/SLiveLinkHubSubjectView.h b/Engine/Plugins/Animation/LiveLink/Source/LiveLinkHub/Private/Subjects/SLiveLinkHubSubjectView.h
--- a/Engine/Plugins/Animation/LiveLink/Source/LiveLinkHub/Private/Subjects/SLiveLinkHubSubjectView.h
+++ b/Engine/Plugins/Animation/LiveLink/Source/LiveLinkHub/Private/Subjects/SLiveLinkHubSubjectView.h
@@ -4,25 +4,32 @@
 
 #include "Widgets/SCompoundWidget.h"
 
+#include "Delegates/DelegateCombinations.h"
+#include "DetailsViewArgs.h"
 #include "IDetailsView.h"
 #include "LiveLinkHubSubjectModel.h"
 #include "LiveLinkTypes.h"
 #include "PropertyEditorModule.h"
 #include "Modules/ModuleManager.h"
 
+DECLARE_DELEGATE_TwoParams(FOnRenameLiveLinkHubSubject, const FLiveLinkSubjectKey& /*SubjectKey*/, FName /*NewName*/);
+
 /**
  * Provides the UI that displays information about a livelink hub subject.
  */
 class SLiveLinkHubSubjectView : public SCompoundWidget
 {
 public:
+
 	SLATE_BEGIN_ARGS(SLiveLinkHubSubjectView) {}
+	SLATE_EVENT(FOnRenameLiveLinkHubSubject, OnRenameSubject)
 	SLATE_END_ARGS()
 
 	//~ Begin SWidget interface
 	void Construct(const FArguments& InArgs, const TSharedRef<ILiveLinkHubSubjectModel>& InSubjectModel)
 	{
 		SubjectModel = InSubjectModel;
+		OnRenameSubjectDelegate = InArgs._OnRenameSubject;
 
 		FDetailsViewArgs DetailsViewArgs;
 		DetailsViewArgs.bUpdatesFromSelection = false;
@@ -34,37 +41,70 @@ public:
 		DetailsViewArgs.bShowOptions = false;
 
 		FPropertyEditorModule& PropertyEditorModule = FModuleManager::GetModuleChecked<FPropertyEditorModule>("PropertyEditor");
+		FStructureDetailsViewArgs StructureDetailsArgs;
 
-		SettingsDetailsView = PropertyEditorModule.CreateDetailView(DetailsViewArgs);
+		SubjectData = MakeShared<TStructOnScope<FLiveLinkHubSubjectProxy>>();
+
+		SettingsDetailsView = PropertyEditorModule.CreateStructureDetailView(DetailsViewArgs, StructureDetailsArgs, SubjectData);
+		SettingsDetailsView->GetOnFinishedChangingPropertiesDelegate().AddSP(this, &SLiveLinkHubSubjectView::OnSubjectPropertyModified);
 
 		ChildSlot
 		[
-			SettingsDetailsView.ToSharedRef()
+			SettingsDetailsView->GetWidget().ToSharedRef()
 		];
 	}
 
+	/** Clear the subject details. */
+	void ClearSubjectDetails()
+	{
+		SubjectData->Reset();
+		SettingsDetailsView->SetStructureData(nullptr);
+	}
+
 	/** Set the subject to be displayed in the details view. */
 	void SetSubject(const FLiveLinkSubjectKey& InSubjectKey)
 	{
 		SubjectKey = InSubjectKey;
-		SettingsDetailsView->SetObject(SubjectModel->GetSubjectConfig(SubjectKey));
+
+		if (TOptional<FLiveLinkHubSubjectProxy> SubjectProxy = SubjectModel->GetSubjectConfig(InSubjectKey))
+		{
+			const FLiveLinkHubSubjectProxy& Proxy = *SubjectProxy;
+			SubjectData->InitializeAs<FLiveLinkHubSubjectProxy>(*SubjectProxy);
+			SettingsDetailsView->SetStructureData(SubjectData);
+		}
 	}
 
 	/** Get the outbound name for a subject. */
-	FString GetSubjectOutboundName(const FLiveLinkSubjectKey& InSubjectKey)
+	FString GetSubjectOutboundName(const FLiveLinkSubjectKey& InSubjectKey) const
 	{
-		if (ULiveLinkHubSubjectProxy* SubjectProxy = SubjectModel->GetSubjectConfig(InSubjectKey))
+		if (TOptional<FLiveLinkHubSubjectProxy> SubjectProxy = SubjectModel->GetSubjectConfig(InSubjectKey))
 		{
 			return SubjectProxy->GetOutboundName().ToString();
 		}
 		return InSubjectKey.SubjectName.ToString();
 	}
 
+	/** Handler called when a subject property is modified, used to trigger a rename on the session's subject config. */
+	void OnSubjectPropertyModified(const FPropertyChangedEvent& PropertyChangedEvent)
+	{
+		if (SubjectData && PropertyChangedEvent.GetPropertyName() == FLiveLinkHubSubjectProxy::GetOutboundNamePropertyName())
+		{
+			if (FLiveLinkHubSubjectProxy* Proxy = SubjectData->Get())
+			{
+				OnRenameSubjectDelegate.ExecuteIfBound(SubjectKey, Proxy->GetOutboundName());
+			}
+		}
+	}
+
 private:
 	/** Details for the selected subject. */
-	TSharedPtr<IDetailsView> SettingsDetailsView;
+	TSharedPtr<IStructureDetailsView> SettingsDetailsView;
 	/** Subject being shown. */
 	FLiveLinkSubjectKey SubjectKey;
 	/** Model that holds the data for a given subject. */
 	TSharedPtr<ILiveLinkHubSubjectModel> SubjectModel;
+	/** Struct on scope used to display subject data in the structure details view. */
+	TSharedPtr<TStructOnScope<FLiveLinkHubSubjectProxy>> SubjectData;
+	/** Delegate called when the outbound name is changed by the user. */
+	FOnRenameLiveLinkHubSubject OnRenameSubjectDelegate;
 };
