Add LiveLinkHub entry to the livelink panel
Don't show livelink hub entry in the hub

diff --git a/Engine/Plugins/Animation/LiveLink/LiveLink.uplugin b/Engine/Plugins/Animation/LiveLink/LiveLink.uplugin
--- a/Engine/Plugins/Animation/LiveLink/LiveLink.uplugin
+++ b/Engine/Plugins/Animation/LiveLink/LiveLink.uplugin
@@ -82,7 +82,8 @@
 		{
 			"Name": "LiveLinkHubMessaging",
 			"Type": "Runtime",
-			"LoadingPhase": "Default"
+			"LoadingPhase": "Default",
+			"ProgramAllowList": [ "LiveLinkHub" ]
 		}
 	]
 }
\ No newline at end of file
diff --git a/Engine/Plugins/Animation/LiveLink/Source/LiveLink/Internal/ILiveLinkModule.h b/Engine/Plugins/Animation/LiveLink/Source/LiveLink/Internal/ILiveLinkModule.h
--- a/Engine/Plugins/Animation/LiveLink/Source/LiveLink/Internal/ILiveLinkModule.h
+++ b/Engine/Plugins/Animation/LiveLink/Source/LiveLink/Internal/ILiveLinkModule.h
@@ -3,15 +3,23 @@
 #pragma once
 
 #include "Modules/ModuleManager.h"
+#include "Templates/SharedPointer.h"
 
 #ifndef WITH_LIVELINK_DISCOVERY_MANAGER_THREAD
 #define WITH_LIVELINK_DISCOVERY_MANAGER_THREAD 1
 #endif
 
+class FDelegateHandle;
 class FLiveLinkHeartbeatEmitter;
 class FLiveLinkMessageBusDiscoveryManager;
+struct FProviderPollResult;
 class FSlateStyleSet;
 
+using FProviderPollResultPtr = TSharedPtr<FProviderPollResult, ESPMode::ThreadSafe>;
+
+
+DECLARE_DELEGATE_RetVal_TwoParams(bool, FOnLiveLinkShouldDisplaySource, UClass*/*FactoryClass*/, FProviderPollResultPtr/*PollResultPtr*/);
+
 /**
  * Interface for messaging modules.
  */
@@ -41,9 +49,27 @@ public:
 	virtual FLiveLinkMessageBusDiscoveryManager& GetMessageBusDiscoveryManager() = 0;
 #endif
 
+	/**
+     * Register a message bus source filter invoked by the MessageBus Source Factory to filter out discovered sources.
+     * The delegate should return true if the poll result should be included in the discovered sources list, and false otherwise.
+     * @param SourceFilter The filter invoked by the source factory.
+     * @return A handle to the delegate, used for unregistering the filter.
+     */
+	virtual FDelegateHandle RegisterMessageBusSourceFilter(const FOnLiveLinkShouldDisplaySource& SourceFilter) = 0;
+
+	/**
+	 * Unregister a previously registered message bus source filter.
+	 * @param Handle A handle to the filter delegate that was registered.
+	 */
+	virtual void UnregisterMessageBusSourceFilter(FDelegateHandle Handle) = 0;
+
+	/**
+	 * @return The list of source filters registered with this module.
+	 */
+	virtual const TMap<FDelegateHandle, FOnLiveLinkShouldDisplaySource>& GetSourceFilters() const = 0;
+
 public:
 
 	/** Virtual destructor. */
 	virtual ~ILiveLinkModule() { }
 };
-
diff --git a/Engine/Plugins/Animation/LiveLink/Source/LiveLink/Internal/SLiveLinkMessageBusSourceFactory.h b/Engine/Plugins/Animation/LiveLink/Source/LiveLink/Internal/SLiveLinkMessageBusSourceFactory.h
new file mode 100644
--- /dev/null
+++ b/Engine/Plugins/Animation/LiveLink/Source/LiveLink/Internal/SLiveLinkMessageBusSourceFactory.h
@@ -0,0 +1,82 @@
+// Copyright Epic Games, Inc. All Rights Reserved.
+
+#pragma once
+
+#include "Widgets/SCompoundWidget.h"
+#include "LiveLinkMessageBusFinder.h"
+#include "Misc/App.h"
+
+namespace ESelectInfo { enum Type : int; }
+template <typename ItemType> class SListView;
+
+struct FLiveLinkPongMessage;
+struct FMessageAddress;
+struct FProviderPollResult;
+class ILiveLinkClient;
+class ITableRow;
+class STableViewBase;
+
+
+DECLARE_DELEGATE_OneParam(FOnLiveLinkMessageBusSourceSelected, FProviderPollResultPtr);
+
+class LIVELINK_API SLiveLinkMessageBusSourceFactory : public SCompoundWidget
+{
+	SLATE_BEGIN_ARGS(SLiveLinkMessageBusSourceFactory) {}
+		SLATE_EVENT(FOnLiveLinkMessageBusSourceSelected, OnSourceSelected)
+		SLATE_ARGUMENT(UClass*, FactoryClass)
+	SLATE_END_ARGS()
+
+	~SLiveLinkMessageBusSourceFactory();
+
+	void Construct(const FArguments& Args);
+
+	FProviderPollResultPtr GetSelectedSource() const { return SelectedResult; }
+
+	virtual void Tick(const FGeometry& AllottedGeometry, const double InCurrentTime, const float InDeltaTime);
+
+private:
+	/** Represents a source with its latest time since it was polled successfully. */
+	struct FLiveLinkSource
+	{
+		FProviderPollResultPtr PollResult;
+		double LastTimeSincePong = 0.0;
+		
+		FLiveLinkSource(FProviderPollResultPtr InPollResult)
+            : PollResult(MoveTemp(InPollResult))
+			, LastTimeSincePong(FApp::GetCurrentTime()) 
+		{
+		}
+	};
+
+private:
+	/** Create a widget entry in the source list based on a live link source. */
+	TSharedRef<ITableRow> MakeSourceListViewWidget(TSharedPtr<FLiveLinkSource> Source, const TSharedRef<STableViewBase>& OwnerTable) const;
+
+	/** Handles selection changing in the list. */
+	void OnSourceListSelectionChanged(TSharedPtr<FLiveLinkSource> Source, ESelectInfo::Type SelectionType);
+
+private:
+	/** The widgets displayed in the list. */
+	TSharedPtr<SListView<TSharedPtr<FLiveLinkSource>>> ListView;
+
+	/** Latest sources pulled from the discovery manager. */
+	TArray<FProviderPollResultPtr> PollData;
+
+	/** The sources displayed by the list. */
+	TArray<TSharedPtr<FLiveLinkSource>> Sources;
+
+	/** The source that was selected. */
+	FProviderPollResultPtr SelectedResult;
+
+	/** The last source that was selected. */
+	FOnLiveLinkMessageBusSourceSelected OnSourceSelected;
+
+	/** The last time at which the UI was refreshed. */
+	double LastUIUpdateSeconds = 0;
+
+	/** Time before a source should disappear from the UI after its last pong. */
+	double SecondsBeforeSourcesDisappear = 2.0;
+
+	/** Class of the factory that created this widget. */
+	UClass* FactoryClass = nullptr;
+};
diff --git a/Engine/Plugins/Animation/LiveLink/Source/LiveLink/Private/LiveLinkMessageBusSourceFactory.cpp b/Engine/Plugins/Animation/LiveLink/Source/LiveLink/Private/LiveLinkMessageBusSourceFactory.cpp
--- a/Engine/Plugins/Animation/LiveLink/Source/LiveLink/Private/LiveLinkMessageBusSourceFactory.cpp
+++ b/Engine/Plugins/Animation/LiveLink/Source/LiveLink/Private/LiveLinkMessageBusSourceFactory.cpp
@@ -26,7 +26,8 @@ FText ULiveLinkMessageBusSourceFactory::GetSourceTooltip() const
 TSharedPtr<SWidget> ULiveLinkMessageBusSourceFactory::BuildCreationPanel(FOnLiveLinkSourceCreated InOnLiveLinkSourceCreated) const
 {
 	return SNew(SLiveLinkMessageBusSourceFactory)
-		.OnSourceSelected(FOnLiveLinkMessageBusSourceSelected::CreateUObject(this, &ULiveLinkMessageBusSourceFactory::OnSourceSelected, InOnLiveLinkSourceCreated));
+		.OnSourceSelected(FOnLiveLinkMessageBusSourceSelected::CreateUObject(this, &ULiveLinkMessageBusSourceFactory::OnSourceSelected, InOnLiveLinkSourceCreated))
+		.FactoryClass(GetClass());
 }
 
 TSharedPtr<FLiveLinkMessageBusSource> ULiveLinkMessageBusSourceFactory::MakeSource(const FText& Name,
diff --git a/Engine/Plugins/Animation/LiveLink/Source/LiveLink/Private/LiveLinkModule.cpp b/Engine/Plugins/Animation/LiveLink/Source/LiveLink/Private/LiveLinkModule.cpp
--- a/Engine/Plugins/Animation/LiveLink/Source/LiveLink/Private/LiveLinkModule.cpp
+++ b/Engine/Plugins/Animation/LiveLink/Source/LiveLink/Private/LiveLinkModule.cpp
@@ -74,6 +74,18 @@ void FLiveLinkModule::ShutdownModule()
 	FLiveLinkLogInstance::DestroyInstance();
 }
 
+FDelegateHandle FLiveLinkModule::RegisterMessageBusSourceFilter(const FOnLiveLinkShouldDisplaySource& Delegate)
+{
+	FDelegateHandle Handle = Delegate.GetHandle();
+	RegisteredSourceFilters.FindOrAdd(Handle) = Delegate;
+	return Handle;
+}
+
+void FLiveLinkModule::UnregisterMessageBusSourceFilter(FDelegateHandle Handle)
+{
+	RegisteredSourceFilters.Remove(Handle);
+}
+
 void FLiveLinkModule::CreateStyle()
 {
 	static FName LiveLinkStyle(TEXT("LiveLinkCoreStyle"));
diff --git a/Engine/Plugins/Animation/LiveLink/Source/LiveLink/Private/LiveLinkModule.h b/Engine/Plugins/Animation/LiveLink/Source/LiveLink/Private/LiveLinkModule.h
--- a/Engine/Plugins/Animation/LiveLink/Source/LiveLink/Private/LiveLinkModule.h
+++ b/Engine/Plugins/Animation/LiveLink/Source/LiveLink/Private/LiveLinkModule.h
@@ -33,6 +33,15 @@ public:
 	virtual FLiveLinkMessageBusDiscoveryManager& GetMessageBusDiscoveryManager() override { return *DiscoveryManager; }
 #endif
 	virtual TSharedPtr<FSlateStyleSet> GetStyle() override { return StyleSet; }
+
+	virtual FDelegateHandle RegisterMessageBusSourceFilter(const FOnLiveLinkShouldDisplaySource& SourceFilter) override;
+	virtual void UnregisterMessageBusSourceFilter(FDelegateHandle Handle) override;
+
+	virtual const TMap<FDelegateHandle, FOnLiveLinkShouldDisplaySource>& GetSourceFilters() const override
+	{
+		return RegisteredSourceFilters;
+	}
+
 	//~ End ILiveLinkModule interface
 
 private:
@@ -51,4 +60,7 @@ private:
 	TUniquePtr<FLiveLinkHeartbeatEmitter> HeartbeatEmitter;
 	TUniquePtr<FLiveLinkMessageBusDiscoveryManager> DiscoveryManager;
 	TUniquePtr<FLiveLinkDebugCommand> LiveLinkDebugCommand;
+
+	// Filters registered externally used to filter what message bus sources should be displayed in the UI.
+	TMap<FDelegateHandle, FOnLiveLinkShouldDisplaySource> RegisteredSourceFilters;
 };
diff --git a/Engine/Plugins/Animation/LiveLink/Source/LiveLink/Private/SLiveLinkMessageBusSourceFactory.cpp b/Engine/Plugins/Animation/LiveLink/Source/LiveLink/Private/SLiveLinkMessageBusSourceFactory.cpp
--- a/Engine/Plugins/Animation/LiveLink/Source/LiveLink/Private/SLiveLinkMessageBusSourceFactory.cpp
+++ b/Engine/Plugins/Animation/LiveLink/Source/LiveLink/Private/SLiveLinkMessageBusSourceFactory.cpp
@@ -88,6 +88,7 @@ SLiveLinkMessageBusSourceFactory::~SLiveLinkMessageBusSourceFactory()
 void SLiveLinkMessageBusSourceFactory::Construct(const FArguments& Args)
 {
 	OnSourceSelected = Args._OnSourceSelected;
+	FactoryClass = Args._FactoryClass;
 
 #if WITH_LIVELINK_DISCOVERY_MANAGER_THREAD
 	ILiveLinkModule::Get().GetMessageBusDiscoveryManager().AddDiscoveryMessageRequest();
@@ -136,6 +137,9 @@ void SLiveLinkMessageBusSourceFactory::Tick(const FGeometry& AllottedGeometry, c
 		{
 			return FApp::GetCurrentTime() - Source->LastTimeSincePong > SecondsBeforeSourcesDisappear;
 		});
+
+
+		const TMap<FDelegateHandle, FOnLiveLinkShouldDisplaySource> RegisteredFilters = ILiveLinkModule::Get().GetSourceFilters();
 		
 		for (const FProviderPollResultPtr& PollResult : PollData)
 		{
@@ -146,7 +150,21 @@ void SLiveLinkMessageBusSourceFactory::Tick(const FGeometry& AllottedGeometry, c
 			}
 			else
 			{
-				Sources.Add(MakeShared<FLiveLinkSource>(PollResult));
+				bool bAllowResult = true;
+
+				for (const TPair<FDelegateHandle, FOnLiveLinkShouldDisplaySource>& Filter : RegisteredFilters)
+				{
+					if (!Filter.Value.Execute(FactoryClass, PollResult))
+					{
+						bAllowResult = false;
+						break;
+					}
+				}
+
+				if (bAllowResult)
+				{
+					Sources.Add(MakeShared<FLiveLinkSource>(PollResult));
+				}
 			}
 		}
 
diff --git a/Engine/Plugins/Animation/LiveLink/Source/LiveLink/Private/SLiveLinkMessageBusSourceFactory.h b/Engine/Plugins/Animation/LiveLink/Source/LiveLink/Private/SLiveLinkMessageBusSourceFactory.h
deleted file mode 100644
diff --git a/Engine/Plugins/Animation/LiveLink/Source/LiveLink/Public/LiveLinkMessageBusSourceFactory.h b/Engine/Plugins/Animation/LiveLink/Source/LiveLink/Public/LiveLinkMessageBusSourceFactory.h
--- a/Engine/Plugins/Animation/LiveLink/Source/LiveLink/Public/LiveLinkMessageBusSourceFactory.h
+++ b/Engine/Plugins/Animation/LiveLink/Source/LiveLink/Public/LiveLinkMessageBusSourceFactory.h
@@ -15,8 +15,8 @@ class LIVELINK_API ULiveLinkMessageBusSourceFactory : public ULiveLinkSourceFact
 public:
 	GENERATED_BODY()
 
-	virtual FText GetSourceDisplayName() const;
-	virtual FText GetSourceTooltip() const;
+	virtual FText GetSourceDisplayName() const override;
+	virtual FText GetSourceTooltip() const override;
 
 	virtual EMenuType GetMenuType() const override { return EMenuType::SubPanel; }
 	virtual TSharedPtr<SWidget> BuildCreationPanel(FOnLiveLinkSourceCreated OnLiveLinkSourceCreated) const override;
@@ -31,6 +31,5 @@ protected:
 																   const FMessageAddress& Address,
 																   double TimeOffset) const;
 
-private:
 	void OnSourceSelected(FProviderPollResultPtr SelectedSource, FOnLiveLinkSourceCreated OnLiveLinkSourceCreated) const;
 };
diff --git a/Engine/Plugins/Animation/LiveLink/Source/LiveLinkEditor/Private/LiveLinkClientPanelToolbar.cpp b/Engine/Plugins/Animation/LiveLink/Source/LiveLinkEditor/Private/LiveLinkClientPanelToolbar.cpp
--- a/Engine/Plugins/Animation/LiveLink/Source/LiveLinkEditor/Private/LiveLinkClientPanelToolbar.cpp
+++ b/Engine/Plugins/Animation/LiveLink/Source/LiveLinkEditor/Private/LiveLinkClientPanelToolbar.cpp
@@ -414,7 +414,7 @@ void SLiveLinkClientPanelToolbar::Construct(const FArguments& Args, FLiveLinkCli
 	GetDerivedClasses(ULiveLinkSourceFactory::StaticClass(), Results, true);
 	for (UClass* SourceFactory : Results)
 	{
-		if (!SourceFactory->HasAllClassFlags(CLASS_Abstract | CLASS_Deprecated | CLASS_NewerVersionExists))
+		if (!SourceFactory->HasAllClassFlags(CLASS_Abstract | CLASS_Deprecated | CLASS_NewerVersionExists) && CastChecked<ULiveLinkSourceFactory>(SourceFactory->GetDefaultObject())->IsEnabled())
 		{
 			Factories.Add(NewObject<ULiveLinkSourceFactory>(GetTransientPackage(), SourceFactory));
 		}
diff --git a/Engine/Plugins/Animation/LiveLink/Source/LiveLinkHub/Private/Session/LiveLinkHubSession.h b/Engine/Plugins/Animation/LiveLink/Source/LiveLinkHub/Private/Session/LiveLinkHubSession.h
--- a/Engine/Plugins/Animation/LiveLink/Source/LiveLinkHub/Private/Session/LiveLinkHubSession.h
+++ b/Engine/Plugins/Animation/LiveLink/Source/LiveLinkHub/Private/Session/LiveLinkHubSession.h
@@ -81,7 +81,7 @@ public:
 
 	virtual TOptional<FLiveLinkHubSubjectProxy> GetSubjectConfig(const FLiveLinkSubjectKey& SubjectKey) const override
 	{
-		FWriteScopeLock Locker(SessionDataLock);
+		FReadScopeLock Locker(SessionDataLock);
 		return SessionData.SubjectsConfig.GetSubjectConfig(SubjectKey);
 	}
 
diff --git a/Engine/Plugins/Animation/LiveLink/Source/LiveLinkHubMessaging/LiveLinkHubMessaging.build.cs b/Engine/Plugins/Animation/LiveLink/Source/LiveLinkHubMessaging/LiveLinkHubMessaging.build.cs
--- a/Engine/Plugins/Animation/LiveLink/Source/LiveLinkHubMessaging/LiveLinkHubMessaging.build.cs
+++ b/Engine/Plugins/Animation/LiveLink/Source/LiveLinkHubMessaging/LiveLinkHubMessaging.build.cs
@@ -27,6 +27,8 @@ namespace UnrealBuildTool.Rules
 				PrivateDependencyModuleNames.AddRange(
 				new string[]
 				{
+					"Slate",
+					"SlateCore",
 					"UnrealEd"
 				});
 			}
diff --git a/Engine/Plugins/Animation/LiveLink/Source/LiveLinkHubMessaging/Private/LiveLinkHubMessageBusSourceFactory.cpp b/Engine/Plugins/Animation/LiveLink/Source/LiveLinkHubMessaging/Private/LiveLinkHubMessageBusSourceFactory.cpp
new file mode 100644
--- /dev/null
+++ b/Engine/Plugins/Animation/LiveLink/Source/LiveLinkHubMessaging/Private/LiveLinkHubMessageBusSourceFactory.cpp
@@ -0,0 +1,50 @@
+// Copyright Epic Games, Inc. All Rights Reserved.
+
+#include "LiveLinkHubMessageBusSourceFactory.h"
+
+#include "LiveLinkHubMessageBusSource.h"
+
+#include "SLiveLinkMessageBusSourceFactory.h"
+
+#ifndef WITH_LIVELINK_HUB
+#define WITH_LIVELINK_HUB 0
+#endif
+
+
+#define LOCTEXT_NAMESPACE "LiveLinkHubMessageBusSourceFactory"
+
+FText ULiveLinkHubMessageBusSourceFactory::GetSourceDisplayName() const
+{
+	return LOCTEXT("SourceDisplayName", "LiveLink Hub");
+}
+
+FText ULiveLinkHubMessageBusSourceFactory::GetSourceTooltip() const
+{
+	return LOCTEXT("SourceTooltip", "Creates a connection to a LiveLink Hub instance.");
+}
+
+TSharedPtr<SWidget> ULiveLinkHubMessageBusSourceFactory::BuildCreationPanel(FOnLiveLinkSourceCreated InOnLiveLinkSourceCreated) const
+{
+	return SNew(SLiveLinkMessageBusSourceFactory)
+		.OnSourceSelected(FOnLiveLinkMessageBusSourceSelected::CreateUObject(this, &ULiveLinkHubMessageBusSourceFactory::OnSourceSelected, InOnLiveLinkSourceCreated))
+		.FactoryClass(GetClass());
+}
+
+TSharedPtr<FLiveLinkMessageBusSource> ULiveLinkHubMessageBusSourceFactory::MakeSource(const FText& Name,
+																				   const FText& MachineName,
+																				   const FMessageAddress& Address,
+																				   double TimeOffset) const
+{
+	return MakeShared<FLiveLinkHubMessageBusSource>(Name, MachineName, Address, TimeOffset);
+}
+
+bool ULiveLinkHubMessageBusSourceFactory::IsEnabled() const
+{
+#if WITH_LIVELINK_HUB
+return false;
+#else
+return true;
+#endif
+}
+
+#undef LOCTEXT_NAMESPACE
diff --git a/Engine/Plugins/Animation/LiveLink/Source/LiveLinkHubMessaging/Private/LiveLinkHubMessageBusSourceFactory.h b/Engine/Plugins/Animation/LiveLink/Source/LiveLinkHubMessaging/Private/LiveLinkHubMessageBusSourceFactory.h
new file mode 100644
--- /dev/null
+++ b/Engine/Plugins/Animation/LiveLink/Source/LiveLinkHubMessaging/Private/LiveLinkHubMessageBusSourceFactory.h
@@ -0,0 +1,31 @@
+// Copyright Epic Games, Inc. All Rights Reserved.
+
+#pragma once
+
+#include "LiveLinkMessageBusFinder.h"		// For FProviderPollResultPtr typedef
+#include "LiveLinkMessageBusSourceFactory.h"
+#include "LiveLinkHubMessageBusSourceFactory.generated.h"
+
+struct FMessageAddress;
+
+
+UCLASS()
+class LIVELINKHUBMESSAGING_API ULiveLinkHubMessageBusSourceFactory : public ULiveLinkMessageBusSourceFactory
+{
+public:
+	GENERATED_BODY()
+
+	//~ Begin ULiveLinkSourceFactory interface
+	virtual FText GetSourceDisplayName() const override;
+	virtual FText GetSourceTooltip() const override;
+	virtual TSharedPtr<SWidget> BuildCreationPanel(FOnLiveLinkSourceCreated OnLiveLinkSourceCreated) const override;
+	virtual bool IsEnabled() const override;
+	//~ End ULiveLinkSourceFactory interface
+
+protected:
+	//~ ULiveLinkMessageBusSourceFactory interface
+	virtual TSharedPtr<class FLiveLinkMessageBusSource> MakeSource(const FText& Name,
+																   const FText& MachineName,
+																   const FMessageAddress& Address,
+																   double TimeOffset) const override;
+};
diff --git a/Engine/Plugins/Animation/LiveLink/Source/LiveLinkHubMessaging/Private/LiveLinkHubMessagingModule.cpp b/Engine/Plugins/Animation/LiveLink/Source/LiveLinkHubMessaging/Private/LiveLinkHubMessagingModule.cpp
--- a/Engine/Plugins/Animation/LiveLink/Source/LiveLinkHubMessaging/Private/LiveLinkHubMessagingModule.cpp
+++ b/Engine/Plugins/Animation/LiveLink/Source/LiveLinkHubMessaging/Private/LiveLinkHubMessagingModule.cpp
@@ -2,6 +2,8 @@
 
 #include "CoreMinimal.h"
 #include "LiveLinkHubConnectionManager.h"
+#include "LiveLinkHubMessageBusSourceFactory.h"
+#include "LiveLinkMessageBusSourceFactory.h"
 #include "Modules/ModuleManager.h"
 
 #ifndef WITH_LIVELINK_DISCOVERY_MANAGER_THREAD
@@ -12,27 +14,78 @@
 #define WITH_LIVELINK_HUB 0
 #endif
 
+#ifndef LIVELINKHUB_USE_CONNECTION_MANAGER
+#define LIVELINKHUB_USE_CONNECTION_MANAGER !WITH_LIVELINK_HUB && WITH_LIVELINK_DISCOVERY_MANAGER_THREAD
+#endif
+
 class FLiveLinkHubMessagingModule : public IModuleInterface
 {
 	// The connection manager is used to communicate with the hub, so we don't need it when we're running the hub itself.
-#if !WITH_LIVELINK_HUB && WITH_LIVELINK_DISCOVERY_MANAGER_THREAD
 public:
 	//~ Begin IModuleInterface
 	virtual void StartupModule() override
 	{
+#if LIVELINKHUB_USE_CONNECTION_MANAGER
 		ConnectionManager = MakeUnique<FLiveLinkHubConnectionManager>();
+#endif
+
+		SourceFilterDelegate = ILiveLinkModule::Get().RegisterMessageBusSourceFilter(FOnLiveLinkShouldDisplaySource::CreateRaw(this, &FLiveLinkHubMessagingModule::OnFilterMessageBusSource));
 	}
 
 	virtual void ShutdownModule() override
 	{
+		if (ILiveLinkModule* LiveLinkModule = FModuleManager::Get().GetModulePtr<ILiveLinkModule>("LiveLink"))
+		{
+			LiveLinkModule->UnregisterMessageBusSourceFilter(SourceFilterDelegate);
+		}
+
+#if LIVELINKHUB_USE_CONNECTION_MANAGER
 		ConnectionManager.Reset();
+#endif
 	}
 	//~ End IModuleInterface
 
 private:
+#if LIVELINKHUB_USE_CONNECTION_MANAGER
 	/** Manages the connection to the live link hub. */
 	TUniquePtr<FLiveLinkHubConnectionManager> ConnectionManager;
 #endif
+
+	/** Filter invoked by the messagebus source factory to filter out sources in the creation panel. */
+	bool OnFilterMessageBusSource(UClass* FactoryClass, FProviderPollResultPtr PollResult) const
+	{
+		if (!PollResult)
+		{
+			return false;
+		}
+
+		const bool bIsLiveLinkHubProvider = PollResult->Annotations.FindRef(FLiveLinkHubMessageAnnotation::ProviderTypeAnnotation) == UE::LiveLinkHub::Private::LiveLinkHubProviderType;
+
+#if WITH_LIVELINK_HUB
+		// When running the hub
+		if (FactoryClass == ULiveLinkMessageBusSourceFactory::StaticClass())
+		{
+			// Don't show the livelink hub in the message bus source discovery. (Should be changed in the future if we want to allow hubs to speak to each others).
+			return !bIsLiveLinkHubProvider;
+		}
+#else
+		// When running in UE
+		if (FactoryClass == ULiveLinkHubMessageBusSourceFactory::StaticClass())
+		{
+			return bIsLiveLinkHubProvider;
+		}
+		else if (FactoryClass == ULiveLinkMessageBusSourceFactory::StaticClass())
+		{
+			return !bIsLiveLinkHubProvider;
+		}
+#endif
+
+		return true;
+	}
+
+	/** Handle to the delegate used to filter message bus sources. */
+	FDelegateHandle SourceFilterDelegate;
 };
 
+
 IMPLEMENT_MODULE(FLiveLinkHubMessagingModule, LiveLinkHubMessaging);
diff --git a/Engine/Source/Runtime/LiveLinkInterface/Public/LiveLinkSourceFactory.h b/Engine/Source/Runtime/LiveLinkInterface/Public/LiveLinkSourceFactory.h
--- a/Engine/Source/Runtime/LiveLinkInterface/Public/LiveLinkSourceFactory.h
+++ b/Engine/Source/Runtime/LiveLinkInterface/Public/LiveLinkSourceFactory.h
@@ -64,4 +64,7 @@ public:
 
 	/** Create a new source from a ConnectionString */
 	LIVELINKINTERFACE_API virtual TSharedPtr<ILiveLinkSource> CreateSource(const FString& ConnectionString) const PURE_VIRTUAL(ULiveLinkSourceFactory::CreateSource, return TSharedPtr<ILiveLinkSource>(); );
+
+	/** Whether this factory is enabled. Can be overriden to disable it in certain configurations. */
+	LIVELINKINTERFACE_API virtual bool IsEnabled() const { return true; }
 };
