Fix deadlock when renaming subject in LiveLink Hub

diff --git a/Engine/Plugins/Animation/LiveLink/Source/LiveLinkHub/Private/Session/LiveLinkHubSession.h b/Engine/Plugins/Animation/LiveLink/Source/LiveLinkHub/Private/Session/LiveLinkHubSession.h
--- a/Engine/Plugins/Animation/LiveLink/Source/LiveLinkHub/Private/Session/LiveLinkHubSession.h
+++ b/Engine/Plugins/Animation/LiveLink/Source/LiveLinkHub/Private/Session/LiveLinkHubSession.h
@@ -75,8 +75,19 @@ public:
 
 	virtual void RenameSubject(const FLiveLinkSubjectKey& SubjectKey, FName NewName) override
 	{
-		FWriteScopeLock Locker(SessionDataLock);
-		SessionData.SubjectsConfig.RenameSubject(SubjectKey, NewName);
+		FLiveLinkHubSubjectSessionConfig ConfigCopy;
+		{
+			FReadScopeLock Locker(SessionDataLock);
+			ConfigCopy = SessionData.SubjectsConfig;
+		}
+
+		ConfigCopy.RenameSubject(SubjectKey, NewName);
+
+		{
+			// Copied over in a different step to avoid acquiring the rw lock in a method called by RenameSubject
+			FWriteScopeLock Locker(SessionDataLock);
+			SessionData.SubjectsConfig = MoveTemp(ConfigCopy);
+		}
 	}
 
 	virtual TOptional<FLiveLinkHubSubjectProxy> GetSubjectConfig(const FLiveLinkSubjectKey& SubjectKey) const override
