Fix crash when creating duplicated livelink hub sources

diff --git a/Engine/Plugins/Animation/LiveLink/Source/LiveLink/Private/SLiveLinkMessageBusSourceFactory.cpp b/Engine/Plugins/Animation/LiveLink/Source/LiveLink/Private/SLiveLinkMessageBusSourceFactory.cpp
--- a/Engine/Plugins/Animation/LiveLink/Source/LiveLink/Private/SLiveLinkMessageBusSourceFactory.cpp
+++ b/Engine/Plugins/Animation/LiveLink/Source/LiveLink/Private/SLiveLinkMessageBusSourceFactory.cpp
@@ -198,8 +198,11 @@ TSharedRef<ITableRow> SLiveLinkMessageBusSourceFactory::MakeSourceListViewWidget
 
 void SLiveLinkMessageBusSourceFactory::OnSourceListSelectionChanged(TSharedPtr<FLiveLinkSource> Source, ESelectInfo::Type SelectionType)
 {
-	SelectedResult = Source->PollResult;
-	OnSourceSelected.ExecuteIfBound(SelectedResult);
+	if (Source)
+	{
+		SelectedResult = Source->PollResult;
+		OnSourceSelected.ExecuteIfBound(SelectedResult);
+	}
 }
 
 #undef LOCTEXT_NAMESPACE
