Enable high dpi mode for livelink hub

diff --git a/Engine/Plugins/Animation/LiveLink/Source/LiveLinkHub/Private/UI/Window/LiveLinkHubWindowController.cpp b/Engine/Plugins/Animation/LiveLink/Source/LiveLinkHub/Private/UI/Window/LiveLinkHubWindowController.cpp
--- a/Engine/Plugins/Animation/LiveLink/Source/LiveLinkHub/Private/UI/Window/LiveLinkHubWindowController.cpp
+++ b/Engine/Plugins/Animation/LiveLink/Source/LiveLinkHub/Private/UI/Window/LiveLinkHubWindowController.cpp
@@ -98,6 +98,7 @@ TSharedPtr<FModalWindowManager> FLiveLinkHubWindowController::InitializeSlateApp
 {
 	TRACE_CPUPROFILER_EVENT_SCOPE(FLiveLinkHubWindowController::InitializeAsStandaloneApplication);
 	FSlateApplication::InitializeAsStandaloneApplication(GetStandardStandaloneRenderer());
+	FSlateApplication::InitHighDPI(true);
 
 	const FText ApplicationTitle = LOCTEXT("AppTitle", "LiveLink Hub");
 	FGlobalTabmanager::Get()->SetApplicationTitle(ApplicationTitle);
