LiveLinkHub - Fix disabling subject not working in the subjects tab

diff --git a/Engine/Plugins/Animation/LiveLink/Source/LiveLinkHub/Private/LiveLinkHub.cpp b/Engine/Plugins/Animation/LiveLink/Source/LiveLinkHub/Private/LiveLinkHub.cpp
--- a/Engine/Plugins/Animation/LiveLink/Source/LiveLinkHub/Private/LiveLinkHub.cpp
+++ b/Engine/Plugins/Animation/LiveLink/Source/LiveLinkHub/Private/LiveLinkHub.cpp
@@ -154,7 +154,10 @@ void FLiveLinkHub::OnFrameDataReceived_AnyThread(const FLiveLinkSubjectKey& InSu
 	FrameDataCopy.InitializeWith(InFrameDataStruct);
 
 	const FName OverridenName = GetSubjectNameOverride(InSubjectKey);
-	LiveLinkProvider->UpdateSubjectFrameData(OverridenName, MoveTemp(FrameDataCopy));
+	if (LiveLinkHubClient->IsSubjectEnabled(OverridenName))
+	{
+		LiveLinkProvider->UpdateSubjectFrameData(OverridenName, MoveTemp(FrameDataCopy));
+	}
 }
 
 void FLiveLinkHub::OnSubjectAdded(FLiveLinkSubjectKey InSubjectKey) const
