Live Link Hub - Make a UI cleanup pass
Add icons for livelink hub

diff --git a/Engine/Plugins/Animation/LiveLink/Source/LiveLinkEditor/Internal/LiveLinkClientPanelToolbar.h b/Engine/Plugins/Animation/LiveLink/Source/LiveLinkEditor/Internal/LiveLinkClientPanelToolbar.h
--- a/Engine/Plugins/Animation/LiveLink/Source/LiveLinkEditor/Internal/LiveLinkClientPanelToolbar.h
+++ b/Engine/Plugins/Animation/LiveLink/Source/LiveLinkEditor/Internal/LiveLinkClientPanelToolbar.h
@@ -21,8 +21,21 @@ struct FAssetData;
 
 class LIVELINKEDITOR_API SLiveLinkClientPanelToolbar : public SCompoundWidget, public FGCObject
 {
-	SLATE_BEGIN_ARGS(SLiveLinkClientPanelToolbar){}
+	SLATE_BEGIN_ARGS(SLiveLinkClientPanelToolbar)
+		: _SourceButtonAlignment(HAlign_Fill)
+		, _ShowPresetPicker(true)
+		, _ShowSettings(true)
+		{}
+	/** Horizontal alignment of the add source button. */
+	SLATE_ARGUMENT(EHorizontalAlignment, SourceButtonAlignment)
+	/** Parent window override. */
 	SLATE_ARGUMENT(TSharedPtr<SWindow>, ParentWindow)
+	/** (Optional) Custom header displayed on the left side of the toolbar.  */
+	SLATE_ARGUMENT(TSharedPtr<SWidget>, CustomHeader)
+	/** Whether to show the preset picker button. */
+	SLATE_ARGUMENT(bool, ShowPresetPicker)
+	/** Whether to show the settings button. */
+	SLATE_ARGUMENT(bool, ShowSettings)
 	SLATE_END_ARGS()
 
 
diff --git a/Engine/Plugins/Animation/LiveLink/Source/LiveLinkEditor/Private/LiveLinkClientPanelToolbar.cpp b/Engine/Plugins/Animation/LiveLink/Source/LiveLinkEditor/Private/LiveLinkClientPanelToolbar.cpp
--- a/Engine/Plugins/Animation/LiveLink/Source/LiveLinkEditor/Private/LiveLinkClientPanelToolbar.cpp
+++ b/Engine/Plugins/Animation/LiveLink/Source/LiveLinkEditor/Private/LiveLinkClientPanelToolbar.cpp
@@ -41,6 +41,7 @@
 #include "Widgets/Layout/SUniformGridPanel.h"
 #include "Widgets/Notifications/SNotificationList.h"
 #include "SPositiveActionButton.h"
+#include "Widgets/SNullWidget.h"
 #include "Widgets/SWindow.h"
 
 
@@ -432,92 +433,108 @@ void SLiveLinkClientPanelToolbar::Construct(const FArguments& Args, FLiveLinkCli
 	[
 		SNew(SBorder)
 		.BorderImage(FAppStyle::GetBrush("ToolPanel.GroupBorder"))
-		.Padding(2.0f)
+		.Padding(4.0f)
 		[
 			SNew(SHorizontalBox)
 			+ SHorizontalBox::Slot()
 			.Padding(.0f)
+			.HAlign(HAlign_Left)
 			.AutoWidth()
 			[
-				SNew(SPositiveActionButton)
-				.OnGetMenuContent(this, &SLiveLinkClientPanelToolbar::OnGenerateSourceMenu)
-				.Icon(FAppStyle::Get().GetBrush("Icons.Plus"))
-				.Text(LOCTEXT("AddSource", "Source"))
-				.ToolTipText(LOCTEXT("AddSource_ToolTip", "Add a new LiveLink source"))
+				Args._CustomHeader ? Args._CustomHeader.ToSharedRef() : SNullWidget::NullWidget
 			]
-
 			+ SHorizontalBox::Slot()
-			.Padding(8.f, 0.f, 0.f, 0.f)
-			.AutoWidth()
+			.Padding(.0f)
+			.HAlign(Args._SourceButtonAlignment)
+			.FillWidth(1.0f)
 			[
-				SNew(SComboButton)
-				.ContentPadding(4.f)
-				.ComboButtonStyle(FLiveLinkEditorPrivate::GetStyleSet(), "ComboButton")
-				.OnGetMenuContent(this, &SLiveLinkClientPanelToolbar::OnPresetGeneratePresetsMenu)
-				.ForegroundColor(FSlateColor::UseForeground())
-				.ButtonContent()
+				SNew(SHorizontalBox)
+				+ SHorizontalBox::Slot()
+				.AutoWidth()
+				.Padding(.0f)
+				[
+					SNew(SPositiveActionButton)
+					.OnGetMenuContent(this, &SLiveLinkClientPanelToolbar::OnGenerateSourceMenu)
+					.Icon(FAppStyle::Get().GetBrush("Icons.Plus"))
+					.Text(LOCTEXT("AddSource", "Add Source"))
+					.ToolTipText(LOCTEXT("AddSource_ToolTip", "Add a new LiveLink source"))
+				]
+				+ SHorizontalBox::Slot()
+				.Padding(8.f, 0.f, 0.f, 0.f)
+				.AutoWidth()
 				[
-					SNew(SHorizontalBox)
-					+ SHorizontalBox::Slot()
-					.Padding(4.f, 0.f, 4.f, 0.f)
-					.AutoWidth()
+					SNew(SComboButton)
+					.ContentPadding(4.f)
+					.ComboButtonStyle(FLiveLinkEditorPrivate::GetStyleSet(), "ComboButton")
+					.OnGetMenuContent(this, &SLiveLinkClientPanelToolbar::OnPresetGeneratePresetsMenu)
+					.ForegroundColor(FSlateColor::UseForeground())
+					.Visibility(Args._ShowPresetPicker ? EVisibility::Visible : EVisibility::Collapsed)
+					.ButtonContent()
 					[
-						SNew(SImage)
-						.Image(FSlateIconFinder::FindIconBrushForClass(ULiveLinkPreset::StaticClass()))
+						SNew(SHorizontalBox)
+						+ SHorizontalBox::Slot()
+						.Padding(4.f, 0.f, 4.f, 0.f)
+						.AutoWidth()
+						[
+							SNew(SImage)
+							.Image(FSlateIconFinder::FindIconBrushForClass(ULiveLinkPreset::StaticClass()))
+						]
+						+ SHorizontalBox::Slot()
+						[
+							SNew(STextBlock)
+							.Text(LOCTEXT("PresetsToolbarButton", "Presets"))
+						]
 					]
-					+ SHorizontalBox::Slot()
+				]
+
+				+ SHorizontalBox::Slot()
+				.Padding(8.f, 0.f, 0.f, 0.f)
+				.VAlign(VAlign_Center)
+				.HAlign(HAlign_Center)
+				.AutoWidth()
+				[
+					SNew(SButton)
+					//.ContentPadding(FMargin(8.f, 0.f, 0.f, 0.f))
+					.ToolTipText(LOCTEXT("RevertChanges_Text", "Revert all changes made to this take back its original state (either its original preset, or an empty preset)."))
+					.ForegroundColor(FSlateColor::UseForeground())
+					.ButtonStyle(FAppStyle::Get(), "HoverHintOnly")
+					.OnClicked(this, &SLiveLinkClientPanelToolbar::OnRevertChanges)
+					.Visibility(Args._ShowPresetPicker ? EVisibility::Visible : EVisibility::Collapsed)
+					.IsEnabled(this, &SLiveLinkClientPanelToolbar::HasLoadedLiveLinkPreset)
 					[
 						SNew(STextBlock)
-						.Text(LOCTEXT("PresetsToolbarButton", "Presets"))
+						.Font(FAppStyle::Get().GetFontStyle("FontAwesome.11"))
+						.Text(FEditorFontGlyphs::Undo)
 					]
 				]
-			]
 
-			+ SHorizontalBox::Slot()
-			.Padding(8.f, 0.f, 0.f, 0.f)
-			.VAlign(VAlign_Center)
-			.HAlign(HAlign_Center)
-			.AutoWidth()
-			[
-				SNew(SButton)
-				//.ContentPadding(FMargin(8.f, 0.f, 0.f, 0.f))
-				.ToolTipText(LOCTEXT("RevertChanges_Text", "Revert all changes made to this take back its original state (either its original preset, or an empty preset)."))
-				.ForegroundColor(FSlateColor::UseForeground())
-				.ButtonStyle(FAppStyle::Get(), "HoverHintOnly")
-				.OnClicked(this, &SLiveLinkClientPanelToolbar::OnRevertChanges)
-				.IsEnabled(this, &SLiveLinkClientPanelToolbar::HasLoadedLiveLinkPreset)
+				+ SHorizontalBox::Slot()
 				[
-					SNew(STextBlock)
-					.Font(FAppStyle::Get().GetFontStyle("FontAwesome.11"))
-					.Text(FEditorFontGlyphs::Undo)
+					SNew(SSpacer)
 				]
-			]
-
-			+ SHorizontalBox::Slot()
-			[
-				SNew(SSpacer)
-			]
 
-			+ SHorizontalBox::Slot()
-			.Padding(.0f)
-			.AutoWidth()
-			.HAlign(HAlign_Right)
-			[
-				SNew(SBox)
-				.WidthOverride(ButtonBoxSize)
-				.HeightOverride(ButtonBoxSize)
+				+ SHorizontalBox::Slot()
+				.Padding(.0f)
+				.AutoWidth()
+				.HAlign(HAlign_Right)
 				[
-					SNew(SCheckBox)
-					.Padding(4.f)
-					.ToolTipText(LOCTEXT("ShowUserSettings_Tip", "Show/Hide the general user settings for LiveLink"))
-					.Style(FAppStyle::Get(), "ToggleButtonCheckbox")
-					.ForegroundColor(FSlateColor::UseForeground())
-					.IsChecked_Lambda([]() { return ECheckBoxState::Unchecked; })
-					.OnCheckStateChanged_Lambda([](ECheckBoxState CheckState){ FModuleManager::LoadModuleChecked<ISettingsModule>("Settings").ShowViewer("Project", "Plugins", "LiveLink"); })
+					SNew(SBox)
+					.WidthOverride(ButtonBoxSize)
+					.Visibility(Args._ShowSettings ? EVisibility::Visible : EVisibility::Collapsed)
+					.HeightOverride(ButtonBoxSize)
 					[
-						SNew(STextBlock)
-						.Font(FAppStyle::Get().GetFontStyle("FontAwesome.14"))
-						.Text(FEditorFontGlyphs::Cogs)
+						SNew(SCheckBox)
+						.Padding(4.f)
+						.ToolTipText(LOCTEXT("ShowUserSettings_Tip", "Show/Hide the general user settings for LiveLink"))
+						.Style(FAppStyle::Get(), "ToggleButtonCheckbox")
+						.ForegroundColor(FSlateColor::UseForeground())
+						.IsChecked_Lambda([]() { return ECheckBoxState::Unchecked; })
+						.OnCheckStateChanged_Lambda([](ECheckBoxState CheckState){ FModuleManager::LoadModuleChecked<ISettingsModule>("Settings").ShowViewer("Project", "Plugins", "LiveLink"); })
+						[
+							SNew(STextBlock)
+							.Font(FAppStyle::Get().GetFontStyle("FontAwesome.14"))
+							.Text(FEditorFontGlyphs::Cogs)
+						]
 					]
 				]
 			]
diff --git a/Engine/Plugins/Animation/LiveLink/Source/LiveLinkEditor/Private/LiveLinkEditorModule.cpp b/Engine/Plugins/Animation/LiveLink/Source/LiveLinkEditor/Private/LiveLinkEditorModule.cpp
--- a/Engine/Plugins/Animation/LiveLink/Source/LiveLinkEditor/Private/LiveLinkEditorModule.cpp
+++ b/Engine/Plugins/Animation/LiveLink/Source/LiveLinkEditor/Private/LiveLinkEditorModule.cpp
@@ -101,6 +101,15 @@ public:
 		StyleSet->Set("LiveLinkClient.Common.Icon", new IMAGE_PLUGIN_BRUSH(TEXT("LiveLink_40x"), Icon40x40));
 		StyleSet->Set("LiveLinkClient.Common.Icon.Small", new IMAGE_PLUGIN_BRUSH_SVG("Starship/LiveLink", Icon16x16));
 
+		StyleSet->Set("LiveLinkHub.Clients.Icon", new IMAGE_PLUGIN_BRUSH_SVG(TEXT("Starship/Clients"), Icon16x16));
+		StyleSet->Set("LiveLinkHub.Sources.Icon", new IMAGE_PLUGIN_BRUSH_SVG(TEXT("Starship/Sources"), Icon16x16));
+		StyleSet->Set("LiveLinkHub.Layout.Icon", new IMAGE_PLUGIN_BRUSH_SVG(TEXT("Starship/Layout"), Icon16x16));
+		StyleSet->Set("LiveLinkHub.Playback.Icon", new IMAGE_PLUGIN_BRUSH_SVG(TEXT("Starship/Playback"), Icon16x16));
+		StyleSet->Set("LiveLinkHub.Subjects.Icon", new IMAGE_PLUGIN_BRUSH_SVG(TEXT("Starship/Subjects"), Icon16x16));
+
+		StyleSet->Set("LiveLinkHub.Icon.Small", new IMAGE_PLUGIN_BRUSH(TEXT("Starship/LiveLinkHub_16x16"), Icon16x16));
+		StyleSet->Set("LiveLinkHub.Icon", new IMAGE_PLUGIN_BRUSH_SVG(TEXT("Starship/LiveLinkHub_Icon"), Icon40x40));
+
 		StyleSet->Set("ClassIcon.LiveLinkPreset",                      new IMAGE_PLUGIN_BRUSH_SVG("Starship/LiveLink", Icon16x16));
 		StyleSet->Set("ClassIcon.LiveLinkFrameInterpolationProcessor", new IMAGE_PLUGIN_BRUSH_SVG("Starship/LiveLink", Icon16x16));
 		StyleSet->Set("ClassIcon.LiveLinkFramePreProcessor",           new IMAGE_PLUGIN_BRUSH_SVG("Starship/LiveLink", Icon16x16));
diff --git a/Engine/Plugins/Animation/LiveLink/Source/LiveLinkHub/Private/Clients/SLiveLinkHubClientsView.h b/Engine/Plugins/Animation/LiveLink/Source/LiveLinkHub/Private/Clients/SLiveLinkHubClientsView.h
--- a/Engine/Plugins/Animation/LiveLink/Source/LiveLinkHub/Private/Clients/SLiveLinkHubClientsView.h
+++ b/Engine/Plugins/Animation/LiveLink/Source/LiveLinkHub/Private/Clients/SLiveLinkHubClientsView.h
@@ -14,10 +14,12 @@
 #include "LiveLinkHubClientsModel.h"
 #include "LiveLinkHubUEClientInfo.h"
 #include "Session/LiveLinkHubSessionManager.h"
+#include "SPositiveActionButton.h"
 #include "Styling/SlateStyleMacros.h"
 #include "Styling/SlateTypes.h"
 #include "Widgets/Input/SButton.h"
 #include "Widgets/Input/SCheckBox.h"
+#include "Widgets/Layout/SScaleBox.h"
 #include "Widgets/Input/SComboButton.h"
 #include "Widgets/Text/STextBlock.h"
 #include "Widgets/Views/STreeView.h"
@@ -313,13 +315,22 @@ public:
 		[
 			SNew(SVerticalBox)
 			+ SVerticalBox::Slot()
+			.Padding(FMargin(10.0, 7.0))
 			.AutoHeight()
 			[
 				SNew(SHorizontalBox)
 				+ SHorizontalBox::Slot()
+				.AutoWidth()
+				.Padding(0.f, 5.f, 0.f, 5.f)
+				[
+					SNew(SImage)
+					.ColorAndOpacity(FSlateColor::UseForeground())
+					.Image(FSlateIcon("LiveLinkStyle", "LiveLinkHub.Clients.Icon").GetIcon())
+				]
+				+ SHorizontalBox::Slot()
 				.HAlign(HAlign_Left)
 				.FillWidth(1.0f)
-				.Padding(FMargin(2.0f))
+				.Padding(FMargin(4.0, 2.0))
 				[
 					SNew(STextBlock)
 					.Font( DEFAULT_FONT( "Regular", 14 ) )
@@ -329,13 +340,10 @@ public:
 				.HAlign(HAlign_Right)
 				.AutoWidth()
 				[
-					SNew(SComboButton)
-					.CollapseMenuOnParentFocus(true)
-					.ButtonContent()
-					[
-						SNew(STextBlock)
-						.Text(LOCTEXT("AddClientLabel", "Add Client"))
-					]
+					SNew(SPositiveActionButton)
+                	.Icon(FAppStyle::Get().GetBrush("Icons.Plus"))
+                	.Text(LOCTEXT("AddClientLabel", "Add Client"))
+                	.ToolTipText(LOCTEXT("AddClient_Tooltip", "Connect to an unreal editor instance."))
 					.MenuContent()
 					[
 						SNew(SBox)
diff --git a/Engine/Plugins/Animation/LiveLink/Source/LiveLinkHub/Private/Recording/SLiveLinkHubRecordingView.h b/Engine/Plugins/Animation/LiveLink/Source/LiveLinkHub/Private/Recording/SLiveLinkHubRecordingView.h
--- a/Engine/Plugins/Animation/LiveLink/Source/LiveLinkHub/Private/Recording/SLiveLinkHubRecordingView.h
+++ b/Engine/Plugins/Animation/LiveLink/Source/LiveLinkHub/Private/Recording/SLiveLinkHubRecordingView.h
@@ -40,8 +40,28 @@ public:
 		ChildSlot
 		[
 			SNew(SButton)
-				.Text(TAttribute<FText>::CreateSP(this, &SLiveLinkHubRecordingView::GetRecordingText))
 				.OnClicked(this, &SLiveLinkHubRecordingView::OnClickRecordButton)
+				[
+					SNew(SHorizontalBox)
+					+ SHorizontalBox::Slot()
+					.AutoWidth()
+					.HAlign(HAlign_Center)
+					.VAlign(VAlign_Center)
+					[
+						SNew(SImage)
+						.Image(FAppStyle::Get().GetBrush("Animation.Record"))
+						.ColorAndOpacity(FSlateColor::UseForeground())
+					]
+					+ SHorizontalBox::Slot()
+					.Padding(FMargin(3, 0, 0, 0))
+					.VAlign(VAlign_Center)
+					.AutoWidth()
+					[
+						SNew(STextBlock)
+						.Text(TAttribute<FText>::CreateSP(this, &SLiveLinkHubRecordingView::GetRecordingText))
+						.TextStyle(FAppStyle::Get(), "NormalText")
+					]
+				]
 		];
 	}
 
@@ -65,7 +85,7 @@ public:
 			OnStartRecordingDelegate.Execute();
 		}
 
-		return FReply::Handled();
+		return FReply::Handled(); 
 	}
 
 	/** Returns whether we're currently recording. */
@@ -83,4 +103,4 @@ private:
 	FSimpleDelegate OnStopRecordingDelegate;
 };
 
-#undef LOCTEXT_NAMESPACE /* LiveLinkHub.RecordingView */
\ No newline at end of file
+#undef LOCTEXT_NAMESPACE /* LiveLinkHub.RecordingView */
diff --git a/Engine/Plugins/Animation/LiveLink/Source/LiveLinkHub/Private/UI/Widgets/LiveLinkHubMainTabController.cpp b/Engine/Plugins/Animation/LiveLink/Source/LiveLinkHub/Private/UI/Widgets/LiveLinkHubMainTabController.cpp
--- a/Engine/Plugins/Animation/LiveLink/Source/LiveLinkHub/Private/UI/Widgets/LiveLinkHubMainTabController.cpp
+++ b/Engine/Plugins/Animation/LiveLink/Source/LiveLinkHub/Private/UI/Widgets/LiveLinkHubMainTabController.cpp
@@ -14,6 +14,7 @@ void FLiveLinkHubMainTabController::Init(const FLiveLinkHubComponentInitParams&
 {
 	FGlobalTabmanager::Get()->RegisterTabSpawner(LiveLinkHubTabID, FOnSpawnTab::CreateRaw(this, &FLiveLinkHubMainTabController::SpawnMainTab, Params.WindowController->GetRootWindow()))
 		.SetDisplayName(LOCTEXT("LiveLinkHubTabLabel", "LiveLink Hub"))
+		.SetIcon(FSlateIcon("LiveLinkStyle", "LiveLinkHub.Icon.Small"))
 		.SetTooltipText(LOCTEXT("MainTabTooltipText", "LiveLink Hub Main Tab"));
 	
 	Params.MainStack->AddTab(LiveLinkHubTabID, ETabState::OpenedTab);
@@ -29,6 +30,10 @@ TSharedRef<SDockTab> FLiveLinkHubMainTabController::SpawnMainTab(const FSpawnTab
 	const TSharedRef<SDockTab> DockTab = SNew(SDockTab)
 		.Label(LOCTEXT("LiveLinkHubTitle", "LiveLink Hub"))
 		.TabRole(MajorTab)
+		.OnCanCloseTab_Lambda([]()
+		{
+			return false;
+		})
 		.CanEverClose(false);
 
 	DockTab->SetContent(
diff --git a/Engine/Plugins/Animation/LiveLink/Source/LiveLinkHub/Private/UI/Widgets/SLiveLinkHubMainTabView.cpp b/Engine/Plugins/Animation/LiveLink/Source/LiveLinkHub/Private/UI/Widgets/SLiveLinkHubMainTabView.cpp
--- a/Engine/Plugins/Animation/LiveLink/Source/LiveLinkHub/Private/UI/Widgets/SLiveLinkHubMainTabView.cpp
+++ b/Engine/Plugins/Animation/LiveLink/Source/LiveLinkHub/Private/UI/Widgets/SLiveLinkHubMainTabView.cpp
@@ -17,6 +17,7 @@
 #include "Subjects/LiveLinkHubSubjectController.h"
 #include "Widgets/Docking/SDockTab.h"
 #include "Widgets/Input/SCheckBox.h"
+#include "Widgets/Input/SComboButton.h"
 #include "Widgets/Text/STextBlock.h"
 
 
@@ -65,6 +66,8 @@ SLiveLinkHubMainTabView::~SLiveLinkHubMainTabView()
 
 void SLiveLinkHubMainTabView::CreateTabs(const TSharedRef<FTabManager>& InTabManager, const TSharedRef<FTabManager::FLayout>& InLayout, const FArguments& InArgs)
 {
+	static const FName LiveLinkStyleName = "LiveLinkStyle";
+
 	InTabManager->RegisterTabSpawner(SourcesTabId, FOnSpawnTab::CreateSP(this, &SLiveLinkHubMainTabView::SpawnSourcesTab))
 		.SetDisplayName(SourcesTabName);
 
@@ -74,10 +77,12 @@ void SLiveLinkHubMainTabView::CreateTabs(const TSharedRef<FTabManager>& InTabMan
 	*/
 
 	InTabManager->RegisterTabSpawner(SubjectsTabId, FOnSpawnTab::CreateSP(this, &SLiveLinkHubMainTabView::SpawnSubjectsTab))
+		.SetIcon(FSlateIcon(LiveLinkStyleName, TEXT("LiveLinkHub.Subjects.Icon")))
 		.SetDisplayName(SubjectsTabName);
 	InTabManager->RegisterTabSpawner(SubjectsDetailsTabId, FOnSpawnTab::CreateSP(this, &SLiveLinkHubMainTabView::SpawnSubjectsDetailsTab))
 		.SetDisplayName(SubjectsDetailsTabName);
 	InTabManager->RegisterTabSpawner(PlaybackTabId, FOnSpawnTab::CreateSP(this, &SLiveLinkHubMainTabView::SpawnPlaybackTab))
+		.SetIcon(FSlateIcon(LiveLinkStyleName, TEXT("LiveLinkHub.Playback.Icon")))
 		.SetDisplayName(PlaybackTabName);
 	InTabManager->RegisterTabSpawner(ClientsTabId, FOnSpawnTab::CreateSP(this, &SLiveLinkHubMainTabView::SpawnClientsTab))
 		.SetDisplayName(ClientsTabName);
@@ -155,16 +160,40 @@ TSharedRef<SDockTab> SLiveLinkHubMainTabView::SpawnSourcesTab(const FSpawnTabArg
 {
 	FLiveLinkClient* Client = (FLiveLinkClient*)&IModularFeatures::Get().GetModularFeature<ILiveLinkClient>(ILiveLinkClient::ModularFeatureName);
 	FLiveLinkHubModule& LiveLinkHubModule = FModuleManager::Get().GetModuleChecked<FLiveLinkHubModule>("LiveLinkHub");
+
+	TSharedRef<SWidget> CustomToolbarHeader = SNew(SHorizontalBox)
+		+ SHorizontalBox::Slot()
+		.AutoWidth()
+		.Padding(0.f, 5.f, 0.f, 5.f)
+		[
+			SNew(SImage)
+			.ColorAndOpacity(FSlateColor::UseForeground())
+			.Image(FSlateIcon("LiveLinkStyle", "LiveLinkHub.Sources.Icon").GetIcon())
+		]
+		+ SHorizontalBox::Slot()
+		.HAlign(HAlign_Left)
+		.Padding(FMargin(4.0, 2.0))
+		[
+			SNew(STextBlock)
+			.Font( DEFAULT_FONT( "Regular", 14 ) )
+			.Text(LOCTEXT("SourcesHeaderText", "Sources"))
+		];
+		
 	
 	return SNew(SDockTab)
 		.TabRole(PanelTab)
 		[
 			SNew(SVerticalBox)
 			+SVerticalBox::Slot()
+			.Padding(FMargin(4.0f, 0.f,0.f,0.f))
 			.AutoHeight()
 			[
 				SNew(SLiveLinkClientPanelToolbar, Client)
+				.SourceButtonAlignment(HAlign_Right)
 				.ParentWindow(LiveLinkHubModule.GetLiveLinkHub()->GetRootWindow())
+				.ShowPresetPicker(false)
+				.ShowSettings(false)
+				.CustomHeader(CustomToolbarHeader)
 			]
 			+SVerticalBox::Slot()
 			.FillHeight(1.f)
diff --git a/Engine/Plugins/Animation/LiveLink/Source/LiveLinkHub/Private/UI/Widgets/SLiveLinkHubTabViewBase.cpp b/Engine/Plugins/Animation/LiveLink/Source/LiveLinkHub/Private/UI/Widgets/SLiveLinkHubTabViewBase.cpp
--- a/Engine/Plugins/Animation/LiveLink/Source/LiveLinkHub/Private/UI/Widgets/SLiveLinkHubTabViewBase.cpp
+++ b/Engine/Plugins/Animation/LiveLink/Source/LiveLinkHub/Private/UI/Widgets/SLiveLinkHubTabViewBase.cpp
@@ -10,6 +10,7 @@
 #include "Modules/ModuleManager.h"
 #include "Recording/LiveLinkHubRecordingController.h"
 #include "SLiveLinkHubStatusBar.h"
+#include "Widgets/Images/SImage.h"
 #include "Widgets/Input/SButton.h"
 #include "Widgets/Input/SComboButton.h"
 #include "Widgets/Layout/SBorder.h"
@@ -42,66 +43,83 @@ void SLiveLinkHubTabViewBase::Construct(const FArguments& InArgs)
 	ChildSlot
 	[
 		SNew(SBorder)
-		.BorderImage(FAppStyle::Get().GetBrush("ToolPanel.GroupBorder"))
-		.Padding(FMargin(1.0f, 2.0f))
+		.BorderImage(FAppStyle::GetBrush("NoBorder"))
+		.Padding(FMargin(0.0))
 		[
 			SNew(SVerticalBox)
 			+SVerticalBox::Slot()
-			.FillHeight(0.05f)
+			.AutoHeight()
 			.HAlign(HAlign_Fill)
-			.VAlign(VAlign_Center)
 			[
-				SNew(SHorizontalBox)
-				+ SHorizontalBox::Slot()
-				.Padding(2.f)
-				.HAlign(HAlign_Left)
-				[
-					// Should be replaced by the mode switcher widget
-					SNew(SComboButton)
-					.ContentPadding(3.f)
-					.OnGetMenuContent_Static(&::GetModeSwitcherContent)
-					.ButtonContent()
-					[
-						// Should switch depending on the currently selected mode.
-						SNew(STextBlock).Text(LOCTEXT("SelectModeLabl", "Creator Mode"))
-					]
-				]
-				+ SHorizontalBox::Slot()
-				.HAlign(HAlign_Right)
-				.VAlign(VAlign_Center)
+				SNew(SBorder)
+				.BorderImage(FAppStyle::GetBrush("ToolPanel.GroupBorder"))
+				.Padding(FMargin(4.0f, 6.0f))
 				[
 					SNew(SHorizontalBox)
 					+ SHorizontalBox::Slot()
-					.Padding(2.f)
-					.VAlign(VAlign_Center)
-					.AutoWidth()
+					.HAlign(HAlign_Left)
 					[
-						LiveLinkHub->GetRecordingController()->MakeRecordToolbarEntry()
+						// Should be replaced by the mode switcher widget
+						SNew(SComboButton)
+						.ContentPadding(FMargin(2.0f, 3.0f, 2.0f, 3.0f))
+						.OnGetMenuContent_Static(&::GetModeSwitcherContent)
+						.ButtonContent()
+						[
+							SNew(SHorizontalBox)
+							+ SHorizontalBox::Slot()
+							.AutoWidth()
+							[
+								SNew(SImage)
+								.ColorAndOpacity(FSlateColor::UseForeground())
+								.Image(FSlateIcon("LiveLinkStyle", "LiveLinkHub.Layout.Icon").GetIcon())
+							]
+							+ SHorizontalBox::Slot()
+							.Padding(FMargin(4.0f, 0.0f, 0.0f, 0.0f))
+                            .AutoWidth()
+                            [
+	                            // Should switch depending on the currently selected mode.
+								SNew(STextBlock).Text(LOCTEXT("SelectModeLabl", "Creator Mode"))
+                            ]
+
+						]
 					]
 					+ SHorizontalBox::Slot()
-					.Padding(2.f)
+					.HAlign(HAlign_Right)
 					.VAlign(VAlign_Center)
-					.AutoWidth()
 					[
-						SNew(SSeparator)
-						.Orientation(Orient_Vertical)
+						SNew(SHorizontalBox)
+						+ SHorizontalBox::Slot()
+						.Padding(2.f)
+						.VAlign(VAlign_Center)
+						.AutoWidth()
+						[
+							LiveLinkHub->GetRecordingController()->MakeRecordToolbarEntry()
+						]
+						+ SHorizontalBox::Slot()
+						.Padding(2.f)
+						.VAlign(VAlign_Center)
+						.AutoWidth()
+						[
+							SNew(SSeparator)
+							.Orientation(Orient_Vertical)
+						]
+						/* Disable until we use timecode in LLH
+						+ SHorizontalBox::Slot()
+						.Padding(2.f)
+	                    .VAlign(VAlign_Center)
+						[
+							SNew(STimecode)
+							.Timecode(MakeAttributeLambda([]
+							{
+								return FApp::GetTimecode();
+							}))
+						]*/
 					]
-					/* Disable until we use timecode in LLH
-					+ SHorizontalBox::Slot()
-					.Padding(2.f)
-                    .VAlign(VAlign_Center)
-					[
-						SNew(STimecode)
-						.Timecode(MakeAttributeLambda([]
-						{
-							return FApp::GetTimecode();
-						}))
-					]*/
 				]
 			]
 			+SVerticalBox::Slot()
 			.FillHeight(1.0f)
-			.Padding(1.0f, 2.0f)
+			.Padding(8.0f)
 			[
 				InArgs._Content.Widget
 			]
@@ -114,8 +132,13 @@ void SLiveLinkHubTabViewBase::Construct(const FArguments& InArgs)
 			+SVerticalBox::Slot()
 			.AutoHeight()
 			[
-				SNew(SLiveLinkHubStatusBar, StatusBarId)
-			]
+				SNew(SBorder)
+            	.BorderImage(FAppStyle::GetBrush("ToolPanel.GroupBorder"))
+            	.Padding(5.0f, 0.0f, 5.0f, 7.0f)
+            	[
+					SNew(SLiveLinkHubStatusBar, StatusBarId)
+            	]
+            ]
 		]
 	];
 }
