Add livelink hub status bar to the editor (Only appears when we have an active connection to a hub instance)

diff --git a/Engine/Plugins/Animation/LiveLink/LiveLink.uplugin b/Engine/Plugins/Animation/LiveLink/LiveLink.uplugin
--- a/Engine/Plugins/Animation/LiveLink/LiveLink.uplugin
+++ b/Engine/Plugins/Animation/LiveLink/LiveLink.uplugin
@@ -45,7 +45,12 @@
 			"Name": "LiveLinkHub",
 			"Type": "EditorAndProgram",
 			"ProgramAllowList": [ "LiveLinkHub" ],
-			"LoadingPhase": "PreDefault"
+			"LoadingPhase": "Default"
+		},
+		{
+			"Name": "LiveLinkHubEditor",
+			"Type": "Editor",
+			"LoadingPhase": "Default"
 		},
 		{
 			"Name": "LiveLinkComponents",
@@ -86,4 +91,4 @@
 			"ProgramAllowList": [ "LiveLinkHub" ]
 		}
 	]
-}
\ No newline at end of file
+}
diff --git a/Engine/Plugins/Animation/LiveLink/Source/LiveLink/Private/LiveLinkMessageBusSource.cpp b/Engine/Plugins/Animation/LiveLink/Source/LiveLink/Private/LiveLinkMessageBusSource.cpp
--- a/Engine/Plugins/Animation/LiveLink/Source/LiveLink/Private/LiveLinkMessageBusSource.cpp
+++ b/Engine/Plugins/Animation/LiveLink/Source/LiveLink/Private/LiveLinkMessageBusSource.cpp
@@ -18,6 +18,20 @@
 #include "MessageEndpointBuilder.h"
 #include "Misc/App.h"
 
+FText FLiveLinkMessageBusSource::ValidSourceStatus()
+{
+	return NSLOCTEXT("LiveLinkMessageBusSource", "ActiveStatus", "Active");
+}
+
+FText FLiveLinkMessageBusSource::InvalidSourceStatus()
+{
+	return NSLOCTEXT("LiveLinkMessageBusSource", "InvalidConnection", "Waiting for connection");
+}
+
+FText FLiveLinkMessageBusSource::TimeoutSourceStatus()
+{
+	return NSLOCTEXT("LiveLinkMessageBusSource", "TimeoutStatus", "Not responding");
+}
 
 FLiveLinkMessageBusSource::FLiveLinkMessageBusSource(const FText& InSourceType, const FText& InSourceMachineName, const FMessageAddress& InConnectionAddress, double InMachineTimeOffset)
 	: ConnectionAddress(InConnectionAddress)
@@ -206,13 +220,13 @@ FText FLiveLinkMessageBusSource::GetSourceStatus() const
 {
 	if (!ConnectionAddress.IsValid())
 	{
-		return NSLOCTEXT("LiveLinkMessageBusSource", "InvalidConnection", "Waiting for connection");
+		return InvalidSourceStatus();
 	}
 	else if (IsSourceStillValid())
 	{
-		return NSLOCTEXT("LiveLinkMessageBusSource", "ActiveStatus", "Active");
+		return ValidSourceStatus();
 	}
-	return NSLOCTEXT("LiveLinkMessageBusSource", "TimeoutStatus", "Not responding");
+	return TimeoutSourceStatus();
 }
 
 TSubclassOf<ULiveLinkSourceSettings> FLiveLinkMessageBusSource::GetSettingsClass() const
diff --git a/Engine/Plugins/Animation/LiveLink/Source/LiveLink/Public/LiveLinkClient.h b/Engine/Plugins/Animation/LiveLink/Source/LiveLink/Public/LiveLinkClient.h
--- a/Engine/Plugins/Animation/LiveLink/Source/LiveLink/Public/LiveLinkClient.h
+++ b/Engine/Plugins/Animation/LiveLink/Source/LiveLink/Public/LiveLinkClient.h
@@ -80,6 +80,7 @@ public:
 	virtual FText GetSourceType(FGuid EntryGuid) const override;
 	virtual FText GetSourceStatus(FGuid EntryGuid) const override;
 	virtual FText GetSourceMachineName(FGuid EntryGuid) const override;
+	virtual bool IsSourceStillValid(FGuid EntryGuid) const override;
 
 
 	virtual void PushSubjectStaticData_AnyThread(const FLiveLinkSubjectKey& SubjectKey, TSubclassOf<ULiveLinkRole> Role, FLiveLinkStaticDataStruct&& StaticData) override;
@@ -162,7 +163,6 @@ public:
 
 	FLiveLinkSubjectTimeSyncData GetTimeSyncData(FLiveLinkSubjectName SubjectName);
 
-	bool IsSourceStillValid(FGuid EntryGuid) const;
 	UE_DEPRECATED(4.23, "FLiveLinkClient::GetSourceTypeForEntry is deprecated. Please use GetSourceType instead!")
 	FText GetSourceTypeForEntry(FGuid EntryGuid) const { return GetSourceType(EntryGuid); }
 	UE_DEPRECATED(4.23, "FLiveLinkClient::GetMachineNameForEntry is deprecated. Please use GetSourceMachineName instead!")
diff --git a/Engine/Plugins/Animation/LiveLink/Source/LiveLink/Public/LiveLinkMessageBusSource.h b/Engine/Plugins/Animation/LiveLink/Source/LiveLink/Public/LiveLinkMessageBusSource.h
--- a/Engine/Plugins/Animation/LiveLink/Source/LiveLink/Public/LiveLinkMessageBusSource.h
+++ b/Engine/Plugins/Animation/LiveLink/Source/LiveLink/Public/LiveLinkMessageBusSource.h
@@ -21,6 +21,12 @@ struct FLiveLinkClearSubject;
 class LIVELINK_API FLiveLinkMessageBusSource : public ILiveLinkSource
 {
 public:
+	/** Text description for a valid source. */
+	static FText ValidSourceStatus();
+	/** Text description for an invalid source. */
+	static FText InvalidSourceStatus();
+	/** Text description for a source that has timed out. */
+	static FText TimeoutSourceStatus();
 
 	FLiveLinkMessageBusSource(const FText& InSourceType, const FText& InSourceMachineName, const FMessageAddress& InConnectionAddress, double InMachineTimeOffset);
 
diff --git a/Engine/Plugins/Animation/LiveLink/Source/LiveLinkHub/Private/Recording/SLiveLinkHubRecordingView.h b/Engine/Plugins/Animation/LiveLink/Source/LiveLinkHub/Private/Recording/SLiveLinkHubRecordingView.h
--- a/Engine/Plugins/Animation/LiveLink/Source/LiveLinkHub/Private/Recording/SLiveLinkHubRecordingView.h
+++ b/Engine/Plugins/Animation/LiveLink/Source/LiveLinkHub/Private/Recording/SLiveLinkHubRecordingView.h
@@ -7,7 +7,11 @@
 #include "Delegates/Delegate.h"
 #include "Delegates/DelegateCombinations.h"
 #include "Styling/SlateTypes.h"
+#include "Widgets/Images/SImage.h"
 #include "Widgets/Input/SButton.h"
+#include "Widgets/Layout/SBox.h"
+#include "Widgets/SBoxPanel.h"
+#include "Widgets/Text/STextBlock.h"
 
 class FLiveLinkHub;
 
diff --git a/Engine/Plugins/Animation/LiveLink/Source/LiveLinkHubEditor/LiveLinkHubEditor.Build.cs b/Engine/Plugins/Animation/LiveLink/Source/LiveLinkHubEditor/LiveLinkHubEditor.Build.cs
new file mode 100644
--- /dev/null
+++ b/Engine/Plugins/Animation/LiveLink/Source/LiveLinkHubEditor/LiveLinkHubEditor.Build.cs
@@ -0,0 +1,31 @@
+// Copyright Epic Games, Inc. All Rights Reserved.
+
+namespace UnrealBuildTool.Rules
+{
+	public class LiveLinkHubEditor : ModuleRules
+	{
+		public LiveLinkHubEditor(ReadOnlyTargetRules Target) : base(Target)
+		{
+			PublicDependencyModuleNames.AddRange(
+			new string[]
+			{
+				"Core",
+				"CoreUObject",
+			});
+
+			PrivateDependencyModuleNames.AddRange(
+			new string[]
+			{
+				"Engine",
+				"LiveLink",
+				"LiveLinkHub",
+				"LiveLinkHubMessaging",
+				"LiveLinkInterface",
+				"Slate",
+				"SlateCore",
+				"ToolMenus",
+				"UnrealEd"
+			});
+		}
+	}
+}
diff --git a/Engine/Plugins/Animation/LiveLink/Source/LiveLinkHubEditor/Private/LiveLinkHubEditorModule.cpp b/Engine/Plugins/Animation/LiveLink/Source/LiveLinkHubEditor/Private/LiveLinkHubEditorModule.cpp
new file mode 100644
--- /dev/null
+++ b/Engine/Plugins/Animation/LiveLink/Source/LiveLinkHubEditor/Private/LiveLinkHubEditorModule.cpp
@@ -0,0 +1,35 @@
+// Copyright Epic Games, Inc. All Rights Reserved.
+#pragma once
+
+#include "LiveLinkHubEditorModule.h"
+
+#include "SLiveLinkHubEditorStatusBar.h"
+#include "ToolMenus.h"
+
+void FLiveLinkHubEditorModule::OnPostEngineInit()
+{
+	RegisterLiveLinkHubStatusBar();
+}
+
+void FLiveLinkHubEditorModule::RegisterLiveLinkHubStatusBar()
+{
+	UToolMenu* Menu = UToolMenus::Get()->ExtendMenu(TEXT("LevelEditor.StatusBar.ToolBar"));
+
+	FToolMenuSection& LiveLinkHubSection = Menu->AddSection(TEXT("LiveLinkHub"), FText::GetEmpty(), FToolMenuInsert(NAME_None, EToolMenuInsertType::First));
+
+	LiveLinkHubSection.AddEntry(
+		FToolMenuEntry::InitWidget(TEXT("LiveLinkHubStatusBar"), CreateLiveLinkHubWidget(), FText::GetEmpty(), true, false)
+	);
+}
+
+void FLiveLinkHubEditorModule::UnregisterLiveLinkHubStatusBar()
+{
+	UToolMenus::UnregisterOwner(this);
+}
+
+TSharedRef<SWidget> FLiveLinkHubEditorModule::CreateLiveLinkHubWidget()
+{
+	return SNew(SLiveLinkHubEditorStatusBar);
+}
+
+IMPLEMENT_MODULE(FLiveLinkHubEditorModule, LiveLinkHubEditor);
diff --git a/Engine/Plugins/Animation/LiveLink/Source/LiveLinkHubEditor/Private/LiveLinkHubEditorModule.h b/Engine/Plugins/Animation/LiveLink/Source/LiveLinkHubEditor/Private/LiveLinkHubEditorModule.h
new file mode 100644
--- /dev/null
+++ b/Engine/Plugins/Animation/LiveLink/Source/LiveLinkHubEditor/Private/LiveLinkHubEditorModule.h
@@ -0,0 +1,40 @@
+// Copyright Epic Games, Inc. All Rights Reserved.
+#pragma once
+
+#include "CoreMinimal.h"
+#include "Misc/CoreDelegates.h"
+#include "Modules/ModuleInterface.h"
+
+/**
+ * Editor module responsible for creating the hub's status bar.
+ */
+class FLiveLinkHubEditorModule : public IModuleInterface
+{
+
+public:
+	//~ Begin IModuleInterface interface
+	virtual void StartupModule()
+	{
+		FCoreDelegates::OnPostEngineInit.AddRaw(this, &FLiveLinkHubEditorModule::OnPostEngineInit);
+	}
+
+	virtual void ShutdownModule()
+	{
+		FCoreDelegates::OnPostEngineInit.RemoveAll(this);
+		UnregisterLiveLinkHubStatusBar();
+	}
+	//~ End IModuleInterface interface
+
+private:
+	/** Callback used to register the livelink hub status bar. */
+	void OnPostEngineInit();
+
+	/** Extend the editor's bottom status bar to add the livelink hub widget. */
+	void RegisterLiveLinkHubStatusBar();
+
+	/** Unregister our tool menu customization. */
+	void UnregisterLiveLinkHubStatusBar();
+
+	/** Instantiate the livelink hub status bar widget. */
+    TSharedRef<class SWidget> CreateLiveLinkHubWidget();
+};
diff --git a/Engine/Plugins/Animation/LiveLink/Source/LiveLinkHubEditor/Private/SLiveLinkHubEditorStatusBar.cpp b/Engine/Plugins/Animation/LiveLink/Source/LiveLinkHubEditor/Private/SLiveLinkHubEditorStatusBar.cpp
new file mode 100644
--- /dev/null
+++ b/Engine/Plugins/Animation/LiveLink/Source/LiveLinkHubEditor/Private/SLiveLinkHubEditorStatusBar.cpp
@@ -0,0 +1,142 @@
+// Copyright Epic Games, Inc. All Rights Reserved.
+
+#include "SLiveLinkHubEditorStatusBar.h"
+
+#include "Editor.h"
+#include "Features/IModularFeatures.h"
+#include "ILiveLinkClient.h"
+#include "ILiveLinkHubMessagingModule.h"
+#include "ILiveLinkModule.h"
+#include "ILiveLinkSource.h"
+#include "LiveLinkMessageBusSource.h"
+#include "LiveLinkPreset.h"
+#include "Styling/SlateIconFinder.h"
+#include "Styling/SlateStyle.h"
+#include "TimerManager.h"
+#include "Widgets/Images/SImage.h"
+#include "Widgets/Input/SComboButton.h"
+#include "Widgets/Text/STextBlock.h"
+
+#define LOCTEXT_NAMESPACE "LiveLinkHubStatusBar"
+
+void SLiveLinkHubEditorStatusBar::Construct(const FArguments& InArgs)
+{
+	ILiveLinkHubMessagingModule& HubMessagingModule = FModuleManager::Get().GetModuleChecked<ILiveLinkHubMessagingModule>("LiveLinkHubMessaging");
+	HubMessagingModule.OnConnectionEstablished().AddSP(this, &SLiveLinkHubEditorStatusBar::OnHubConnectionEstablished);
+
+	LiveLinkClient = &IModularFeatures::Get().GetModularFeature<ILiveLinkClient>(ILiveLinkClient::ModularFeatureName);
+	LiveLinkClient->OnLiveLinkSourceRemoved().AddSP(this, &SLiveLinkHubEditorStatusBar::OnSourceRemoved);
+
+	constexpr bool bLoop = true;
+	GEditor->GetTimerManager()->SetTimer(TimerHandle, FTimerDelegate::CreateSP(this, &SLiveLinkHubEditorStatusBar::CheckHubConnection), CheckConnectionIntervalSeconds, bLoop);
+
+	ChildSlot
+	[
+		SNew(SComboButton)
+    		.ContentPadding(FMargin(6.0f, 0.0f))
+    		.MenuPlacement(MenuPlacement_AboveAnchor)
+    		.ComboButtonStyle(&FAppStyle::Get().GetWidgetStyle<FComboButtonStyle>("SimpleComboButton"))
+			.HasDownArrow(true)
+			.Visibility(this, &SLiveLinkHubEditorStatusBar::GetVisibility)
+    		.ButtonContent()
+    		[
+				SNew(SHorizontalBox)
+				+ SHorizontalBox::Slot()
+				.AutoWidth()
+				.VAlign(VAlign_Center)
+				.Padding(0, 0, 3, 0)
+				[
+					SNew(SImage)
+					.ColorAndOpacity(this, &SLiveLinkHubEditorStatusBar::GetIconColor)
+					.Image(FSlateIcon("LiveLinkStyle","LiveLinkClient.Common.Icon.Small").GetIcon())
+				]
+				+ SHorizontalBox::Slot()
+                .AutoWidth()
+                .VAlign(VAlign_Center)
+                .Padding(2, 0, 10, 0)
+                [
+                	SNew(STextBlock)
+                	.Text(this, &SLiveLinkHubEditorStatusBar::GetStatusText)
+                	.ToolTipText(this, &SLiveLinkHubEditorStatusBar::GetToolTipText)
+                ]
+    		]
+	];
+}
+
+SLiveLinkHubEditorStatusBar::~SLiveLinkHubEditorStatusBar()
+{
+	if (ILiveLinkHubMessagingModule* HubMessagingModule = FModuleManager::Get().GetModulePtr<ILiveLinkHubMessagingModule>("LiveLinkHubMessaging"))
+	{
+		HubMessagingModule->OnConnectionEstablished().RemoveAll(this);
+	}
+
+	IModularFeatures::Get().OnModularFeatureRegistered().RemoveAll(this);
+
+	if (GEditor && GEditor->IsTimerManagerValid())
+	{
+		GEditor->GetTimerManager()->ClearTimer(TimerHandle);
+	}
+}
+
+FText SLiveLinkHubEditorStatusBar::GetStatusText() const
+{
+	const bool bConnectionActive = ConnectionState == EHubConnectionState::Valid;
+	return bConnectionActive ? LOCTEXT("LiveLinkHubConnectionActive", "LiveLinkHub Connected") : LOCTEXT("LiveLinkHubConnectionError", "LiveLinkHub Error");
+}
+
+FText SLiveLinkHubEditorStatusBar::GetToolTipText() const
+{
+	const bool bConnectionActive = ConnectionState == EHubConnectionState::Valid;
+	return bConnectionActive ? LOCTEXT("LiveLinkHubConnectionActiveTooltip", "LiveLinkHub is connected to this instance of Unreal Editor.") : LOCTEXT("LiveLinkHubConnectionErrorTooltip", "The LiveLinkHub connection was terminated.");
+}
+
+FSlateColor SLiveLinkHubEditorStatusBar::GetIconColor() const
+{
+	const bool bConnectionActive = ConnectionState == EHubConnectionState::Valid;
+	return bConnectionActive ? FSlateColor(FLinearColor::Green) : FSlateColor(FLinearColor::Red);
+}
+
+void SLiveLinkHubEditorStatusBar::OnSourceRemoved(FGuid SourceId)
+{
+	if (HubSourceId == SourceId)
+	{
+		HubSourceId.Invalidate();
+	}
+}
+
+void SLiveLinkHubEditorStatusBar::OnHubConnectionEstablished(FGuid SourceId)
+{
+	bHasConnectedOnce = true;
+
+	HubSourceId = SourceId;
+	ConnectionState = EHubConnectionState::Valid;
+}
+
+void SLiveLinkHubEditorStatusBar::CheckHubConnection()
+{
+	if (!HubSourceId.IsValid())
+	{
+		ConnectionState = Invalid;
+		return;
+	}
+
+	if (LiveLinkClient->IsSourceStillValid(HubSourceId))
+	{
+		ConnectionState = EHubConnectionState::Valid;
+		return;
+	}
+
+	ConnectionState = EHubConnectionState::Timeout;
+}
+
+EVisibility SLiveLinkHubEditorStatusBar::GetVisibility() const
+{
+	if (ConnectionState != EHubConnectionState::Invalid || bHasConnectedOnce)
+	{
+		return EVisibility::Visible;
+	}
+
+	return EVisibility::Collapsed;
+}
+
+#undef LOCTEXT_NAMESPACE /*LiveLinkHubStatusBar*/
diff --git a/Engine/Plugins/Animation/LiveLink/Source/LiveLinkHubEditor/Private/SLiveLinkHubEditorStatusBar.h b/Engine/Plugins/Animation/LiveLink/Source/LiveLinkHubEditor/Private/SLiveLinkHubEditorStatusBar.h
new file mode 100644
--- /dev/null
+++ b/Engine/Plugins/Animation/LiveLink/Source/LiveLinkHubEditor/Private/SLiveLinkHubEditorStatusBar.h
@@ -0,0 +1,69 @@
+// Copyright Epic Games, Inc. All Rights Reserved.
+
+#pragma once
+
+#include "Engine/TimerHandle.h"
+#include "Internationalization/Text.h"
+#include "Templates/SharedPointer.h"
+#include "Types/SlateEnums.h"
+#include "Widgets/DeclarativeSyntaxSupport.h"
+#include "Widgets/SCompoundWidget.h"
+
+class ILiveLinkSource;
+
+enum EHubConnectionState
+{
+	Invalid,
+	Valid,
+	Timeout
+};
+
+/** Widget responsible for showing the current connection status with the hub. Will only appear if a connection has been established with a hub instance. */
+class SLiveLinkHubEditorStatusBar : public SCompoundWidget
+{
+	SLATE_BEGIN_ARGS(SLiveLinkHubEditorStatusBar) {}
+	SLATE_END_ARGS()
+
+	~SLiveLinkHubEditorStatusBar() override;
+
+	void Construct(const FArguments& InArgs);
+
+private:
+	/** Get the status text for the widget. */
+	FText GetStatusText() const;
+	/** Get the tooltip text for the widget. */
+	FText GetToolTipText() const;
+	/** Get the icon color for the widget according to the current status. */
+	FSlateColor GetIconColor() const;
+
+	/** Handler used to update the current connection state. */
+	void OnSourceRemoved(FGuid SourceId);
+
+	/** Handler used to update the connection state and source id when a connection with a hub instance is established. */
+	void OnHubConnectionEstablished(FGuid SourceId);
+
+	/** Check whether the hub connection is still active. */
+	void CheckHubConnection();
+
+	/** Returns the visibility of the widget (Depends on whether we've established connection). */
+	EVisibility GetVisibility() const;
+
+private:
+	/** LiveLink client used to retrieve the status of the hub connection. */
+	class ILiveLinkClient* LiveLinkClient = nullptr;
+
+	/** Whether the editor was able to connect to a hub, will always show the status bar icon after this. */
+	bool bHasConnectedOnce = false;
+
+	/** Id of the livelink source that represents the LLHub. */
+	FGuid HubSourceId;
+
+	/** Current connection state to the hub. */
+	EHubConnectionState ConnectionState = Invalid;
+
+	/** Handle to the timer responsible for triggering CheckHubConnection. */
+	FTimerHandle TimerHandle;
+
+	/** Interval of the timer to check for connection validity. */
+	static constexpr float CheckConnectionIntervalSeconds = 1.0f;
+};
diff --git a/Engine/Plugins/Animation/LiveLink/Source/LiveLinkHubMessaging/Private/LiveLinkHubConnectionManager.h b/Engine/Plugins/Animation/LiveLink/Source/LiveLinkHubMessaging/Private/LiveLinkHubConnectionManager.h
--- a/Engine/Plugins/Animation/LiveLink/Source/LiveLinkHubMessaging/Private/LiveLinkHubConnectionManager.h
+++ b/Engine/Plugins/Animation/LiveLink/Source/LiveLinkHubMessaging/Private/LiveLinkHubConnectionManager.h
@@ -8,6 +8,7 @@
 #include "ILiveLinkModule.h"
 #include "ILiveLinkSource.h"
 #include "LiveLinkHubMessages.h"
+#include "ILiveLinkHubMessagingModule.h"
 #include "LiveLinkHubMessageBusSource.h"
 #include "LiveLinkMessageBusDiscoveryManager.h"
 #include "LiveLinkMessageBusFinder.h"
@@ -121,7 +122,7 @@ private:
 			}
 
 			TSharedPtr<ILiveLinkSource> LiveLinkSource = MakeShared<FLiveLinkHubMessageBusSource>(FText::FromString(PollResult->Name), FText::FromString(PollResult->MachineName), PollResult->Address, PollResult->MachineTimeOffset);
-			LiveLinkClient->AddSource(LiveLinkSource);
+			FGuid SourceId = LiveLinkClient->AddSource(LiveLinkSource);
 
 			// Remove the old source if we're re-establising a connection.
 			if (TSharedPtr<ILiveLinkSource> OldSource = LastAddedSource.Value.Pin())
@@ -135,6 +136,8 @@ private:
 			}
 
 			LastAddedSource = TPair<FMessageAddress, TWeakPtr<ILiveLinkSource>>{ PollResult->Address, LiveLinkSource };
+			ILiveLinkHubMessagingModule& HubMessagingModule = FModuleManager::GetModuleChecked<ILiveLinkHubMessagingModule>("LiveLinkHubMessaging");
+			HubMessagingModule.OnConnectionEstablished().Broadcast(SourceId);
 		}
 		else
 		{
@@ -152,4 +155,4 @@ private:
 class FLiveLinkHubConnectionManager
 {
 };
-#endif
\ No newline at end of file
+#endif
diff --git a/Engine/Plugins/Animation/LiveLink/Source/LiveLinkHubMessaging/Private/LiveLinkHubMessagingModule.cpp b/Engine/Plugins/Animation/LiveLink/Source/LiveLinkHubMessaging/Private/LiveLinkHubMessagingModule.cpp
--- a/Engine/Plugins/Animation/LiveLink/Source/LiveLinkHubMessaging/Private/LiveLinkHubMessagingModule.cpp
+++ b/Engine/Plugins/Animation/LiveLink/Source/LiveLinkHubMessaging/Private/LiveLinkHubMessagingModule.cpp
@@ -1,5 +1,7 @@
 // Copyright Epic Games, Inc. All Rights Reserved.
 
+#include "ILiveLinkHubMessagingModule.h"
+
 #include "CoreMinimal.h"
 #include "LiveLinkHubConnectionManager.h"
 #include "LiveLinkHubMessageBusSourceFactory.h"
@@ -18,7 +20,8 @@
 #define LIVELINKHUB_USE_CONNECTION_MANAGER !WITH_LIVELINK_HUB && WITH_LIVELINK_DISCOVERY_MANAGER_THREAD
 #endif
 
-class FLiveLinkHubMessagingModule : public IModuleInterface
+
+class FLiveLinkHubMessagingModule : public ILiveLinkHubMessagingModule
 {
 	// The connection manager is used to communicate with the hub, so we don't need it when we're running the hub itself.
 public:
@@ -45,6 +48,12 @@ public:
 	}
 	//~ End IModuleInterface
 
+	//~ ILiveLinkHubMessagingModule interface
+	virtual FOnHubConnectionEstablished& OnConnectionEstablished() override
+	{
+		return ConnectionEstablishedDelegate;
+	}
+
 private:
 #if LIVELINKHUB_USE_CONNECTION_MANAGER
 	/** Manages the connection to the live link hub. */
@@ -85,6 +94,9 @@ private:
 
 	/** Handle to the delegate used to filter message bus sources. */
 	FDelegateHandle SourceFilterDelegate;
+
+	/** Delegate called when the connection between a livelink hub and the editor is establshed. */
+	FOnHubConnectionEstablished ConnectionEstablishedDelegate;
 };
 
 
diff --git a/Engine/Plugins/Animation/LiveLink/Source/LiveLinkHubMessaging/Public/ILiveLinkHubMessagingModule.h b/Engine/Plugins/Animation/LiveLink/Source/LiveLinkHubMessaging/Public/ILiveLinkHubMessagingModule.h
new file mode 100644
--- /dev/null
+++ b/Engine/Plugins/Animation/LiveLink/Source/LiveLinkHubMessaging/Public/ILiveLinkHubMessagingModule.h
@@ -0,0 +1,18 @@
+// Copyright Epic Games, Inc. All Rights Reserved.
+
+#pragma once
+
+#include "Delegates/Delegate.h"
+#include "Modules/ModuleInterface.h"
+
+#include "Delegates/DelegateCombinations.h"
+#include "Templates/SharedPointer.h"
+
+DECLARE_MULTICAST_DELEGATE_OneParam(FOnHubConnectionEstablished, FGuid SourceId);
+
+class LIVELINKHUBMESSAGING_API ILiveLinkHubMessagingModule : public IModuleInterface
+{
+public:
+	/** Delegate called when a connection is established to a livelink hub. */
+	virtual FOnHubConnectionEstablished& OnConnectionEstablished() = 0;
+};
diff --git a/Engine/Source/Runtime/LiveLinkInterface/Public/ILiveLinkClient.h b/Engine/Source/Runtime/LiveLinkInterface/Public/ILiveLinkClient.h
--- a/Engine/Source/Runtime/LiveLinkInterface/Public/ILiveLinkClient.h
+++ b/Engine/Source/Runtime/LiveLinkInterface/Public/ILiveLinkClient.h
@@ -165,6 +165,9 @@ public:
 	/** Get the machine name of the source. */
 	virtual FText GetSourceMachineName(FGuid EntryGuid) const = 0;
 
+	/** Returns whether the Source is connected to its data provider and can still push valid data. */
+	virtual bool IsSourceStillValid(FGuid EntryGuid) const = 0;
+
 	/** Push static data for a specific subject for a certain role. This will clear all buffered frames */
 	virtual void PushSubjectStaticData_AnyThread(const FLiveLinkSubjectKey& SubjectKey, TSubclassOf<ULiveLinkRole> Role, FLiveLinkStaticDataStruct&& StaticData) = 0;
 
