[Backout] - CL30895823 (backed out due to errors)
Original CL Desc
-----------------------------------------------------------------
Add livelink hub status bar to the editor (Only appears when we have an active connection to a hub instance)

diff --git a/Engine/Plugins/Animation/LiveLink/LiveLink.uplugin b/Engine/Plugins/Animation/LiveLink/LiveLink.uplugin
--- a/Engine/Plugins/Animation/LiveLink/LiveLink.uplugin
+++ b/Engine/Plugins/Animation/LiveLink/LiveLink.uplugin
@@ -45,12 +45,7 @@
 			"Name": "LiveLinkHub",
 			"Type": "EditorAndProgram",
 			"ProgramAllowList": [ "LiveLinkHub" ],
-			"LoadingPhase": "Default"
-		},
-		{
-			"Name": "LiveLinkHubEditor",
-			"Type": "Editor",
-			"LoadingPhase": "Default"
+			"LoadingPhase": "PreDefault"
 		},
 		{
 			"Name": "LiveLinkComponents",
@@ -91,4 +86,4 @@
 			"ProgramAllowList": [ "LiveLinkHub" ]
 		}
 	]
-}
+}
\ No newline at end of file
diff --git a/Engine/Plugins/Animation/LiveLink/Source/LiveLink/Private/LiveLinkMessageBusSource.cpp b/Engine/Plugins/Animation/LiveLink/Source/LiveLink/Private/LiveLinkMessageBusSource.cpp
--- a/Engine/Plugins/Animation/LiveLink/Source/LiveLink/Private/LiveLinkMessageBusSource.cpp
+++ b/Engine/Plugins/Animation/LiveLink/Source/LiveLink/Private/LiveLinkMessageBusSource.cpp
@@ -18,20 +18,6 @@
 #include "MessageEndpointBuilder.h"
 #include "Misc/App.h"
 
-FText FLiveLinkMessageBusSource::ValidSourceStatus()
-{
-	return NSLOCTEXT("LiveLinkMessageBusSource", "ActiveStatus", "Active");
-}
-
-FText FLiveLinkMessageBusSource::InvalidSourceStatus()
-{
-	return NSLOCTEXT("LiveLinkMessageBusSource", "InvalidConnection", "Waiting for connection");
-}
-
-FText FLiveLinkMessageBusSource::TimeoutSourceStatus()
-{
-	return NSLOCTEXT("LiveLinkMessageBusSource", "TimeoutStatus", "Not responding");
-}
 
 FLiveLinkMessageBusSource::FLiveLinkMessageBusSource(const FText& InSourceType, const FText& InSourceMachineName, const FMessageAddress& InConnectionAddress, double InMachineTimeOffset)
 	: ConnectionAddress(InConnectionAddress)
@@ -220,13 +206,13 @@ FText FLiveLinkMessageBusSource::GetSourceStatus() const
 {
 	if (!ConnectionAddress.IsValid())
 	{
-		return InvalidSourceStatus();
+		return NSLOCTEXT("LiveLinkMessageBusSource", "InvalidConnection", "Waiting for connection");
 	}
 	else if (IsSourceStillValid())
 	{
-		return ValidSourceStatus();
+		return NSLOCTEXT("LiveLinkMessageBusSource", "ActiveStatus", "Active");
 	}
-	return TimeoutSourceStatus();
+	return NSLOCTEXT("LiveLinkMessageBusSource", "TimeoutStatus", "Not responding");
 }
 
 TSubclassOf<ULiveLinkSourceSettings> FLiveLinkMessageBusSource::GetSettingsClass() const
diff --git a/Engine/Plugins/Animation/LiveLink/Source/LiveLink/Public/LiveLinkClient.h b/Engine/Plugins/Animation/LiveLink/Source/LiveLink/Public/LiveLinkClient.h
--- a/Engine/Plugins/Animation/LiveLink/Source/LiveLink/Public/LiveLinkClient.h
+++ b/Engine/Plugins/Animation/LiveLink/Source/LiveLink/Public/LiveLinkClient.h
@@ -80,7 +80,6 @@ public:
 	virtual FText GetSourceType(FGuid EntryGuid) const override;
 	virtual FText GetSourceStatus(FGuid EntryGuid) const override;
 	virtual FText GetSourceMachineName(FGuid EntryGuid) const override;
-	virtual bool IsSourceStillValid(FGuid EntryGuid) const override;
 
 
 	virtual void PushSubjectStaticData_AnyThread(const FLiveLinkSubjectKey& SubjectKey, TSubclassOf<ULiveLinkRole> Role, FLiveLinkStaticDataStruct&& StaticData) override;
@@ -163,6 +162,7 @@ public:
 
 	FLiveLinkSubjectTimeSyncData GetTimeSyncData(FLiveLinkSubjectName SubjectName);
 
+	bool IsSourceStillValid(FGuid EntryGuid) const;
 	UE_DEPRECATED(4.23, "FLiveLinkClient::GetSourceTypeForEntry is deprecated. Please use GetSourceType instead!")
 	FText GetSourceTypeForEntry(FGuid EntryGuid) const { return GetSourceType(EntryGuid); }
 	UE_DEPRECATED(4.23, "FLiveLinkClient::GetMachineNameForEntry is deprecated. Please use GetSourceMachineName instead!")
diff --git a/Engine/Plugins/Animation/LiveLink/Source/LiveLink/Public/LiveLinkMessageBusSource.h b/Engine/Plugins/Animation/LiveLink/Source/LiveLink/Public/LiveLinkMessageBusSource.h
--- a/Engine/Plugins/Animation/LiveLink/Source/LiveLink/Public/LiveLinkMessageBusSource.h
+++ b/Engine/Plugins/Animation/LiveLink/Source/LiveLink/Public/LiveLinkMessageBusSource.h
@@ -21,12 +21,6 @@ struct FLiveLinkClearSubject;
 class LIVELINK_API FLiveLinkMessageBusSource : public ILiveLinkSource
 {
 public:
-	/** Text description for a valid source. */
-	static FText ValidSourceStatus();
-	/** Text description for an invalid source. */
-	static FText InvalidSourceStatus();
-	/** Text description for a source that has timed out. */
-	static FText TimeoutSourceStatus();
 
 	FLiveLinkMessageBusSource(const FText& InSourceType, const FText& InSourceMachineName, const FMessageAddress& InConnectionAddress, double InMachineTimeOffset);
 
diff --git a/Engine/Plugins/Animation/LiveLink/Source/LiveLinkHub/Private/Recording/SLiveLinkHubRecordingView.h b/Engine/Plugins/Animation/LiveLink/Source/LiveLinkHub/Private/Recording/SLiveLinkHubRecordingView.h
--- a/Engine/Plugins/Animation/LiveLink/Source/LiveLinkHub/Private/Recording/SLiveLinkHubRecordingView.h
+++ b/Engine/Plugins/Animation/LiveLink/Source/LiveLinkHub/Private/Recording/SLiveLinkHubRecordingView.h
@@ -7,11 +7,7 @@
 #include "Delegates/Delegate.h"
 #include "Delegates/DelegateCombinations.h"
 #include "Styling/SlateTypes.h"
-#include "Widgets/Images/SImage.h"
 #include "Widgets/Input/SButton.h"
-#include "Widgets/Layout/SBox.h"
-#include "Widgets/SBoxPanel.h"
-#include "Widgets/Text/STextBlock.h"
 
 class FLiveLinkHub;
 
diff --git a/Engine/Plugins/Animation/LiveLink/Source/LiveLinkHubEditor/LiveLinkHubEditor.Build.cs b/Engine/Plugins/Animation/LiveLink/Source/LiveLinkHubEditor/LiveLinkHubEditor.Build.cs
deleted file mode 100644
diff --git a/Engine/Plugins/Animation/LiveLink/Source/LiveLinkHubEditor/Private/LiveLinkHubEditorModule.cpp b/Engine/Plugins/Animation/LiveLink/Source/LiveLinkHubEditor/Private/LiveLinkHubEditorModule.cpp
deleted file mode 100644
diff --git a/Engine/Plugins/Animation/LiveLink/Source/LiveLinkHubEditor/Private/LiveLinkHubEditorModule.h b/Engine/Plugins/Animation/LiveLink/Source/LiveLinkHubEditor/Private/LiveLinkHubEditorModule.h
deleted file mode 100644
diff --git a/Engine/Plugins/Animation/LiveLink/Source/LiveLinkHubEditor/Private/SLiveLinkHubEditorStatusBar.cpp b/Engine/Plugins/Animation/LiveLink/Source/LiveLinkHubEditor/Private/SLiveLinkHubEditorStatusBar.cpp
deleted file mode 100644
diff --git a/Engine/Plugins/Animation/LiveLink/Source/LiveLinkHubEditor/Private/SLiveLinkHubEditorStatusBar.h b/Engine/Plugins/Animation/LiveLink/Source/LiveLinkHubEditor/Private/SLiveLinkHubEditorStatusBar.h
deleted file mode 100644
diff --git a/Engine/Plugins/Animation/LiveLink/Source/LiveLinkHubMessaging/Private/LiveLinkHubConnectionManager.h b/Engine/Plugins/Animation/LiveLink/Source/LiveLinkHubMessaging/Private/LiveLinkHubConnectionManager.h
--- a/Engine/Plugins/Animation/LiveLink/Source/LiveLinkHubMessaging/Private/LiveLinkHubConnectionManager.h
+++ b/Engine/Plugins/Animation/LiveLink/Source/LiveLinkHubMessaging/Private/LiveLinkHubConnectionManager.h
@@ -8,7 +8,6 @@
 #include "ILiveLinkModule.h"
 #include "ILiveLinkSource.h"
 #include "LiveLinkHubMessages.h"
-#include "ILiveLinkHubMessagingModule.h"
 #include "LiveLinkHubMessageBusSource.h"
 #include "LiveLinkMessageBusDiscoveryManager.h"
 #include "LiveLinkMessageBusFinder.h"
@@ -122,7 +121,7 @@ private:
 			}
 
 			TSharedPtr<ILiveLinkSource> LiveLinkSource = MakeShared<FLiveLinkHubMessageBusSource>(FText::FromString(PollResult->Name), FText::FromString(PollResult->MachineName), PollResult->Address, PollResult->MachineTimeOffset);
-			FGuid SourceId = LiveLinkClient->AddSource(LiveLinkSource);
+			LiveLinkClient->AddSource(LiveLinkSource);
 
 			// Remove the old source if we're re-establising a connection.
 			if (TSharedPtr<ILiveLinkSource> OldSource = LastAddedSource.Value.Pin())
@@ -136,8 +135,6 @@ private:
 			}
 
 			LastAddedSource = TPair<FMessageAddress, TWeakPtr<ILiveLinkSource>>{ PollResult->Address, LiveLinkSource };
-			ILiveLinkHubMessagingModule& HubMessagingModule = FModuleManager::GetModuleChecked<ILiveLinkHubMessagingModule>("LiveLinkHubMessaging");
-			HubMessagingModule.OnConnectionEstablished().Broadcast(SourceId);
 		}
 		else
 		{
@@ -155,4 +152,4 @@ private:
 class FLiveLinkHubConnectionManager
 {
 };
-#endif
+#endif
\ No newline at end of file
diff --git a/Engine/Plugins/Animation/LiveLink/Source/LiveLinkHubMessaging/Private/LiveLinkHubMessagingModule.cpp b/Engine/Plugins/Animation/LiveLink/Source/LiveLinkHubMessaging/Private/LiveLinkHubMessagingModule.cpp
--- a/Engine/Plugins/Animation/LiveLink/Source/LiveLinkHubMessaging/Private/LiveLinkHubMessagingModule.cpp
+++ b/Engine/Plugins/Animation/LiveLink/Source/LiveLinkHubMessaging/Private/LiveLinkHubMessagingModule.cpp
@@ -1,7 +1,5 @@
 // Copyright Epic Games, Inc. All Rights Reserved.
 
-#include "ILiveLinkHubMessagingModule.h"
-
 #include "CoreMinimal.h"
 #include "LiveLinkHubConnectionManager.h"
 #include "LiveLinkHubMessageBusSourceFactory.h"
@@ -20,8 +18,7 @@
 #define LIVELINKHUB_USE_CONNECTION_MANAGER !WITH_LIVELINK_HUB && WITH_LIVELINK_DISCOVERY_MANAGER_THREAD
 #endif
 
-
-class FLiveLinkHubMessagingModule : public ILiveLinkHubMessagingModule
+class FLiveLinkHubMessagingModule : public IModuleInterface
 {
 	// The connection manager is used to communicate with the hub, so we don't need it when we're running the hub itself.
 public:
@@ -48,12 +45,6 @@ public:
 	}
 	//~ End IModuleInterface
 
-	//~ ILiveLinkHubMessagingModule interface
-	virtual FOnHubConnectionEstablished& OnConnectionEstablished() override
-	{
-		return ConnectionEstablishedDelegate;
-	}
-
 private:
 #if LIVELINKHUB_USE_CONNECTION_MANAGER
 	/** Manages the connection to the live link hub. */
@@ -94,9 +85,6 @@ private:
 
 	/** Handle to the delegate used to filter message bus sources. */
 	FDelegateHandle SourceFilterDelegate;
-
-	/** Delegate called when the connection between a livelink hub and the editor is establshed. */
-	FOnHubConnectionEstablished ConnectionEstablishedDelegate;
 };
 
 
diff --git a/Engine/Plugins/Animation/LiveLink/Source/LiveLinkHubMessaging/Public/ILiveLinkHubMessagingModule.h b/Engine/Plugins/Animation/LiveLink/Source/LiveLinkHubMessaging/Public/ILiveLinkHubMessagingModule.h
deleted file mode 100644
diff --git a/Engine/Source/Runtime/LiveLinkInterface/Public/ILiveLinkClient.h b/Engine/Source/Runtime/LiveLinkInterface/Public/ILiveLinkClient.h
--- a/Engine/Source/Runtime/LiveLinkInterface/Public/ILiveLinkClient.h
+++ b/Engine/Source/Runtime/LiveLinkInterface/Public/ILiveLinkClient.h
@@ -165,9 +165,6 @@ public:
 	/** Get the machine name of the source. */
 	virtual FText GetSourceMachineName(FGuid EntryGuid) const = 0;
 
-	/** Returns whether the Source is connected to its data provider and can still push valid data. */
-	virtual bool IsSourceStillValid(FGuid EntryGuid) const = 0;
-
 	/** Push static data for a specific subject for a certain role. This will clear all buffered frames */
 	virtual void PushSubjectStaticData_AnyThread(const FLiveLinkSubjectKey& SubjectKey, TSubclassOf<ULiveLinkRole> Role, FLiveLinkStaticDataStruct&& StaticData) = 0;
 
