Fix Mac horde issue 420797. Remove pragma in .cpp file.

diff --git a/Engine/Plugins/Animation/LiveLink/Source/LiveLinkHubEditor/Private/LiveLinkHubEditorModule.cpp b/Engine/Plugins/Animation/LiveLink/Source/LiveLinkHubEditor/Private/LiveLinkHubEditorModule.cpp
--- a/Engine/Plugins/Animation/LiveLink/Source/LiveLinkHubEditor/Private/LiveLinkHubEditorModule.cpp
+++ b/Engine/Plugins/Animation/LiveLink/Source/LiveLinkHubEditor/Private/LiveLinkHubEditorModule.cpp
@@ -1,5 +1,4 @@
 // Copyright Epic Games, Inc. All Rights Reserved.
-#pragma once
 
 #include "LiveLinkHubEditorModule.h"
 
