LiveLinkHub: Add playback controls functionality and UI.

diff --git a/Engine/Plugins/Animation/LiveLink/Source/LiveLink/Internal/LiveLinkSubject.h b/Engine/Plugins/Animation/LiveLink/Source/LiveLink/Internal/LiveLinkSubject.h
--- a/Engine/Plugins/Animation/LiveLink/Source/LiveLink/Internal/LiveLinkSubject.h
+++ b/Engine/Plugins/Animation/LiveLink/Source/LiveLink/Internal/LiveLinkSubject.h
@@ -102,7 +102,7 @@ public:
 
 	void CacheSettings(ULiveLinkSourceSettings* SourceSetting, ULiveLinkSubjectSettings* SubjectSetting);
 
-	ELiveLinkSourceMode GetMode() const { return CachedSettings.SourceMode; }
+	ELiveLinkSourceMode GetMode() const;
 	FLiveLinkSubjectTimeSyncData GetTimeSyncData();
 	bool IsTimeSynchronized() const;
 
@@ -198,6 +198,9 @@ private:
 	// Connection settings specified by user
 	FLiveLinkCachedSettings CachedSettings;
 
+	// Override mode, determined by frame data
+	TOptional<ELiveLinkSourceMode> ModeOverride;
+
 	// Last time a frame was pushed
 	double LastPushTime = 0.0;
 
diff --git a/Engine/Plugins/Animation/LiveLink/Source/LiveLink/Private/LiveLinkSubject.cpp b/Engine/Plugins/Animation/LiveLink/Source/LiveLink/Private/LiveLinkSubject.cpp
--- a/Engine/Plugins/Animation/LiveLink/Source/LiveLink/Private/LiveLinkSubject.cpp
+++ b/Engine/Plugins/Animation/LiveLink/Source/LiveLink/Private/LiveLinkSubject.cpp
@@ -368,8 +368,28 @@ void FLiveLinkSubject::AddFrameData(FLiveLinkFrameDataStruct&& InFrameData)
 		RemoveFrames(1);
 	}
 
+	ELiveLinkSourceMode Mode = CachedSettings.SourceMode;
+	ModeOverride.Reset();
+	
+	// Check if this frame should be force synced, IE go to frame immediately and discard other frames.
+	if (const FString* ForceSyncStr = InFrameData.IsValid()
+									  ? InFrameData.GetBaseData()->MetaData.StringMetaData.Find(TEXT("ForceSync"))
+									  : nullptr)
+	{
+		Mode = ELiveLinkSourceMode::Latest;
+		bool bForceSync = false;
+		LexFromString(bForceSync, **ForceSyncStr);
+
+		if (bForceSync)
+		{
+			ModeOverride = Mode;
+			// todo: should we only remove frames if the timestamp is less than the new frame? Can they arrive out of order?
+			RemoveFrames(FrameData.Num());
+		}
+	}
+
 	int32 FrameIndex = INDEX_NONE;
-	switch (CachedSettings.SourceMode)
+	switch (Mode)
 	{
 	case ELiveLinkSourceMode::EngineTime:
 		FrameIndex = FindNewFrame_WorldTime(InFrameData.GetBaseData()->WorldTime);
@@ -1082,6 +1102,11 @@ void FLiveLinkSubject::CacheSettings(ULiveLinkSourceSettings* SourceSetting, ULi
 	}
 }
 
+ELiveLinkSourceMode FLiveLinkSubject::GetMode() const
+{
+	return ModeOverride.Get(CachedSettings.SourceMode);
+}
+
 FLiveLinkSubjectTimeSyncData FLiveLinkSubject::GetTimeSyncData()
 {
 	FLiveLinkSubjectTimeSyncData SyncData;
diff --git a/Engine/Plugins/Animation/LiveLink/Source/LiveLinkHub/Private/LiveLinkHubClient.cpp b/Engine/Plugins/Animation/LiveLink/Source/LiveLinkHub/Private/LiveLinkHubClient.cpp
--- a/Engine/Plugins/Animation/LiveLink/Source/LiveLinkHub/Private/LiveLinkHubClient.cpp
+++ b/Engine/Plugins/Animation/LiveLink/Source/LiveLinkHub/Private/LiveLinkHubClient.cpp
@@ -383,7 +383,16 @@ void FLiveLinkHubClient::PushSubjectStaticData_AnyThread(const FLiveLinkSubjectK
 		// Note that SetStaticData may not be needed on the hub but we might need it for interpolation or preprocessing in the future.
 		FFunctionGraphTask::CreateAndDispatchWhenReady([this, LiveLinkSubject, SubjectKey, Role, StaticData = MoveTemp(InStaticData)]() mutable
 			{
-				LiveLinkSubject->SetStaticData(Role, MoveTemp(StaticData));
+				FScopeLock Lock(&CollectionAccessCriticalSection);
+				// Validate live link subject hasn't changed
+				if (FLiveLinkCollectionSubjectItem* SubjectItem = Collection->FindSubject(SubjectKey))
+				{
+					const FLiveLinkSubject* CurrentLiveLinkSubject = SubjectItem->GetLiveSubject();
+					if (LiveLinkSubject == CurrentLiveLinkSubject)
+					{
+						LiveLinkSubject->SetStaticData(Role, MoveTemp(StaticData));
+					}
+				}
 			}, TStatId(), nullptr, ENamedThreads::GameThread);
 	}
 
diff --git a/Engine/Plugins/Animation/LiveLink/Source/LiveLinkHub/Private/Recording/Implementations/LiveLinkUAssetRecordingPlayer.cpp b/Engine/Plugins/Animation/LiveLink/Source/LiveLinkHub/Private/Recording/Implementations/LiveLinkUAssetRecordingPlayer.cpp
--- a/Engine/Plugins/Animation/LiveLink/Source/LiveLinkHub/Private/Recording/Implementations/LiveLinkUAssetRecordingPlayer.cpp
+++ b/Engine/Plugins/Animation/LiveLink/Source/LiveLinkHub/Private/Recording/Implementations/LiveLinkUAssetRecordingPlayer.cpp
@@ -8,6 +8,8 @@
 class FLiveLinkPlaybackTrackIterator
 {
 public:
+	virtual ~FLiveLinkPlaybackTrackIterator() = default;
+
 	FLiveLinkPlaybackTrackIterator(FLiveLinkPlaybackTrack& InTrack, int32 InInitialIndex)
 		: Track(InTrack)
 		, FrameIndex(InInitialIndex)
@@ -28,7 +30,7 @@ public:
 
 	double FrameTimestamp() const
 	{
-		return Track.Timestamps[FrameIndex];
+		return FrameIndex >= 0 && FrameIndex < Track.Timestamps.Num() ? Track.Timestamps[FrameIndex] : 0.f;
 	}
 
 	const FInstancedStruct& FrameData() const
@@ -41,29 +43,67 @@ public:
 		return FrameIndex;
 	}
 
+protected:
+	/* @Return True if there are more vertices on the component */
+	virtual bool HasMoreFrames() const = 0;
+
+	/* Advances to the next frame */
+	virtual void Advance() = 0;
+
+protected:
+	/** Track that's currently being iterated. */
+	FLiveLinkPlaybackTrack& Track;
+	/** "Playhead" for this track */
+	int32 FrameIndex = 0;
+};
+
+class FLiveLinkPlaybackTrackForwardIterator : public FLiveLinkPlaybackTrackIterator
+{
+public:
+	FLiveLinkPlaybackTrackForwardIterator(FLiveLinkPlaybackTrack& InTrack, int32 InInitialIndex)
+		: FLiveLinkPlaybackTrackIterator(InTrack, InInitialIndex)
+	{
+	}
+
 private:
 	/* @Return True if there are more vertices on the component */
-	bool HasMoreFrames() const
+	virtual bool HasMoreFrames() const override
 	{
 		return FrameIndex < Track.Timestamps.Num() && FrameIndex < Track.FrameData.Num();
 	}
 
 	/* Advances to the next frame */
-	void Advance()
+	virtual void Advance() override
 	{
 		++FrameIndex;
 	}
+};
+
+class FLiveLinkPlaybackTrackReverseIterator : public FLiveLinkPlaybackTrackIterator
+{
+public:
+	FLiveLinkPlaybackTrackReverseIterator(FLiveLinkPlaybackTrack& InTrack, int32 InInitialIndex)
+		: FLiveLinkPlaybackTrackIterator(InTrack, InInitialIndex)
+	{
+	}
 
 private:
-	/** Track that's currently being iterated. */
-	FLiveLinkPlaybackTrack& Track;
-	/** "Playhead" for this track */
-	int32 FrameIndex = 0;
+	/* @Return True if there are more vertices on the component */
+	virtual bool HasMoreFrames() const override
+	{
+		return FrameIndex >= 0;
+	}
+
+	/* Advances to the next frame */
+	virtual void Advance() override
+	{
+		--FrameIndex;
+	}
 };
 
 void FLiveLinkPlaybackTrack::GetFramesUntil(double InPlayhead, TArray<FLiveLinkRecordedFrame>& OutFrames)
 {
-	for (auto It = FLiveLinkPlaybackTrackIterator(*this, LastReadIndex + 1); It; ++It)
+	for (FLiveLinkPlaybackTrackForwardIterator It = FLiveLinkPlaybackTrackForwardIterator(*this, LastReadIndex + 1); It; ++It)
 	{
 		if (It.FrameTimestamp() > InPlayhead)
 		{
@@ -76,11 +116,89 @@ void FLiveLinkPlaybackTrack::GetFramesUntil(double InPlayhead, TArray<FLiveLinkR
 		FrameToPlay.Data = It.FrameData();
 		FrameToPlay.SubjectKey = SubjectKey;
 		FrameToPlay.LiveLinkRole = LiveLinkRole;
+		FrameToPlay.FrameIndex = LastReadIndex;
+
+		OutFrames.Add(MoveTemp(FrameToPlay));
+	}
+}
+
+void FLiveLinkPlaybackTrack::GetFramesUntilReverse(double InPlayhead, TArray<FLiveLinkRecordedFrame>& OutFrames)
+{
+	if (LastReadIndex == INDEX_NONE)
+	{
+		LastReadIndex = FrameData.Num();
+	}
+
+	// We need to look up what the last frame would be if this was running forward, and then end on that frame.
+	// Since we iterate in reverse, but all other operations like GoToFrame use forward look ahead, it's possible the time stamp comparison
+	// will differ by a frame with a reverse look up. There's probably a better way of handling this.
+	const int32 FinalFrameIndex = PlayheadToFrameIndex(InPlayhead);
+	
+	for (FLiveLinkPlaybackTrackReverseIterator It = FLiveLinkPlaybackTrackReverseIterator(*this, LastReadIndex - 1); It; ++It)
+	{
+		if (FinalFrameIndex == LastReadIndex)
+		{
+			break;
+		}
+
+		LastReadIndex = It.CurrentIndex();
+
+		FLiveLinkRecordedFrame FrameToPlay;
+		FrameToPlay.Data = It.FrameData();
+		FrameToPlay.SubjectKey = SubjectKey;
+		FrameToPlay.LiveLinkRole = LiveLinkRole;
+		FrameToPlay.FrameIndex = LastReadIndex;
 
 		OutFrames.Add(MoveTemp(FrameToPlay));
 	}
 }
 
+bool FLiveLinkPlaybackTrack::TryGetFrame(int32 InIndex, FLiveLinkRecordedFrame& OutFrame)
+{
+	if (InIndex >= 0 && InIndex < FrameData.Num())
+	{
+		LastReadIndex = InIndex;
+		
+		FLiveLinkRecordedFrame FrameToPlay;
+		FrameToPlay.Data = FrameData[InIndex];
+		FrameToPlay.SubjectKey = SubjectKey;
+		FrameToPlay.LiveLinkRole = LiveLinkRole;
+		FrameToPlay.FrameIndex = LastReadIndex;
+
+		OutFrame = MoveTemp(FrameToPlay);
+		return true;
+	}
+	
+	return false;
+}
+
+int32 FLiveLinkPlaybackTrack::PlayheadToFrameIndex(double InPlayhead)
+{
+	int32 CurrentIndex = 0;
+
+	for (int32 Idx = 0; Idx < Timestamps.Num(); ++Idx)
+	{
+		if (Timestamps[Idx] > InPlayhead)
+		{
+			break;
+		}
+
+		CurrentIndex = Idx;
+	}
+
+	return CurrentIndex;
+}
+
+double FLiveLinkPlaybackTrack::FrameIndexToPlayhead(int32 InIndex)
+{
+	if (InIndex >= 0 && InIndex < Timestamps.Num())
+	{
+		return Timestamps[InIndex];
+	}
+
+	return INDEX_NONE;
+}
+
 TArray<FLiveLinkRecordedFrame> FLiveLinkPlaybackTracks::FetchNextFrames(double Playhead)
 {
 	TArray<FLiveLinkRecordedFrame> NextFrames;
@@ -97,11 +215,74 @@ TArray<FLiveLinkRecordedFrame> FLiveLinkPlaybackTracks::FetchNextFrames(double P
 	return NextFrames;
 }
 
-void FLiveLinkPlaybackTracks::Restart()
+TArray<FLiveLinkRecordedFrame> FLiveLinkPlaybackTracks::FetchPreviousFrames(double Playhead)
+{
+	TArray<FLiveLinkRecordedFrame> PreviousFrames;
+
+	if (Tracks.Num())
+	{
+		// todo: sort frames by timestamp
+		for (FLiveLinkPlaybackTrack& Track : Tracks)
+		{
+			Track.GetFramesUntilReverse(Playhead, PreviousFrames);
+		}
+	}
+
+	return PreviousFrames;
+}
+
+TArray<FLiveLinkRecordedFrame> FLiveLinkPlaybackTracks::FetchNextFramesAtIndex(int32 FrameIndex)
+{
+	TArray<FLiveLinkRecordedFrame> NextFrames;
+
+	if (FrameIndex >= 0)
+	{
+		for (FLiveLinkPlaybackTrack& Track : Tracks)
+		{
+			FLiveLinkRecordedFrame Frame;
+			if (Track.TryGetFrame(FrameIndex, Frame))
+			{
+				NextFrames.Add(MoveTemp(Frame));
+			}
+		}
+	}
+
+	return NextFrames;
+}
+
+int32 FLiveLinkPlaybackTracks::PlayheadToFrameIndex(double InPlayhead)
+{
+	for (FLiveLinkPlaybackTrack& Track : Tracks)
+	{
+		// todo: Is this the best way to determine if this is keyframe data and not static data?
+		if (Track.LiveLinkRole == nullptr)
+		{
+			return Track.PlayheadToFrameIndex(InPlayhead);
+		}
+	}
+
+	return INDEX_NONE;
+}
+
+double FLiveLinkPlaybackTracks::FrameIndexToPlayhead(int32 InIndex)
+{
+	for (FLiveLinkPlaybackTrack& Track : Tracks)
+	{
+		// todo: Is this the best way to determine if this is keyframe data and not static data?
+		if (Track.LiveLinkRole == nullptr)
+		{
+			return Track.FrameIndexToPlayhead(InIndex);
+		}
+	}
+
+	return INDEX_NONE;
+}
+
+void FLiveLinkPlaybackTracks::Restart(int32 InIndex)
 {
 	for (FLiveLinkPlaybackTrack& Track : Tracks)
 	{
-		Track.Restart();
+		Track.Restart(InIndex);
 	}
 }
 
diff --git a/Engine/Plugins/Animation/LiveLink/Source/LiveLinkHub/Private/Recording/Implementations/LiveLinkUAssetRecordingPlayer.h b/Engine/Plugins/Animation/LiveLink/Source/LiveLinkHub/Private/Recording/Implementations/LiveLinkUAssetRecordingPlayer.h
--- a/Engine/Plugins/Animation/LiveLink/Source/LiveLinkHub/Private/Recording/Implementations/LiveLinkUAssetRecordingPlayer.h
+++ b/Engine/Plugins/Animation/LiveLink/Source/LiveLinkHub/Private/Recording/Implementations/LiveLinkUAssetRecordingPlayer.h
@@ -9,11 +9,32 @@
 /** Playback track that holds recorded data for a given subject. */
 struct FLiveLinkPlaybackTrack
 {
+	/** Retrieve all frames from last read index to the new playhead, forward looking. */
 	void GetFramesUntil(double InPlayhead, TArray<FLiveLinkRecordedFrame>& OutFrames);
 
-	void Restart()
+	/** Retrieve all frames from last read index to the new playhead, reverse looking. */
+	void GetFramesUntilReverse(double InPlayhead, TArray<FLiveLinkRecordedFrame>& OutFrames);
+	
+	/** Retrieve the frame at the read index. */
+	bool TryGetFrame(int32 InIndex, FLiveLinkRecordedFrame& OutFrame);
+
+	/**
+	 * Convert the playhead time to a frame index.
+	 * @param InPlayhead The play head to convert to a frame index.
+	 * @return The frame index or INDEX_NONE.
+	 */
+	int32 PlayheadToFrameIndex(double InPlayhead);
+
+	/**
+	 * Convert the frameindex to a playhead.
+	 * @return The frame index or INDEX_NONE.
+	 */
+	double FrameIndexToPlayhead(int32 InIndex);
+
+	/** Reset the LastReadIndex. */
+	void Restart(int32 NewIndex = INDEX_NONE)
 	{
-		LastReadIndex = -1;
+		LastReadIndex = NewIndex < FrameData.Num() && NewIndex < Timestamps.Num() ? NewIndex : INDEX_NONE;
 	}
 
 	/** Frame data to read. */
@@ -36,7 +57,19 @@ struct FLiveLinkPlaybackTracks
 	/** Get the next frames */
 	TArray<FLiveLinkRecordedFrame> FetchNextFrames(double Playhead);
 
-	void Restart();
+	/** Get the previous frames as if going in reverse */
+	TArray<FLiveLinkRecordedFrame> FetchPreviousFrames(double Playhead);
+	
+	/** Get the next frame(s) at the index */
+	TArray<FLiveLinkRecordedFrame> FetchNextFramesAtIndex(int32 FrameIndex);
+
+	/** Convert the playhead to a frame index */
+	int32 PlayheadToFrameIndex(double InPlayhead);
+
+	/** Convert the index to a playhead */
+	double FrameIndexToPlayhead(int32 InIndex);
+	
+	void Restart(int32 InIndex);
 
 public:
 	/** LiveLink tracks to playback. */
@@ -48,14 +81,34 @@ class FLiveLinkUAssetRecordingPlayer : public ILiveLinkRecordingPlayer
 public:
 	void PreparePlayback(const class ULiveLinkRecording* CurrentRecording);
 
-	virtual TArray<FLiveLinkRecordedFrame> FetchNextFrames(double Playhead) override
+	virtual TArray<FLiveLinkRecordedFrame> FetchNextFramesAtTimestamp(double Playhead) override
 	{
 		return CurrentRecordingPlayback.FetchNextFrames(Playhead);
 	}
 
-	virtual void RestartPlayback() override
+	virtual TArray<FLiveLinkRecordedFrame> FetchPreviousFramesAtTimestamp(double Playhead) override
+	{
+		return CurrentRecordingPlayback.FetchPreviousFrames(Playhead);
+	}
+
+	virtual TArray<FLiveLinkRecordedFrame> FetchNextFramesAtIndex(int32 FrameIndex) override
+	{
+		return CurrentRecordingPlayback.FetchNextFramesAtIndex(FrameIndex);
+	}
+
+	virtual int32 PlayheadToFrameIndex(double InPlayhead, bool bReverse) override
+	{
+		return CurrentRecordingPlayback.PlayheadToFrameIndex(InPlayhead);
+	}
+
+	virtual double FrameIndexToPlayhead(int32 InIndex) override
+	{
+		return CurrentRecordingPlayback.FrameIndexToPlayhead(InIndex);
+	}
+
+	virtual void RestartPlayback(int32 InIndex) override
 	{
-		CurrentRecordingPlayback.Restart();
+		CurrentRecordingPlayback.Restart(InIndex);
 	}
 	
 private:
diff --git a/Engine/Plugins/Animation/LiveLink/Source/LiveLinkHub/Private/Recording/LiveLinkHubPlaybackController.cpp b/Engine/Plugins/Animation/LiveLink/Source/LiveLinkHub/Private/Recording/LiveLinkHubPlaybackController.cpp
--- a/Engine/Plugins/Animation/LiveLink/Source/LiveLinkHub/Private/Recording/LiveLinkHubPlaybackController.cpp
+++ b/Engine/Plugins/Animation/LiveLink/Source/LiveLinkHub/Private/Recording/LiveLinkHubPlaybackController.cpp
@@ -13,6 +13,7 @@
 #include "Features/IModularFeatures.h"
 #include "HAL/RunnableThread.h"
 #include "HAL/Event.h"
+#include "UI/Widgets/SLiveLinkHubPlaybackWidget.h"
 #include "UObject/Object.h"
 #include "UObject/Package.h"
 
@@ -23,68 +24,298 @@ FLiveLinkHubPlaybackController::FLiveLinkHubPlaybackController()
 	RecordingPlayer = MakeUnique<FLiveLinkUAssetRecordingPlayer>();
 }
 
+FLiveLinkHubPlaybackController::~FLiveLinkHubPlaybackController()
+{
+	StopPlayback();
+	Stopping = true;
+	PlaybackEvent->Trigger();
+
+	if (Thread.IsValid())
+	{
+		Thread->WaitForCompletion();
+		Thread.Reset();
+	}
+}
+
+TSharedRef<SWidget> FLiveLinkHubPlaybackController::MakePlaybackWidget()
+{
+	return SNew(SLiveLinkHubPlaybackWidget)
+		.IsEnabled_Raw(this, &FLiveLinkHubPlaybackController::IsReady)
+		.OnPlayForward_Raw(this, &FLiveLinkHubPlaybackController::BeginPlayback, false)
+		.OnPlayReverse_Raw(this, &FLiveLinkHubPlaybackController::BeginPlayback, true)
+		.OnFirstFrame_Lambda([this]()
+		{
+			GoToFrame(GetSelectionStartFrame());
+		})
+		.OnLastFrame_Lambda([this]()
+		{
+			GoToFrame(GetSelectionEndFrame());
+		})
+		.OnPreviousFrame_Lambda([this]()
+		{
+			GoToFrame(CurrentFrameIndex - 1);
+		})
+		.OnNextFrame_Lambda([this]()
+		{
+			GoToFrame(CurrentFrameIndex + 1);
+		})
+		.SetCurrentTime_Raw(this, &FLiveLinkHubPlaybackController::GoToTime)
+		.GetViewRange_Lambda([this]()
+		{
+			return SliderViewRange;
+		})
+		.SetViewRange_Lambda([this](TRange<double> NewRange)
+		{
+			SliderViewRange = MoveTemp(NewRange);
+		})
+		.GetTotalLength_Raw(this, &FLiveLinkHubPlaybackController::GetLength)
+		.GetCurrentTime_Raw(this, &FLiveLinkHubPlaybackController::GetCurrentTime)
+		.GetCurrentFrame_Raw(this, &FLiveLinkHubPlaybackController::GetCurrentFrame)
+		.GetSelectionStartTime_Raw(this, &FLiveLinkHubPlaybackController::GetSelectionStartTime)
+		.SetSelectionStartTime_Raw(this,& FLiveLinkHubPlaybackController::SetSelectionStartTime)
+		.GetSelectionEndTime_Raw(this, &FLiveLinkHubPlaybackController::GetSelectionEndTime)
+		.SetSelectionEndTime_Raw(this,& FLiveLinkHubPlaybackController::SetSelectionEndTime)
+		.IsPaused_Raw(this, &FLiveLinkHubPlaybackController::IsPaused)
+		.IsInReverse_Raw(this, &FLiveLinkHubPlaybackController::IsPlayingInReverse)
+		.IsLooping_Raw(this, &FLiveLinkHubPlaybackController::IsLooping)
+		.OnSetLooping_Raw(this, &FLiveLinkHubPlaybackController::SetLooping)
+		.GetTimeDelta_Raw(this, &FLiveLinkHubPlaybackController::GetTimeDelta);
+}
+
 void FLiveLinkHubPlaybackController::StartPlayback()
 {
+	ResumePlayback();
 	PlaybackStartTime = FPlatformTime::Seconds();
-	Playhead = 0;
-	bIsInPlayback = true;
 	
 	FPlatformMisc::MemoryBarrier();
 	PlaybackEvent->Trigger();
 }
 
-void FLiveLinkHubPlaybackController::PlayRecording(ULiveLinkRecording* InLiveLinkRecording)
-{
-	PreparePlayback(InLiveLinkRecording);
-}
-
-void FLiveLinkHubPlaybackController::StopPlayback()
+void FLiveLinkHubPlaybackController::ResumePlayback()
 {
-	bIsInPlayback = false;
+	bIsInPlayback = true;
+	bIsPaused = false;
+	StartTimestamp = GetTimeFromFrameIndex(GetCurrentFrame());
+	// Force sync so interpolation doesn't interfere if the first frame isn't the current frame
+	SyncToFrame(CurrentFrameIndex);
 }
 
 void FLiveLinkHubPlaybackController::PreparePlayback(ULiveLinkRecording* InLiveLinkRecording)
 {
-	if (InLiveLinkRecording)
+	if (InLiveLinkRecording == nullptr)
+	{
+		UE_LOG(LogLiveLinkHub, Warning, TEXT("Started a recording playback with an invalid recording."));
+	}
+	else if (InLiveLinkRecording != RecordingToPlay.Get())
 	{
-		if (bIsInPlayback)
+		if (RecordingToPlay.IsValid())
 		{
-			check(RecordingToPlay);
-			UE_LOG(LogLiveLinkHub, Warning, TEXT("Can't start playback since we are already doing playback"));
-			return;
+			Eject();
 		}
-
+		
 		RecordingToPlay.Reset(InLiveLinkRecording);
 
+		// The start and end of playback.
+		SelectionStartTime = 0.f;
+		SelectionEndTime = GetLength();
+		
+		// The range the user sees.
+		SliderViewRange = TRange<double>(SelectionStartTime, SelectionEndTime);
+		
 		RollbackPreset.Reset(NewObject<ULiveLinkPreset>(GetTransientPackage(), TEXT("RecordingRollbackPreset")));
 		// Save the current state of the sources/subjects in a rollback preset.
 		RollbackPreset->BuildFromClient();
 
 		RecordingPlayer->PreparePlayback(RecordingToPlay.Get());
 
-		RecordingToPlay->RecordingPreset->ApplyToClientLatent(
-			[this](bool)
+		RecordingToPlay->RecordingPreset->ApplyToClientLatent([this](bool)
+		{
+			bIsReady = true;
+			SyncToFrame(0); // Needed to establish connection with client
+		});
+	}
+}
+
+void FLiveLinkHubPlaybackController::PlayRecording(ULiveLinkRecording* InLiveLinkRecording)
+{
+	PreparePlayback(InLiveLinkRecording);
+}
+
+void FLiveLinkHubPlaybackController::BeginPlayback(bool bInReverse)
+{
+	const bool bReverseChange = bIsReverse != bInReverse;
+	bIsReverse = bInReverse;
+	
+	// Either we are paused and should unpause, or we are toggling forward/reverse play modes.
+	if (bIsPaused || !bIsInPlayback || bReverseChange)
+	{
+		if (ShouldRestart())
+		{
+			// Check if we're at the end of the recording and restart, ie user pressed play again.
+			RestartPlayback();
+		}
+		else
+		{
+			if (bIsReverse)
 			{
-				StartPlayback();
-			});
+				RecordingPlayer->RestartPlayback(CurrentFrameIndex);
+			}
+			
+			// Resume as normal for anywhere else in the recording.
+			PlaybackStartTime = FPlatformTime::Seconds();
+		}
+		
+		ResumePlayback();
 	}
-	else
+	else if (bIsInPlayback)
 	{
-		UE_LOG(LogLiveLinkHub, Warning, TEXT("Started a recording playback with an invalid recording."));
+		PausePlayback();
 	}
+
+	FPlatformMisc::MemoryBarrier();
+	PlaybackEvent->Trigger();
 }
 
-FLiveLinkHubPlaybackController::~FLiveLinkHubPlaybackController()
+void FLiveLinkHubPlaybackController::RestartPlayback()
 {
+	CurrentFrameIndex = INDEX_NONE;
+
+	const bool bOldReverse = bIsReverse; // Stop playback resets reverse
 	StopPlayback();
-	Stopping = true;
-	PlaybackEvent->Trigger();
+	StartTimestamp = GetTimeFromFrameIndex(CurrentFrameIndex);
+	RecordingPlayer->RestartPlayback(CurrentFrameIndex);
+	bIsInPlayback = true;
+	bIsReverse = bOldReverse;
+}
 
-	if (Thread.IsValid())
+void FLiveLinkHubPlaybackController::PausePlayback()
+{
+	bIsPaused = true;
+}
+
+void FLiveLinkHubPlaybackController::StopPlayback()
+{
+	bIsInPlayback = false;
+
+	// Wait for the playback thread to exit...
+	const uint32 CurrentThreadId = FPlatformTLS::GetCurrentThreadId();
+	if (CurrentThreadId != Thread->GetThreadID())
 	{
-		Thread->WaitForCompletion();
-		Thread.Reset();
+		while(!bIsPlaybackWaiting) { }
+	}
+	
+	const bool bReverse = bIsReverse.load();
+	Playhead = bReverse ? GetSelectionEndTime() : GetSelectionStartTime();
+	if (CurrentFrameIndex == INDEX_NONE)
+	{
+		CurrentFrameIndex = GetFrameIndexFromTime(bReverse ? GetSelectionEndTime() : GetSelectionStartTime());
+	}
+	PlaybackStartTime = FPlatformTime::Seconds();
+
+	RecordingPlayer->RestartPlayback();
+	bIsReverse = false;
+}
+
+void FLiveLinkHubPlaybackController::Eject()
+{
+	bIsReady = false;
+	
+	StopPlayback();
+	
+	bIsPaused = false;
+	CurrentFrameIndex = 0;
+	RecordingPlayer->RestartPlayback(CurrentFrameIndex);
+
+	SetSelectionStartTime(0.f);
+	SetSelectionEndTime(0.f);
+	Playhead = 0.f;
+	StartTimestamp = 0.f;
+	
+	if (RollbackPreset.IsValid())
+	{
+		RollbackPreset->ApplyToClientLatent();
 	}
+	// Recording is done, clear the pointer.
+	RecordingToPlay.Reset();
+}
+
+void FLiveLinkHubPlaybackController::GoToTime(double InTime)
+{
+	// Stop needs to occur to restart playback.
+	StopPlayback();
+
+	PlaybackStartTime -= InTime;
+	Playhead = InTime;
+
+	const int32 FrameIndex = GetFrameIndexFromTime(InTime);
+	SyncToFrame(FrameIndex);
+}
+
+void FLiveLinkHubPlaybackController::GoToFrame(int32 InFrameIndex)
+{
+	// Stop needs to occur to restart playback.
+	StopPlayback();
+
+	const double NewTime = GetTimeFromFrameIndex(InFrameIndex);
+	PlaybackStartTime -= NewTime;
+	Playhead = NewTime;
+	
+	SyncToFrame(InFrameIndex);
+}
+
+int32 FLiveLinkHubPlaybackController::GetFrameIndexFromTime(double InTime, bool bReverse) const
+{
+	return RecordingPlayer->PlayheadToFrameIndex(InTime, bReverse);
+}
+
+double FLiveLinkHubPlaybackController::GetTimeFromFrameIndex(int32 InFrameIndex) const
+{
+	return RecordingPlayer->FrameIndexToPlayhead(InFrameIndex);
+}
+
+int32 FLiveLinkHubPlaybackController::GetSelectionStartFrame() const
+{
+	return GetFrameIndexFromTime(GetSelectionStartTime(), bIsReverse);
+}
+
+int32 FLiveLinkHubPlaybackController::GetSelectionEndFrame() const
+{
+	return GetFrameIndexFromTime(GetSelectionEndTime(), bIsReverse);
+}
+
+double FLiveLinkHubPlaybackController::GetSelectionStartTime() const
+{
+	return SelectionStartTime;
+}
+
+void FLiveLinkHubPlaybackController::SetSelectionStartTime(double InTime)
+{
+	SelectionStartTime = InTime;
+}
+
+double FLiveLinkHubPlaybackController::GetSelectionEndTime() const
+{
+	return SelectionEndTime;
+}
+
+void FLiveLinkHubPlaybackController::SetSelectionEndTime(double InTime)
+{
+	SelectionEndTime = InTime;
+}
+
+double FLiveLinkHubPlaybackController::GetLength() const
+{
+	return RecordingToPlay ? RecordingToPlay->LengthInSeconds : 0.f;
+}
+
+double FLiveLinkHubPlaybackController::GetCurrentTime() const
+{
+	return Playhead.load();
+}
+
+int32 FLiveLinkHubPlaybackController::GetCurrentFrame() const
+{
+	return CurrentFrameIndex;
 }
 
 void FLiveLinkHubPlaybackController::Start()
@@ -104,61 +335,53 @@ uint32 FLiveLinkHubPlaybackController::Run()
 {
 	while (!Stopping.load())
 	{
+		bIsPlaybackWaiting = true;
 		PlaybackEvent->Wait();
-
+		bIsPlaybackWaiting = false;
 		while (bIsInPlayback)
 		{
-			if (Playhead >= RecordingToPlay->LengthInSeconds)
+			if (bIsPaused)
+			{
+				FPlatformProcess::Sleep(0.002);
+			}
+			else
 			{
-				if (bLoopPlayback && RecordingToPlay->LengthInSeconds != 0)
+				const bool bSynced = SyncToPlayhead();
+
+				auto SetPlayhead = [&]()
 				{
-					Playhead = 0.0;
-					PlaybackStartTime = FPlatformTime::Seconds();
+					const double Delta = FPlatformTime::Seconds() - PlaybackStartTime;
+					Playhead = bIsReverse ? StartTimestamp - Delta : StartTimestamp + Delta;
+					Playhead = FMath::Clamp(Playhead.load(), GetSelectionStartTime(), GetSelectionEndTime());
+				};
 
-					RecordingPlayer->RestartPlayback();
-				}
-				else
+				SetPlayhead();
+			
+				// Don't sleep if we pushed frames since that can take a small amount of time.
+				if (!bSynced)
 				{
-					// Stop playback
-					break;
+					FPlatformProcess::Sleep(0.002);
+					SetPlayhead();
 				}
 			}
 			
-			TArray<FLiveLinkRecordedFrame> NextFrames = RecordingPlayer->FetchNextFrames(Playhead);
-
-			for (const FLiveLinkRecordedFrame& NextFrame : NextFrames)
+			if (ShouldRestart())
 			{
-				// If we're sending static data
-				if (NextFrame.LiveLinkRole)
+				if (bLoopPlayback && RecordingToPlay->LengthInSeconds != 0 && !bIsPaused)
 				{
-					FLiveLinkStaticDataStruct StaticDataStruct;
-					StaticDataStruct.InitializeWith(NextFrame.Data.GetScriptStruct(), (FLiveLinkBaseStaticData*)NextFrame.Data.GetMemory());
-					Client->PushSubjectStaticData_AnyThread(NextFrame.SubjectKey, NextFrame.LiveLinkRole, MoveTemp(StaticDataStruct));
+					RestartPlayback();
 				}
 				else
 				{
-					FLiveLinkFrameDataStruct FrameDataStruct;
-					FrameDataStruct.InitializeWith(NextFrame.Data.GetScriptStruct(), (FLiveLinkBaseFrameData*)NextFrame.Data.GetMemory());
-
-					Client->PushSubjectFrameData_AnyThread(NextFrame.SubjectKey, MoveTemp(FrameDataStruct));
+					// Stop playback
+					break;
 				}
 			}
 
-			Playhead = FPlatformTime::Seconds() - PlaybackStartTime;
-
-			// Don't sleep if we pushed frames since that can take a small amount of time.
-			if (!NextFrames.Num())
-			{
-				FPlatformProcess::Sleep(0.002);
-				Playhead = FPlatformTime::Seconds() - PlaybackStartTime;
-			}
 		}
 
 		// If the loop ended because the recording is over.
 		bIsInPlayback = false;
-
-		// Recording is done, clear the pointer.
-		CurrentRecording.Reset();
 		
 		// Trigger the playback finished delegate on the game thread.
 		FSimpleDelegateGraphTask::CreateAndDispatchWhenReady(FSimpleDelegateGraphTask::FDelegate::CreateRaw(this, &FLiveLinkHubPlaybackController::OnPlaybackFinished_Internal), TStatId(), nullptr, ENamedThreads::GameThread);
@@ -169,8 +392,69 @@ uint32 FLiveLinkHubPlaybackController::Run()
 
 void FLiveLinkHubPlaybackController::OnPlaybackFinished_Internal()
 {
-	RollbackPreset->ApplyToClientLatent();
-	RecordingToPlay.Reset();
-
 	PlaybackFinishedDelegate.Broadcast();
 }
+
+void FLiveLinkHubPlaybackController::PushSubjectData(const FLiveLinkRecordedFrame& NextFrame, bool bForceSync)
+{
+	// If we're sending static data
+	if (NextFrame.LiveLinkRole)
+	{
+		FLiveLinkStaticDataStruct StaticDataStruct;
+		StaticDataStruct.InitializeWith(NextFrame.Data.GetScriptStruct(), (FLiveLinkBaseStaticData*)NextFrame.Data.GetMemory());
+		Client->PushSubjectStaticData_AnyThread(NextFrame.SubjectKey, NextFrame.LiveLinkRole, MoveTemp(StaticDataStruct));
+	}
+	else
+	{
+		// Record the frame index when pushing so it is accurate. If we only calculate based on time it may not match the actual frames sent.
+		CurrentFrameIndex = NextFrame.FrameIndex;
+		
+		FLiveLinkFrameDataStruct FrameDataStruct;
+		FrameDataStruct.InitializeWith(NextFrame.Data.GetScriptStruct(), (FLiveLinkBaseFrameData*)NextFrame.Data.GetMemory());
+		
+		if (bForceSync)
+		{
+			FrameDataStruct.GetBaseData()->MetaData.StringMetaData.Add(TEXT("ForceSync"), TEXT("true"));
+		}
+		Client->PushSubjectFrameData_AnyThread(NextFrame.SubjectKey, MoveTemp(FrameDataStruct));
+	}
+}
+
+bool FLiveLinkHubPlaybackController::SyncToPlayhead()
+{
+	const double Timestamp = Playhead.load();
+	TArray<FLiveLinkRecordedFrame> NextFrames = bIsReverse ? RecordingPlayer->FetchPreviousFramesAtTimestamp(Timestamp)
+		: RecordingPlayer->FetchNextFramesAtTimestamp(Timestamp);
+	
+	for (const FLiveLinkRecordedFrame& NextFrame : NextFrames)
+	{
+		// todo: Have to forcesync -- interpolation fails due to improper frame times
+		const bool bForceSync = bIsReverse;
+		PushSubjectData(NextFrame, bForceSync);
+	}
+	
+	return NextFrames.Num() > 0;
+}
+
+bool FLiveLinkHubPlaybackController::SyncToFrame(int32 InFrameIndex)
+{
+	TArray<FLiveLinkRecordedFrame> NextFrames = RecordingPlayer->FetchNextFramesAtIndex(InFrameIndex);
+	
+	for (const FLiveLinkRecordedFrame& NextFrame : NextFrames)
+	{
+		PushSubjectData(NextFrame, true);
+	}
+
+	return NextFrames.Num() > 0;
+}
+
+bool FLiveLinkHubPlaybackController::ShouldRestart() const
+{
+	return RecordingToPlay.IsValid() && ((bIsReverse && CurrentFrameIndex <= GetSelectionStartFrame()) || (!bIsReverse && CurrentFrameIndex >= GetSelectionEndFrame()));
+}
+
+double FLiveLinkHubPlaybackController::GetTimeDelta() const
+{
+	// todo: get frame frate
+	return RecordingToPlay.IsValid() ? RecordingToPlay->LengthInSeconds / 60.f : 1.f;
+}
diff --git a/Engine/Plugins/Animation/LiveLink/Source/LiveLinkHub/Private/Recording/LiveLinkHubPlaybackController.h b/Engine/Plugins/Animation/LiveLink/Source/LiveLinkHub/Private/Recording/LiveLinkHubPlaybackController.h
--- a/Engine/Plugins/Animation/LiveLink/Source/LiveLinkHub/Private/Recording/LiveLinkHubPlaybackController.h
+++ b/Engine/Plugins/Animation/LiveLink/Source/LiveLinkHub/Private/Recording/LiveLinkHubPlaybackController.h
@@ -22,23 +22,102 @@ class FRunnableThread;
 class ULiveLinkRecording;
 class ULiveLinkPreset;
 class ULiveLinkRecording;
+class SWidget;
 
 class FLiveLinkHubPlaybackController : public FRunnable
 {
 public:
 	FLiveLinkHubPlaybackController();
 	virtual ~FLiveLinkHubPlaybackController() override;
+
+	/** Create the playback widget. */
+	TSharedRef<SWidget> MakePlaybackWidget();
+
+	/** Apply the recording's preset then prepare the data needed to playback. */
+	void PreparePlayback(ULiveLinkRecording* InLiveLinkRecording);
 	
 	/** Start playing a livelink recording. */
 	void PlayRecording(ULiveLinkRecording* InLiveLinkRecording);
 
+	/** The current recording. */
+	const TStrongObjectPtr<ULiveLinkRecording>& GetRecording() const { return RecordingToPlay; }
+	
+	/** Start playing the currently prepared recording. */
+	void BeginPlayback(bool bInReverse);
+
+	/** Prepare to restart the playback. */
+	void RestartPlayback();
+
+	/** Pause playback. */
+	void PausePlayback();
+	
 	/** Stop playing a livelink recording. */
 	void StopPlayback();
 
+	/** Stop playback and restore the previous settings. */
+	void Eject();
+
+	/** Go to a specific time. */
+	void GoToTime(double InTime);
+
+	/** Go to a specific frame index. */
+	void GoToFrame(int32 InFrameIndex);
+
+	/** Calculate the frame index for the given time. */
+	int32 GetFrameIndexFromTime(double InTime, bool bReverse = false) const;
+
+	/** Retrieve the timestamp from the given frame index. */
+	double GetTimeFromFrameIndex(int32 InFrameIndex) const;
+
+	/** Retrieve the selection start frame. */
+	int32 GetSelectionStartFrame() const;
+
+	/** Retrieve the selection end frame. */
+	int32 GetSelectionEndFrame() const;
+
+	/** Retrieve the selection start time. */
+	double GetSelectionStartTime() const;
+
+	/** Set the selection start time. */
+	void SetSelectionStartTime(double InTime);
+
+	/** Retrieve the selection end time. */
+	double GetSelectionEndTime() const;
+
+	/** Set the selection end time. */
+	void SetSelectionEndTime(double InTime);
+
+	/** Retrieve the length of the recording. */
+	double GetLength() const;
+
+	/** Retrieve the playhead. */
+	double GetCurrentTime() const;
+
+	/** Return the exact frame index of frames that have been processed. */
+	int32 GetCurrentFrame() const;
+
+	/** If the controller is ready for commands. */
+	bool IsReady() const
+	{
+		return bIsReady;
+	}
+	
 	/** Returns whether we've started or are actively playing a recording. */ 
 	bool IsInPlayback() const
 	{
-		return RecordingToPlay || bIsInPlayback.load();
+		return bIsInPlayback.load();
+	}
+
+	/** If playback is paused. */
+	bool IsPaused() const
+	{
+		return bIsPaused.load() || !IsInPlayback();
+	}
+
+	/** If playback is playing in reverse. */
+	bool IsPlayingInReverse() const
+	{
+		return bIsReverse.load();
 	}
 
 	/** Returns whether the recording is set to loop. */
@@ -72,22 +151,54 @@ public:
 private:
 	/** Trigger the playback thread to start reading data. */
 	void StartPlayback();
-	/** Apply the recording's preset then prepare the data needed to playback. */
-	void PreparePlayback(ULiveLinkRecording* InLiveLinkRecording);
+	/** Resume on the playback thread. */
+	void ResumePlayback();
+
 	/** Handler called when playback is finished on the playback thread. Is responsible for resetting the livelink state to what it was before we started playback. */
 	void OnPlaybackFinished_Internal();
+
+	/**
+	 * Send data to the client.
+	 * @param NextFrame The frame to push.
+	 * @param bForceSync Force the client to sync directly to this frame, discarding all other frames.
+	 */
+	void PushSubjectData(const FLiveLinkRecordedFrame& NextFrame, bool bForceSync = false);
+
+	/**
+	 * Sync the animation to the current playhead value.
+	 *
+	 * @return true if any frames were pushed.
+	 */
+	bool SyncToPlayhead();
+
+	/** Force sync to a specific frame. */
+	bool SyncToFrame(int32 InFrameIndex);
+
+	/** Checks if the current playback settings indicates the recording should restart. */
+	bool ShouldRestart() const;
+
+	/** Retrieve the time delta to use for the recording when dragging the playhead spinbox. */
+	double GetTimeDelta() const;
 	
 private:
+	/** If the system has established a connection with the client. */
+	bool bIsReady = false;
 	/** Flag for terminating the thread loop */
 	std::atomic<bool> Stopping = false;
 	/** Thread to do playback on **/
 	TUniquePtr<FRunnableThread> Thread;
 	/** Event signaling that a recording is available for playback. */
 	FEventRef PlaybackEvent = FEventRef();
+	/** If the playback thread is waiting. */
+	std::atomic<bool> bIsPlaybackWaiting = false;
 	/** Whether we're doing playback. */
 	std::atomic<bool> bIsInPlayback = false;
-	/** Recording we're currently playing. */
-	TSharedPtr<FLiveLinkRecordingData> CurrentRecording;
+	/** Whether we're currently paused. */
+	std::atomic<bool> bIsPaused = false;
+	/** If the recording is playing in reverse. */
+	std::atomic<bool>bIsReverse = false;
+	/** The timestamp of the animation when first playing. Can be > 0 when running in reverse. */
+	std::atomic<double> StartTimestamp = 0.f;
 	/** LiveLinkRecording to play.  */
 	TStrongObjectPtr<ULiveLinkRecording> RecordingToPlay;
 	/** Delegate called when a recording playback is finished (if it's not looping). */
@@ -98,10 +209,23 @@ private:
 	std::atomic<bool> bLoopPlayback = false;
 	/** Implementation of the playback functionality. */
 	TUniquePtr<ILiveLinkRecordingPlayer> RecordingPlayer;
+
 	/** LiveLinkClient used to transmit the data to connected clients. */
 	ILiveLinkClient* Client = nullptr;
 	/** Time that the playback started. */
 	double PlaybackStartTime = 0.0;
 	/** Playhead for the current playback. */
-	double Playhead = 0.0;
+	std::atomic<double> Playhead = 0.0;
+
+	/** The view range of the slider, defaults to start/end time. */
+	TRange<double> SliderViewRange = TRange<double>(0.f, 0.f);
+
+	/** The playback selection start time. */
+	double SelectionStartTime = 0.f;
+
+	/** The playback selection end time. */
+	double SelectionEndTime = 0.f;
+	
+	/** The index of the latest processed frame. This may not exactly match the index calculated from the playhead. */
+	int32 CurrentFrameIndex = INDEX_NONE;
 };
diff --git a/Engine/Plugins/Animation/LiveLink/Source/LiveLinkHub/Private/Recording/LiveLinkHubRecordingListController.h b/Engine/Plugins/Animation/LiveLink/Source/LiveLinkHub/Private/Recording/LiveLinkHubRecordingListController.h
--- a/Engine/Plugins/Animation/LiveLink/Source/LiveLinkHub/Private/Recording/LiveLinkHubRecordingListController.h
+++ b/Engine/Plugins/Animation/LiveLink/Source/LiveLinkHub/Private/Recording/LiveLinkHubRecordingListController.h
@@ -28,55 +28,44 @@ public:
 	TSharedRef<SWidget> MakeRecordingList()
 	{
 		return SNew(SLiveLinkHubRecordingListView)
-			.IsLooping_Raw(this, &FLiveLinkHubRecordingListController::IsLooping)
-			.OnSetLooping_Raw(this, &FLiveLinkHubRecordingListController::SetLooping)
-			.OnImportRecording_Raw(this, &FLiveLinkHubRecordingListController::OnImportRecording);
+			.OnImportRecording_Raw(this, &FLiveLinkHubRecordingListController::OnImportRecording)
+			.OnEject_Raw(this, &FLiveLinkHubRecordingListController::OnEjectRecording)
+			.CanEject_Raw(this, &FLiveLinkHubRecordingListController::CanEjectRecording);
 	}
 
 private:
-	/** Returns whether the recording playback should loop. */
-	bool IsLooping() const
-	{
-		if (TSharedPtr<FLiveLinkHub> HubPtr = LiveLinkHub.Pin())
-		{
-			return HubPtr->GetPlaybackController()->IsLooping();
-		}
-		return false;
-	}
-
-	/** Sets whether the playback should loop. */
-	void SetLooping(bool bLooping) const
-	{
-		if (TSharedPtr<FLiveLinkHub> HubPtr = LiveLinkHub.Pin())
-		{
-			HubPtr->GetPlaybackController()->SetLooping(bLooping);
-		}
-	}
-
 	/** Handler called when a recording a clicked to start the recording.  */
-	void OnImportRecording(const FAssetData& AssetData) const
+	void OnImportRecording(const FAssetData& AssetData)
 	{
 		UObject* RecordingAssetData = AssetData.GetAsset();
 		if (!RecordingAssetData)
 		{
 			UE_LOG(LogLiveLinkHub, Warning, TEXT("Failed to import recording %s"), *AssetData.AssetName.ToString());
-			StopPlayback();
 			return;
 		}
 
 		ULiveLinkRecording* ImportedRecording = Cast<ULiveLinkRecording>(RecordingAssetData);
-		LiveLinkHub.Pin()->GetPlaybackController()->PlayRecording(ImportedRecording);
+		
+		if (TSharedPtr<FLiveLinkHub> HubPtr = LiveLinkHub.Pin())
+		{
+			HubPtr->GetPlaybackController()->PreparePlayback(ImportedRecording);
+		}
 	}
 
-	/** Handler called when user clicks in the recording list to stop the current recording (Temporary until we have a stop button).  */
-	void StopPlayback() const
+	void OnEjectRecording()
 	{
-		if (TSharedPtr<FLiveLinkHub> HubPtr = LiveLinkHub.Pin())
+		if (const TSharedPtr<FLiveLinkHub> HubPtr = LiveLinkHub.Pin())
 		{
-			HubPtr->GetPlaybackController()->StopPlayback();
+			HubPtr->GetPlaybackController()->Eject();
 		}
 	}
 
+	bool CanEjectRecording() const
+	{
+		const TSharedPtr<FLiveLinkHub> HubPtr = LiveLinkHub.Pin();
+		return HubPtr && HubPtr->GetPlaybackController()->GetRecording().IsValid();
+	}
+
 private:
 	/** LiveLinkHub object that holds the different controllers. */
 	TWeakPtr<FLiveLinkHub> LiveLinkHub;
diff --git a/Engine/Plugins/Animation/LiveLink/Source/LiveLinkHub/Private/Recording/LiveLinkRecordingPlayer.h b/Engine/Plugins/Animation/LiveLink/Source/LiveLinkHub/Private/Recording/LiveLinkRecordingPlayer.h
--- a/Engine/Plugins/Animation/LiveLink/Source/LiveLinkHub/Private/Recording/LiveLinkRecordingPlayer.h
+++ b/Engine/Plugins/Animation/LiveLink/Source/LiveLinkHub/Private/Recording/LiveLinkRecordingPlayer.h
@@ -18,6 +18,8 @@ struct FLiveLinkRecordedFrame
 	FLiveLinkSubjectKey SubjectKey;
 	/** Role used to interpret the data (Only present with recorded static data). */
 	TSubclassOf<ULiveLinkRole> LiveLinkRole;
+	/** The frame index of this frame within a track. */
+	int32 FrameIndex;
 };
 
 /** Object responsible for reading a livelink recording and providing the frames to the LiveLinkPlaybackController. */
@@ -30,8 +32,20 @@ public:
 	virtual void PreparePlayback(const ULiveLinkRecording* Recording) = 0;
 
 	/** Fetch next frames at the provided playhead position. */
-	virtual TArray<FLiveLinkRecordedFrame> FetchNextFrames(double Playhead) = 0;
+	virtual TArray<FLiveLinkRecordedFrame> FetchNextFramesAtTimestamp(double Playhead) = 0;
+
+	/** Fetch previous frames at the provided playhead position. */
+	virtual TArray<FLiveLinkRecordedFrame> FetchPreviousFramesAtTimestamp(double Playhead) = 0;
+	
+	/** Fetch next frames at the provided frame index. */
+	virtual TArray<FLiveLinkRecordedFrame> FetchNextFramesAtIndex(int32 FrameIndex) = 0;
+
+	/** Convert the playhead to a frame index. */
+	virtual int32 PlayheadToFrameIndex(double InPlayhead, bool bReverse) = 0;
+
+	/** Convert the frame index to a timstamp. */
+	virtual double FrameIndexToPlayhead(int32 InIndex) = 0;
 
 	/** Restart the recording from the beginning. */
-	virtual void RestartPlayback() = 0;
+	virtual void RestartPlayback(int32 InIndex = INDEX_NONE) = 0;
 };
diff --git a/Engine/Plugins/Animation/LiveLink/Source/LiveLinkHub/Private/Recording/SLiveLinkHubRecordingListView.h b/Engine/Plugins/Animation/LiveLink/Source/LiveLinkHub/Private/Recording/SLiveLinkHubRecordingListView.h
--- a/Engine/Plugins/Animation/LiveLink/Source/LiveLinkHub/Private/Recording/SLiveLinkHubRecordingListView.h
+++ b/Engine/Plugins/Animation/LiveLink/Source/LiveLinkHub/Private/Recording/SLiveLinkHubRecordingListView.h
@@ -16,78 +16,79 @@
 
 #define LOCTEXT_NAMESPACE "LiveLinkHub.RecordingListView"
 
-DECLARE_DELEGATE_RetVal(bool, FOnGetLooping);
-DECLARE_DELEGATE_OneParam(FOnSetLooping, bool);
-DECLARE_DELEGATE_OneParam(FOnImportRecording, const struct FAssetData&);
-
 class SLiveLinkHubRecordingListView : public SCompoundWidget
 {
 public:
+	DECLARE_DELEGATE_OneParam(FOnImportRecording, const struct FAssetData&);
+	DECLARE_DELEGATE(FOnEject);
+	DECLARE_DELEGATE_RetVal(bool, FCanEject)
+	
 	SLATE_BEGIN_ARGS(SLiveLinkHubRecordingListView)
 		{}
-		SLATE_EVENT(FOnGetLooping, IsLooping)
-		SLATE_EVENT(FOnSetLooping, OnSetLooping)
 		SLATE_EVENT(FOnImportRecording, OnImportRecording)
+		SLATE_EVENT(FOnEject, OnEject)
+		SLATE_EVENT(FCanEject, CanEject)
 	SLATE_END_ARGS()
 
 	//~ Begin SWidget interface
 	void Construct(const FArguments& InArgs)
 	{
-		OnGetLoopingDelegate = InArgs._IsLooping;
-		OnSetLoopingDelegate = InArgs._OnSetLooping;
 		OnImportRecordingDelegate = InArgs._OnImportRecording;
+		OnEjectDelegate = InArgs._OnEject;
+		OnCanEjectDelegate = InArgs._CanEject;
 
 		ChildSlot
 		[
 			SNew(SVerticalBox)
 			+ SVerticalBox::Slot()
-			[
-				CreateRecordingPicker()
-			]
-			+ SVerticalBox::Slot()
 			.AutoHeight()
 			[
 				SNew(SHorizontalBox)
 				+ SHorizontalBox::Slot()
-				.Padding(4.f, 0)
+				.HAlign(HAlign_Fill)
 				.VAlign(VAlign_Center)
-				.AutoWidth()
+				.Padding(4.f)
 				[
 					SNew(STextBlock)
-					.Text(LOCTEXT("LoopRecordingLabel", "Loop"))
+					.Text(LOCTEXT("RecordingTitle", "Recordings"))
 				]
 				+ SHorizontalBox::Slot()
-				.Padding(4.f, 0)
-				.VAlign(VAlign_Center)
 				.AutoWidth()
+				.HAlign(HAlign_Right)
+				.VAlign(VAlign_Center)
+				.Padding(4.f)
 				[
-					SNew(SCheckBox)
-					.Padding(FMargin(4.0f, 0.0f))
-					.IsChecked(this, &SLiveLinkHubRecordingListView::IsLoopedChecked)
-					.OnCheckStateChanged(this, &SLiveLinkHubRecordingListView::OnCheckStateChanged)
+					SNew(SButton)
+					.Text(LOCTEXT("EjectButton", "Exit Playback"))
+					.OnClicked(this, &SLiveLinkHubRecordingListView::OnEjectClicked)
+					.IsEnabled(this, &SLiveLinkHubRecordingListView::CanEjectRecording)
 				]
 			]
+			+ SVerticalBox::Slot()
+			[
+				CreateRecordingPicker()
+			]
 		];
 	}
 	//~ End SWidget interface
 
 private:
-	/** Returns whether the looping checkbox should be checked or unchecked. */
-	ECheckBoxState IsLoopedChecked() const
+	/** Callback to notice the hub that we've selected a recording to play. */
+	void OnImportRecording(const FAssetData& AssetData) const
 	{
-		return OnGetLoopingDelegate.Execute() ? ECheckBoxState::Checked : ECheckBoxState::Unchecked;
+		OnImportRecordingDelegate.Execute(AssetData);
 	}
 
-	/** Handler to change the looping option for playback. */
-	void OnCheckStateChanged(ECheckBoxState InState)
+	FReply OnEjectClicked()
 	{
-		OnSetLoopingDelegate.Execute(InState == ECheckBoxState::Checked);
+		OnEjectDelegate.ExecuteIfBound();
+		return FReply::Handled();
 	}
 
-	/** Callback to notice the hub that we've selected a recording to play. */
-	void OnImportRecording(const FAssetData& AssetData) const
+	bool CanEjectRecording() const
 	{
-		OnImportRecordingDelegate.Execute(AssetData);
+		check(OnCanEjectDelegate.IsBound());
+		return OnCanEjectDelegate.Execute();
 	}
 
 	/** Creates the asset picker widget for selecting a recording. */
@@ -117,7 +118,7 @@ private:
 			AssetPickerConfig.Filter.ClassPaths.Add(ULiveLinkRecording::StaticClass()->GetClassPathName());
 			AssetPickerConfig.Filter.bRecursiveClasses = true;
 			AssetPickerConfig.Filter.bRecursivePaths = true;
-			AssetPickerConfig.OnAssetSelected = FOnAssetSelected::CreateRaw(this, &SLiveLinkHubRecordingListView::OnImportRecording);
+			AssetPickerConfig.OnAssetDoubleClicked = FOnAssetSelected::CreateRaw(this, &SLiveLinkHubRecordingListView::OnImportRecording);
 		}
 
 		MenuBuilder.BeginSection(NAME_None, LOCTEXT("ImportRecording_MenuSection", "Import Recording"));
@@ -137,12 +138,12 @@ private:
 	}
 
 private:
-	/** Delegate for checking if the recording playback should loop. */
-	FOnGetLooping OnGetLoopingDelegate;
-	/** Delegate to set the looping option. */
-	FOnSetLooping OnSetLoopingDelegate;
 	/** Delegate used for noticing the hub that a recording was selected for playback. */
 	FOnImportRecording OnImportRecordingDelegate;
+	/** Delegate used when ejecting a recording. */
+	FOnEject OnEjectDelegate;
+	/** Delegate used to determine a recording can be ejected. */
+	FCanEject OnCanEjectDelegate;
 };
 
 #undef LOCTEXT_NAMESPACE /* LiveLinkHub.RecordingListView */
diff --git a/Engine/Plugins/Animation/LiveLink/Source/LiveLinkHub/Private/UI/Widgets/SLiveLinkHubMainTabView.cpp b/Engine/Plugins/Animation/LiveLink/Source/LiveLinkHub/Private/UI/Widgets/SLiveLinkHubMainTabView.cpp
--- a/Engine/Plugins/Animation/LiveLink/Source/LiveLinkHub/Private/UI/Widgets/SLiveLinkHubMainTabView.cpp
+++ b/Engine/Plugins/Animation/LiveLink/Source/LiveLinkHub/Private/UI/Widgets/SLiveLinkHubMainTabView.cpp
@@ -245,7 +245,17 @@ TSharedRef<SDockTab> SLiveLinkHubMainTabView::SpawnPlaybackTab(const FSpawnTabAr
 		.Label(PlaybackTabName)
 		.TabRole(PanelTab)
 		[
-			LiveLinkHubModule.GetRecordingListController()->MakeRecordingList()
+			SNew(SVerticalBox)
+			+ SVerticalBox::Slot()
+			.AutoHeight()
+			[
+				LiveLinkHubModule.GetRecordingListController()->MakeRecordingList()
+			]
+			+ SVerticalBox::Slot()
+			.AutoHeight()
+			[
+				LiveLinkHubModule.GetPlaybackController()->MakePlaybackWidget()
+			]
 		];
 }
 
diff --git a/Engine/Plugins/Animation/LiveLink/Source/LiveLinkHub/Private/UI/Widgets/SLiveLinkHubPlaybackWidget.cpp b/Engine/Plugins/Animation/LiveLink/Source/LiveLinkHub/Private/UI/Widgets/SLiveLinkHubPlaybackWidget.cpp
new file mode 100644
--- /dev/null
+++ b/Engine/Plugins/Animation/LiveLink/Source/LiveLinkHub/Private/UI/Widgets/SLiveLinkHubPlaybackWidget.cpp
@@ -0,0 +1,435 @@
+// Copyright Epic Games, Inc. All Rights Reserved.
+
+#include "SLiveLinkHubPlaybackWidget.h"
+
+#include "FrameNumberNumericInterface.h"
+#include "SSimpleTimeSlider.h"
+#include "Widgets/Input/SButton.h"
+#include "Widgets/Input/SSpinBox.h"
+
+#define LOCTEXT_NAMESPACE "SLiveLinkHubPlaybackWidget"
+
+void SLiveLinkHubPlaybackWidget::Construct(const FArguments& InArgs)
+{
+	OnPlayForwardDelegate = InArgs._OnPlayForward;
+	OnPlayReverseDelegate = InArgs._OnPlayReverse;
+
+	OnFirstFrameDelegate = InArgs._OnFirstFrame;
+	OnLastFrameDelegate = InArgs._OnLastFrame;
+
+	OnPreviousFrameDelegate = InArgs._OnPreviousFrame;
+	OnNextFrameDelegate = InArgs._OnNextFrame;
+
+	OnGetPausedDelegate = InArgs._IsPaused;
+	OnGetIsInReverseDelegate = InArgs._IsInReverse;
+
+	OnGetCurrentTimeDelegate = InArgs._GetCurrentTime;
+	OnGetTotalLengthDelegate = InArgs._GetTotalLength;
+
+	OnSetCurrentTimeDelegate = InArgs._SetCurrentTime;
+
+	OnGetViewRangeDelegate = InArgs._GetViewRange;
+	OnSetViewRangeDelegate = InArgs._SetViewRange;
+
+	OnGetCurrentFrameDelegate = InArgs._GetCurrentFrame;
+
+	OnGetSelectionStartTimeDelegate = InArgs._GetSelectionStartTime;
+	OnGetSelectionEndTimeDelegate = InArgs._GetSelectionEndTime;
+	OnSetSelectionStartTimeDelegate = InArgs._SetSelectionStartTime;
+	OnSetSelectionEndTimeDelegate = InArgs._SetSelectionEndTime;
+	
+	OnSetLoopingDelegate = InArgs._OnSetLooping;
+	OnGetLoopingDelegate = InArgs._IsLooping;
+
+	OnGetTimeDelta = InArgs._GetTimeDelta;
+
+	TAttribute<EFrameNumberDisplayFormats> GetDisplayFormatAttr = TAttribute<EFrameNumberDisplayFormats>(this, &SLiveLinkHubPlaybackWidget::GetDisplayFormat);
+	TAttribute<FFrameRate> GetTickResolutionAttr = TAttribute<FFrameRate>(this, &SLiveLinkHubPlaybackWidget::GetFocusedTickResolution);
+	TAttribute<FFrameRate> GetDisplayRateAttr    = TAttribute<FFrameRate>(this, &SLiveLinkHubPlaybackWidget::GetFocusedDisplayRate);
+	
+	// Create our numeric type interface so we can pass it to the time slider below.
+	TSharedPtr<FFrameNumberInterface> NumberInterface = MakeShareable(new FFrameNumberInterface(GetDisplayFormatAttr, 0, GetTickResolutionAttr, GetDisplayRateAttr));
+	
+	ChildSlot
+	[
+		SNew(SVerticalBox)
+		+SVerticalBox::Slot()
+		.AutoHeight()
+		.Padding(16.f, 8.f)
+		[
+			SNew(SSimpleTimeSlider)
+			.ClampRangeHighlightSize(0.15f)
+			.ClampRangeHighlightColor(FLinearColor::Gray.CopyWithNewOpacity(0.5f))
+			.ScrubPosition(this, &SLiveLinkHubPlaybackWidget::GetCurrentTime)
+			.ViewRange(this, &SLiveLinkHubPlaybackWidget::GetViewRange)
+			.OnViewRangeChanged(this, &SLiveLinkHubPlaybackWidget::SetViewRange)
+			.ClampRange(this, &SLiveLinkHubPlaybackWidget::GetClampRange)
+			.OnScrubPositionChanged_Lambda(
+				[this](double NewScrubTime, bool bIsScrubbing)
+						{
+							if (bIsScrubbing)
+							{
+								SetCurrentTime(NewScrubTime);
+							}
+						})
+		]
+		+SVerticalBox::Slot()
+		.AutoHeight()
+		.Padding(8.f, 8.f)
+		[
+			SNew(SHorizontalBox)
+			+SHorizontalBox::Slot()
+			.AutoWidth()
+			[
+				SNew(SButton)
+				.ButtonStyle(FAppStyle::Get(), "Animation.PlayControlsButton")
+				.OnClicked(this, &SLiveLinkHubPlaybackWidget::OnFirstFramePressed)
+				.ToolTipText(LOCTEXT("ToFront", "To Front"))
+				.ContentPadding(0.0f)
+				[
+					SNew(SImage)
+					.ColorAndOpacity(FSlateColor::UseSubduedForeground())
+					.Image(FAppStyle::Get().GetBrush("Animation.Backward_End"))
+				]
+			]
+			+SHorizontalBox::Slot()
+			.AutoWidth()
+			[
+				SNew(SButton)
+				.ButtonStyle(FAppStyle::Get(), "Animation.PlayControlsButton")
+				.OnClicked(this, &SLiveLinkHubPlaybackWidget::OnPreviousFramePressed)
+				.ToolTipText(LOCTEXT("ToPrevious", "To Previous"))
+				.ContentPadding(0.0f)
+				[
+					SNew(SImage)
+					.ColorAndOpacity(FSlateColor::UseSubduedForeground())
+					.Image(FAppStyle::Get().GetBrush("Animation.Backward_Step"))
+				]
+			]
+			+SHorizontalBox::Slot()
+			.AutoWidth()
+			[
+				SNew(SButton)
+				.OnClicked(this, &SLiveLinkHubPlaybackWidget::OnPlayReversePressed)
+				.ToolTipText_Lambda([this]()
+				{
+					const bool bIsPaused = IsPaused();
+					return bIsPaused ? LOCTEXT("ReverseButton", "Reverse") : LOCTEXT("PauseButton", "Pause");
+				})
+				.ButtonStyle( FAppStyle::Get(), "Animation.PlayControlsButton" )
+				.ContentPadding(0.0f)
+				[
+					SNew(SImage)
+					.ColorAndOpacity(FSlateColor::UseSubduedForeground())
+					.Image(this, &SLiveLinkHubPlaybackWidget::GetPlayReverseIcon)
+				]
+			]
+			+SHorizontalBox::Slot()
+			.AutoWidth()
+			[
+				SNew(SButton)
+				.OnClicked(this, &SLiveLinkHubPlaybackWidget::OnPlayForwardPressed)
+				.ToolTipText_Lambda([this]()
+				{
+					const bool bIsPaused = IsPaused();
+					return bIsPaused ? LOCTEXT("PlayButton", "Play") : LOCTEXT("PauseButton", "Pause");
+				})
+				.ButtonStyle( FAppStyle::Get(), "Animation.PlayControlsButton" )
+				.ContentPadding(0.0f)
+				[
+					SNew(SImage)
+					.ColorAndOpacity(FSlateColor::UseSubduedForeground())
+					.Image(this, &SLiveLinkHubPlaybackWidget::GetPlayForwardIcon)
+				]
+			]
+			+SHorizontalBox::Slot()
+			.AutoWidth()
+			[
+				SNew(SButton)
+				.ButtonStyle(FAppStyle::Get(), "Animation.PlayControlsButton")
+				.OnClicked(this, &SLiveLinkHubPlaybackWidget::OnNextFramePressed)
+				.ToolTipText(LOCTEXT("ToNext", "To Next"))
+				.ContentPadding(0.0f)
+				[
+					SNew(SImage)
+					.ColorAndOpacity(FSlateColor::UseSubduedForeground())
+					.Image(FAppStyle::Get().GetBrush("Animation.Forward_Step"))
+				]
+			]
+			+SHorizontalBox::Slot()
+			.AutoWidth()
+			[
+				SNew(SButton)
+				.ButtonStyle(FAppStyle::Get(), "Animation.PlayControlsButton")
+				.OnClicked(this, &SLiveLinkHubPlaybackWidget::OnLastFramePressed)
+				.ToolTipText(LOCTEXT("ToEnd", "To End"))
+				.ContentPadding(0.0f)
+				[
+					SNew(SImage)
+					.ColorAndOpacity(FSlateColor::UseSubduedForeground())
+					.Image(FAppStyle::Get().GetBrush("Animation.Forward_End"))
+				]
+			]
+			+SHorizontalBox::Slot()
+			.AutoWidth()
+			[
+				SNew(SButton)
+				.ButtonStyle(FAppStyle::Get(), "Animation.PlayControlsButton")
+				.OnClicked(this, &SLiveLinkHubPlaybackWidget::OnLoopPressed)
+				.ToolTipText_Raw(this, &SLiveLinkHubPlaybackWidget::GetLoopTooltip)
+				.ContentPadding(0.0f)
+				[
+					SNew(SImage)
+					.ColorAndOpacity(FSlateColor::UseSubduedForeground())
+					.Image(this, &SLiveLinkHubPlaybackWidget::GetLoopIcon)
+				]
+			]
+			+SHorizontalBox::Slot()
+			.Padding(12.f, 0.f, 0.f, 0.f)
+			[
+				SNew(SSpinBox<double>)
+				.Value(this, &SLiveLinkHubPlaybackWidget::GetSelectionStartTime)
+				.ToolTipText(LOCTEXT("SelectionStartTime", "Selection start time"))
+				.OnValueCommitted(this, &SLiveLinkHubPlaybackWidget::OnSelectionStartTimeCommitted)
+				.OnValueChanged(this, &SLiveLinkHubPlaybackWidget::SetSelectionStartTime)
+				.MinValue(TOptional<double>())
+				.MaxValue(TOptional<double>())
+				.Style(&FAppStyle::Get().GetWidgetStyle<FSpinBoxStyle>("Sequencer.HyperlinkSpinBox"))
+				.TypeInterface(NumberInterface)
+				.ClearKeyboardFocusOnCommit(true)
+				.Delta(this, &SLiveLinkHubPlaybackWidget::GetSpinboxDelta)
+				.LinearDeltaSensitivity(25)
+			]
+			+SHorizontalBox::Slot()
+			.Padding(4.f, 0.f)
+			[
+				SNew(SSpinBox<double>)
+				.Value(this, &SLiveLinkHubPlaybackWidget::GetCurrentTime)
+				.ToolTipText(LOCTEXT("CurrentTime", "Current playback time"))
+				.OnValueCommitted(this, &SLiveLinkHubPlaybackWidget::OnCurrentTimeCommitted)
+				.OnValueChanged(this, &SLiveLinkHubPlaybackWidget::SetCurrentTime)
+				.MinValue(TOptional<double>())
+				.MaxValue(TOptional<double>())
+				.Style(&FAppStyle::Get().GetWidgetStyle<FSpinBoxStyle>("Sequencer.HyperlinkSpinBox"))
+				.TypeInterface(NumberInterface)
+				.ClearKeyboardFocusOnCommit(true)
+				.Delta(this, &SLiveLinkHubPlaybackWidget::GetSpinboxDelta)
+				.LinearDeltaSensitivity(25)
+			]
+			+SHorizontalBox::Slot()
+			.Padding(4.f, 0.f)
+			[
+				SNew(SSpinBox<double>)
+				.Value(this, &SLiveLinkHubPlaybackWidget::GetSelectionEndTime)
+				.ToolTipText(LOCTEXT("SelectionEndTime", "Selection end time"))
+				.OnValueCommitted(this, &SLiveLinkHubPlaybackWidget::OnSelectionEndTimeCommitted)
+				.OnValueChanged(this, &SLiveLinkHubPlaybackWidget::SetSelectionEndTime)
+				.MinValue(TOptional<double>())
+				.MaxValue(TOptional<double>())
+				.Style(&FAppStyle::Get().GetWidgetStyle<FSpinBoxStyle>("Sequencer.HyperlinkSpinBox"))
+				.TypeInterface(NumberInterface)
+				.ClearKeyboardFocusOnCommit(true)
+				.Delta(this, &SLiveLinkHubPlaybackWidget::GetSpinboxDelta)
+				.LinearDeltaSensitivity(25)
+			]
+			+SHorizontalBox::Slot()
+			.Padding(4.f, 0.f)
+			[
+				SNew(SSpinBox<double>)
+				.Value(this, &SLiveLinkHubPlaybackWidget::GetTotalLength)
+				.IsEnabled(false)
+				.ToolTipText(LOCTEXT("RecordingLength", "Recording length"))
+				.MinValue(TOptional<double>())
+				.MaxValue(TOptional<double>())
+				.Style(&FAppStyle::Get().GetWidgetStyle<FSpinBoxStyle>("Sequencer.HyperlinkSpinBox"))
+				.TypeInterface(NumberInterface)
+			]
+		]
+	];
+}
+
+FReply SLiveLinkHubPlaybackWidget::OnPlayForwardPressed()
+{
+	OnPlayForwardDelegate.ExecuteIfBound();
+	return FReply::Handled();
+}
+
+FReply SLiveLinkHubPlaybackWidget::OnPlayReversePressed()
+{
+	OnPlayReverseDelegate.ExecuteIfBound();
+	return FReply::Handled();
+}
+
+FReply SLiveLinkHubPlaybackWidget::OnFirstFramePressed()
+{
+	OnFirstFrameDelegate.ExecuteIfBound();
+	return FReply::Handled();
+}
+
+FReply SLiveLinkHubPlaybackWidget::OnLastFramePressed()
+{
+	OnLastFrameDelegate.ExecuteIfBound();
+	return FReply::Handled();
+}
+
+FReply SLiveLinkHubPlaybackWidget::OnPreviousFramePressed()
+{
+	OnPreviousFrameDelegate.ExecuteIfBound();
+	return FReply::Handled();
+}
+
+FReply SLiveLinkHubPlaybackWidget::OnNextFramePressed()
+{
+	OnNextFrameDelegate.ExecuteIfBound();
+	return FReply::Handled();
+}
+
+FReply SLiveLinkHubPlaybackWidget::OnLoopPressed()
+{
+	check(OnSetLoopingDelegate.IsBound() && OnGetLoopingDelegate.IsBound());
+	OnSetLoopingDelegate.Execute(!OnGetLoopingDelegate.Execute());
+	return FReply::Handled();
+}
+
+void SLiveLinkHubPlaybackWidget::SetCurrentTime(double InTime)
+{
+	check(OnSetCurrentTimeDelegate.IsBound());
+	OnSetCurrentTimeDelegate.Execute(InTime);
+}
+
+void SLiveLinkHubPlaybackWidget::OnCurrentTimeCommitted(double InTime, ETextCommit::Type InTextCommit)
+{
+	SetCurrentTime(InTime);
+}
+
+double SLiveLinkHubPlaybackWidget::GetCurrentTime() const
+{
+	check(OnGetCurrentTimeDelegate.IsBound());
+	return OnGetCurrentTimeDelegate.Execute();
+}
+
+double SLiveLinkHubPlaybackWidget::GetTotalLength() const
+{
+	check(OnGetTotalLengthDelegate.IsBound());
+	return OnGetTotalLengthDelegate.Execute();
+}
+
+int32 SLiveLinkHubPlaybackWidget::GetCurrentFrame() const
+{
+	check(OnGetCurrentFrameDelegate.IsBound());
+	return OnGetCurrentFrameDelegate.Execute();
+}
+
+double SLiveLinkHubPlaybackWidget::GetSelectionStartTime() const
+{
+	check(OnGetSelectionStartTimeDelegate.IsBound());
+	return OnGetSelectionStartTimeDelegate.Execute();
+}
+
+void SLiveLinkHubPlaybackWidget::OnSelectionStartTimeCommitted(double InTime, ETextCommit::Type InTextCommit)
+{
+	SetSelectionStartTime(InTime);
+}
+
+void SLiveLinkHubPlaybackWidget::SetSelectionStartTime(double InTime)
+{
+	check(OnSetSelectionStartTimeDelegate.IsBound());
+	OnSetSelectionStartTimeDelegate.Execute(InTime);
+}
+
+double SLiveLinkHubPlaybackWidget::GetSelectionEndTime() const
+{
+	check(OnGetSelectionEndTimeDelegate.IsBound());
+	return OnGetSelectionEndTimeDelegate.Execute();
+}
+
+void SLiveLinkHubPlaybackWidget::OnSelectionEndTimeCommitted(double InTime, ETextCommit::Type InTextCommit)
+{
+	SetSelectionEndTime(InTime);
+}
+
+void SLiveLinkHubPlaybackWidget::SetSelectionEndTime(double InTime)
+{
+	check(OnSetSelectionEndTimeDelegate.IsBound());
+	OnSetSelectionEndTimeDelegate.Execute(InTime);
+}
+
+TRange<double> SLiveLinkHubPlaybackWidget::GetViewRange() const
+{
+	check(OnGetViewRangeDelegate.IsBound());
+	return OnGetViewRangeDelegate.Execute();
+}
+
+void SLiveLinkHubPlaybackWidget::SetViewRange(TRange<double> InRange)
+{
+	check(OnSetViewRangeDelegate.IsBound());
+	OnSetViewRangeDelegate.Execute(InRange);
+}
+
+TRange<double> SLiveLinkHubPlaybackWidget::GetClampRange() const
+{
+	return TRange<double>(GetSelectionStartTime(), GetSelectionEndTime());
+}
+
+bool SLiveLinkHubPlaybackWidget::IsPaused() const
+{
+	return OnGetPausedDelegate.IsBound() ? OnGetPausedDelegate.Execute() : false;
+}
+
+bool SLiveLinkHubPlaybackWidget::IsPlayingInReverse() const
+{
+	return OnGetIsInReverseDelegate.IsBound() ? OnGetIsInReverseDelegate.Execute() : false;
+}
+
+double SLiveLinkHubPlaybackWidget::GetSpinboxDelta() const
+{
+	check(OnGetTimeDelta.IsBound());
+	return OnGetTimeDelta.Execute();
+}
+
+EFrameNumberDisplayFormats SLiveLinkHubPlaybackWidget::GetDisplayFormat() const
+{
+	return DisplayFormat;
+}
+
+FFrameRate SLiveLinkHubPlaybackWidget::GetFocusedTickResolution() const
+{
+	// todo
+	const int32 TickResolutionValue = 1000 * 24;
+	return FFrameRate(1, 1);
+}
+
+FFrameRate SLiveLinkHubPlaybackWidget::GetFocusedDisplayRate() const
+{
+	// todo
+	const int32 SequenceFrameRate = 24;
+	return FFrameRate(SequenceFrameRate, 1);
+}
+
+const FSlateBrush* SLiveLinkHubPlaybackWidget::GetPlayForwardIcon() const
+{
+	return IsPaused() || IsPlayingInReverse() ? FAppStyle::Get().GetBrush("Animation.Forward") : FAppStyle::Get().GetBrush("Animation.Pause");
+}
+
+const FSlateBrush* SLiveLinkHubPlaybackWidget::GetPlayReverseIcon() const
+{
+	return IsPaused() || !IsPlayingInReverse() ? FAppStyle::Get().GetBrush("Animation.Backward") : FAppStyle::Get().GetBrush("Animation.Pause");
+}
+
+const FSlateBrush* SLiveLinkHubPlaybackWidget::GetLoopIcon() const
+{
+	check(OnGetLoopingDelegate.IsBound());
+	const bool bLooping = OnGetLoopingDelegate.Execute();
+	return bLooping ? FAppStyle::Get().GetBrush("Animation.Loop.Enabled")
+		: FAppStyle::Get().GetBrush("Animation.Loop.Disabled");
+}
+
+FText SLiveLinkHubPlaybackWidget::GetLoopTooltip() const
+{
+	check(OnGetLoopingDelegate.IsBound());
+	const bool bLooping = OnGetLoopingDelegate.Execute();
+
+	return bLooping ? LOCTEXT("Loop", "Loop") : LOCTEXT("NoLoop", "No looping");
+}
+
+#undef LOCTEXT_NAMESPACE
diff --git a/Engine/Plugins/Animation/LiveLink/Source/LiveLinkHub/Private/UI/Widgets/SLiveLinkHubPlaybackWidget.h b/Engine/Plugins/Animation/LiveLink/Source/LiveLinkHub/Private/UI/Widgets/SLiveLinkHubPlaybackWidget.h
new file mode 100644
--- /dev/null
+++ b/Engine/Plugins/Animation/LiveLink/Source/LiveLinkHub/Private/UI/Widgets/SLiveLinkHubPlaybackWidget.h
@@ -0,0 +1,222 @@
+// Copyright Epic Games, Inc. All Rights Reserved.
+
+#pragma once
+
+#include "FrameNumberDisplayFormat.h"
+#include "Misc/FrameRate.h"
+#include "Widgets/SCompoundWidget.h"
+
+class FLiveLinkHubPlaybackController;
+class ULiveLinkRecording;
+
+/**
+ * The playback widget for controlling live link hub animations.
+ */
+class SLiveLinkHubPlaybackWidget : public SCompoundWidget
+{
+public:
+	DECLARE_DELEGATE_RetVal(bool, FOnGetChecked);
+	DECLARE_DELEGATE_OneParam(FOnSetChecked, bool);
+	DECLARE_DELEGATE_RetVal(double, FOnGetTime);
+	DECLARE_DELEGATE_OneParam(FOnSetTime, double);
+	DECLARE_DELEGATE_RetVal(int32, FOnGetFrame);
+	DECLARE_DELEGATE_OneParam(FOnSetViewRange, const TRange<double>&)
+	DECLARE_DELEGATE_RetVal(TRange<double>, FOnGetViewRange)
+
+	DECLARE_DELEGATE(FOnButtonPressed);
+	
+	SLATE_BEGIN_ARGS(SLiveLinkHubPlaybackWidget) { }
+	/** When forward play is pressed. */
+	SLATE_EVENT(FOnButtonPressed, OnPlayForward)
+	/** When reverse play is pressed. */
+	SLATE_EVENT(FOnButtonPressed, OnPlayReverse)
+	/** When go to first frame is pressed. */
+	SLATE_EVENT(FOnButtonPressed, OnFirstFrame)
+	/** When go to last frame is pressed. */
+	SLATE_EVENT(FOnButtonPressed, OnLastFrame)
+	/** When previous frame is pressed. */
+	SLATE_EVENT(FOnButtonPressed, OnPreviousFrame)
+	/** When next frame is pressed. */
+	SLATE_EVENT(FOnButtonPressed, OnNextFrame)
+	/** Is the recording looping? */
+	SLATE_EVENT(FOnGetChecked, IsLooping)
+	/** Set the loop state of the recording. */
+	SLATE_EVENT(FOnSetChecked, OnSetLooping)
+	/** Is the recording paused? */
+	SLATE_EVENT(FOnGetChecked, IsPaused)
+	/** Is the recording playing in reverse? */
+	SLATE_EVENT(FOnGetChecked, IsInReverse)
+	/** Get the current time of the recording. */
+	SLATE_EVENT(FOnGetTime, GetCurrentTime)
+	/** Set the current time of the recording. */
+	SLATE_EVENT(FOnSetTime, SetCurrentTime)
+	/** Get the current frame of the recording. */
+	SLATE_EVENT(FOnGetFrame, GetCurrentFrame)
+	/** Get the total length of the recording. */
+	SLATE_EVENT(FOnGetTime, GetTotalLength)
+	/** Get the selection start of the recording. */
+	SLATE_EVENT(FOnGetTime, GetSelectionStartTime)
+	/** Set the selection start of the recording. */
+	SLATE_EVENT(FOnSetTime, SetSelectionStartTime)
+	/** Get the selection end of the recording. */
+	SLATE_EVENT(FOnGetTime, GetSelectionEndTime)
+	/** Set the selection end of the recording. */
+	SLATE_EVENT(FOnSetTime, SetSelectionEndTime)
+
+	/** Get the view range (visible selection range). */
+	SLATE_EVENT(FOnSetViewRange, SetViewRange)
+	/** Set the view range (visible selection range). */
+	SLATE_EVENT(FOnGetViewRange, GetViewRange)
+	
+	/** Resolution for spinbox increments. */
+	SLATE_EVENT(FOnGetTime, GetTimeDelta)
+	
+	SLATE_END_ARGS()
+
+	/**
+	 * @param InArgs 
+	 */
+	void Construct(const FArguments& InArgs);
+
+private:
+	/** When play forward is pressed. */
+	FReply OnPlayForwardPressed();
+	/** When play reverse is pressed. */
+	FReply OnPlayReversePressed();
+	/** When first frame is pressed. */
+	FReply OnFirstFramePressed();
+	/** When last frame is pressed. */
+	FReply OnLastFramePressed();
+	/** When previous frame is pressed. */
+	FReply OnPreviousFramePressed();
+	/** When next frame is pressed. */
+	FReply OnNextFramePressed();
+	/** When loop recording is pressed. */
+	FReply OnLoopPressed();
+
+	/** Sets the current playhead time. */
+	void SetCurrentTime(double InTime);
+
+	/** When the playhead time is committed by a user change. */
+	void OnCurrentTimeCommitted(double InTime, ETextCommit::Type InTextCommit);
+	/** The current playhead time. */
+	double GetCurrentTime() const;
+	/** The total length of the recording. */
+	double GetTotalLength() const;
+	/** The current frame being played. */
+	int32 GetCurrentFrame() const;
+
+	/** Retrieve the selection start time. */
+	double GetSelectionStartTime() const;
+	/** When the selection start time is committed by a user change. */
+	void OnSelectionStartTimeCommitted(double InTime, ETextCommit::Type InTextCommit);
+	/** Set the selection start time. */
+	void SetSelectionStartTime(double InTime);
+
+	/** Retrieve the selection end time. */
+	double GetSelectionEndTime() const;
+	/** When the selection end time is committed by a user change. */
+	void OnSelectionEndTimeCommitted(double InTime, ETextCommit::Type InTextCommit);
+	/** Set the selection end time. */
+	void SetSelectionEndTime(double InTime);
+
+	/** Retrieve the visible range of the scrubber. */
+	TRange<double> GetViewRange() const;
+	/** Set the visible range of the scrubber. */
+	void SetViewRange(TRange<double> InRange);
+
+	/** Retrieve the range to clamp playback to (the selection). */
+	TRange<double> GetClampRange() const;
+
+	/** Is playback paused? */
+	bool IsPaused() const;
+	/** Is playback in reverse? */
+	bool IsPlayingInReverse() const;
+
+	/** Retrieve the delta used when dragging the spinbox. */
+	double GetSpinboxDelta() const;
+
+	/** Frame or seconds to display for text and input. */
+	EFrameNumberDisplayFormats GetDisplayFormat() const;
+
+	/** Tick resolution for number inputs. */
+	FFrameRate GetFocusedTickResolution() const;
+	/** The display rate for number inputs. */
+	FFrameRate GetFocusedDisplayRate() const;
+
+	/** The forward icon to use for the play forward button. */
+	const FSlateBrush* GetPlayForwardIcon() const;
+	/** The reverse icon to use for the play reverse button. */
+	const FSlateBrush* GetPlayReverseIcon() const;
+	/** The loop icon to use for the loop button. */
+	const FSlateBrush* GetLoopIcon() const;
+	/** The tooltip for the loop button. */
+	FText GetLoopTooltip() const;
+private:
+	TWeakPtr<FLiveLinkHubPlaybackController> PlaybackController;
+	/** Delegate for pressing play forward. */
+	FOnButtonPressed OnPlayForwardDelegate;
+	
+	/** Delegate for pressing play reverse. */
+	FOnButtonPressed OnPlayReverseDelegate;
+
+	/** Delegate for pressing go to first frame. */
+	FOnButtonPressed OnFirstFrameDelegate;
+
+	/** Delegate for pressing go to last frame. */
+	FOnButtonPressed OnLastFrameDelegate;
+
+	/** Delegate for going to the previous frame. */
+	FOnButtonPressed OnPreviousFrameDelegate;
+	
+	/** Delegate for going to the next frame. */
+	FOnButtonPressed OnNextFrameDelegate;
+
+	/** Delegate for checking if the recording playback is paused. */
+	FOnGetChecked OnGetPausedDelegate;
+	
+	/** Delegate for checking if the recording playback is playing in reverse. */
+	FOnGetChecked OnGetIsInReverseDelegate;
+
+	/** Delegate for checking if the recording playback should loop. */
+	FOnGetChecked OnGetLoopingDelegate;
+	
+	/** Delegate to set the looping option. */
+	FOnSetChecked OnSetLoopingDelegate;
+	
+	/** Delegate for getting the total length of the recording. */
+	FOnGetTime OnGetTotalLengthDelegate;
+
+	/** Delegate for checking the playhead time. */
+	FOnGetTime OnGetCurrentTimeDelegate;
+	
+	/** Delegate for setting the current playhead time. */
+	FOnSetTime OnSetCurrentTimeDelegate;
+
+	/** Delegate for getting the current frame. */
+	FOnGetFrame OnGetCurrentFrameDelegate;
+
+	/** Delegate for getting the view range. */
+	FOnGetViewRange OnGetViewRangeDelegate;
+	
+	/** Delegate for setting the view range. */
+	FOnSetViewRange OnSetViewRangeDelegate;
+
+	/** Delegate for getting the selection start. */
+	FOnGetTime OnGetSelectionStartTimeDelegate;
+
+	/** Delegate for setting the selection start. */
+	FOnSetTime OnSetSelectionStartTimeDelegate;
+
+	/** Delegate for getting the selection end. */
+	FOnGetTime OnGetSelectionEndTimeDelegate;
+
+	/** Delegate for setting the selection start. */
+	FOnSetTime OnSetSelectionEndTimeDelegate;
+
+	/** Delegate for the spinbox delta. */
+	FOnGetTime OnGetTimeDelta;
+
+	/** The format to display/edit values in. */
+	EFrameNumberDisplayFormats DisplayFormat = EFrameNumberDisplayFormats::Seconds;
+};
\ No newline at end of file
