Remove subject details tab to fit playback layout

diff --git a/Engine/Plugins/Animation/LiveLink/Source/LiveLinkHub/Private/UI/Widgets/SLiveLinkHubMainTabView.cpp b/Engine/Plugins/Animation/LiveLink/Source/LiveLinkHub/Private/UI/Widgets/SLiveLinkHubMainTabView.cpp
--- a/Engine/Plugins/Animation/LiveLink/Source/LiveLinkHub/Private/UI/Widgets/SLiveLinkHubMainTabView.cpp
+++ b/Engine/Plugins/Animation/LiveLink/Source/LiveLinkHub/Private/UI/Widgets/SLiveLinkHubMainTabView.cpp
@@ -18,6 +18,7 @@
 #include "Widgets/Docking/SDockTab.h"
 #include "Widgets/Input/SCheckBox.h"
 #include "Widgets/Input/SComboButton.h"
+#include "Widgets/Layout/SSplitter.h"
 #include "Widgets/Text/STextBlock.h"
 
 
@@ -26,7 +27,6 @@
 const FName SLiveLinkHubMainTabView::SourcesTabId("SourcesTabId");
 const FName SLiveLinkHubMainTabView::SourceDetailsTabId("SourceDetailsTabId");
 const FName SLiveLinkHubMainTabView::SubjectsTabId("SubjectsTabId");
-const FName SLiveLinkHubMainTabView::SubjectsDetailsTabId("SubjectsDetailsTabId");
 const FName SLiveLinkHubMainTabView::PlaybackTabId("PlaybackTabId");
 const FName SLiveLinkHubMainTabView::ClientsTabId("ClientsTabId");
 const FName SLiveLinkHubMainTabView::ClientDetailsTabId("ClientDetailsTabId");
@@ -34,7 +34,6 @@ const FName SLiveLinkHubMainTabView::ClientDetailsTabId("ClientDetailsTabId");
 const FText SLiveLinkHubMainTabView::SourcesTabName = LOCTEXT("SourcesTabLabel", "Sources");
 const FText SLiveLinkHubMainTabView::SourceDetailsTabName = LOCTEXT("SourceDetailsTabLabel", "Source Details");
 const FText SLiveLinkHubMainTabView::SubjectsTabName = LOCTEXT("SubjectsTabLabel", "Subjects");
-const FText SLiveLinkHubMainTabView::SubjectsDetailsTabName = LOCTEXT("SubjectsDetailsTabLabel", "Subjects Details");
 const FText SLiveLinkHubMainTabView::PlaybackTabName = LOCTEXT("PlaybackTabLabel", "Playback");
 const FText SLiveLinkHubMainTabView::ClientsTabName = LOCTEXT("ClientsTabLabel", "Clients");
 const FText SLiveLinkHubMainTabView::ClientDetailsTabName = LOCTEXT("ClientDetailsTabLabel", "Client Details");
@@ -52,7 +51,7 @@ void SLiveLinkHubMainTabView::Construct(const FArguments& InArgs)
 		{
 			CreateTabs(InTabManager, InLayout, InArgs);
 		}))
-		.LayoutName("LiveLinkHubSourcesTabView_v1.1")
+		.LayoutName("LiveLinkHubSourcesTabView_v1.2")
 	);
 }
 
@@ -79,8 +78,6 @@ void SLiveLinkHubMainTabView::CreateTabs(const TSharedRef<FTabManager>& InTabMan
 	InTabManager->RegisterTabSpawner(SubjectsTabId, FOnSpawnTab::CreateSP(this, &SLiveLinkHubMainTabView::SpawnSubjectsTab))
 		.SetIcon(FSlateIcon(LiveLinkStyleName, TEXT("LiveLinkHub.Subjects.Icon")))
 		.SetDisplayName(SubjectsTabName);
-	InTabManager->RegisterTabSpawner(SubjectsDetailsTabId, FOnSpawnTab::CreateSP(this, &SLiveLinkHubMainTabView::SpawnSubjectsDetailsTab))
-		.SetDisplayName(SubjectsDetailsTabName);
 	InTabManager->RegisterTabSpawner(PlaybackTabId, FOnSpawnTab::CreateSP(this, &SLiveLinkHubMainTabView::SpawnPlaybackTab))
 		.SetIcon(FSlateIcon(LiveLinkStyleName, TEXT("LiveLinkHub.Playback.Icon")))
 		.SetDisplayName(PlaybackTabName);
@@ -124,13 +121,6 @@ void SLiveLinkHubMainTabView::CreateTabs(const TSharedRef<FTabManager>& InTabMan
 							->AddTab(PlaybackTabId, ETabState::OpenedTab)
 							->SetForegroundTab(SubjectsTabId)
 						)
-						->Split
-						(
-							FTabManager::NewStack()
-							->SetSizeCoefficient(0.5f)
-							->SetHideTabWell(true)
-							->AddTab(SubjectsDetailsTabId, ETabState::OpenedTab)
-						)
 					)
 					->Split
 					(
@@ -216,24 +206,26 @@ TSharedRef<SDockTab> SLiveLinkHubMainTabView::SpawnSourcesDetailsTab(const FSpaw
 }
 
 TSharedRef<SDockTab> SLiveLinkHubMainTabView::SpawnSubjectsTab(const FSpawnTabArgs& InTabArgs)
-{
-	return SNew(SDockTab)
-		.Label(SubjectsTabName)
-		.TabRole(PanelTab)
-		[
-			PanelController->SubjectsView->SubjectsTreeView.ToSharedRef()
-		];
-}
-
-TSharedRef<SDockTab> SLiveLinkHubMainTabView::SpawnSubjectsDetailsTab(const FSpawnTabArgs& InTabArgs)
 {
 	FLiveLinkHubModule& LiveLinkHubModule = FModuleManager::Get().GetModuleChecked<FLiveLinkHubModule>("LiveLinkHub");
 
 	return SNew(SDockTab)
-		.Label(SubjectsDetailsTabName)
+		.Label(SubjectsTabName)
 		.TabRole(PanelTab)
 		[
-			LiveLinkHubModule.GetSubjectController()->MakeSubjectView()
+			SNew(SSplitter)
+			.Orientation(EOrientation::Orient_Vertical)
+			+ SSplitter::Slot()
+			.Value(0.5f)
+			[
+				PanelController->SubjectsView->SubjectsTreeView.ToSharedRef()
+			]
+			+
+			SSplitter::Slot()
+			.Value(0.5f)
+			[
+				LiveLinkHubModule.GetSubjectController()->MakeSubjectView()
+			]
 		];
 }
 
diff --git a/Engine/Plugins/Animation/LiveLink/Source/LiveLinkHub/Private/UI/Widgets/SLiveLinkHubMainTabView.h b/Engine/Plugins/Animation/LiveLink/Source/LiveLinkHub/Private/UI/Widgets/SLiveLinkHubMainTabView.h
--- a/Engine/Plugins/Animation/LiveLink/Source/LiveLinkHub/Private/UI/Widgets/SLiveLinkHubMainTabView.h
+++ b/Engine/Plugins/Animation/LiveLink/Source/LiveLinkHub/Private/UI/Widgets/SLiveLinkHubMainTabView.h
@@ -16,7 +16,6 @@ public:
 	static const FName SourcesTabId;
 	static const FName SourceDetailsTabId;
 	static const FName SubjectsTabId;
-	static const FName SubjectsDetailsTabId;
 	static const FName PlaybackTabId;
 	static const FName ClientsTabId;
 	static const FName ClientDetailsTabId;
@@ -25,7 +24,6 @@ public:
 	static const FText SourcesTabName;
 	static const FText SourceDetailsTabName;
 	static const FText SubjectsTabName;
-	static const FText SubjectsDetailsTabName;
 	static const FText PlaybackTabName;
 	static const FText ClientsTabName;
 	static const FText ClientDetailsTabName;
@@ -47,7 +45,6 @@ private:
 	TSharedRef<SDockTab> SpawnSourcesTab(const FSpawnTabArgs& InTabArgs);
 	TSharedRef<SDockTab> SpawnSourcesDetailsTab(const FSpawnTabArgs& InTabArgs);
 	TSharedRef<SDockTab> SpawnSubjectsTab(const FSpawnTabArgs& InTabArgs);
-	TSharedRef<SDockTab> SpawnSubjectsDetailsTab(const FSpawnTabArgs& InTabArgs);
 	TSharedRef<SDockTab> SpawnPlaybackTab(const FSpawnTabArgs& InTabArgs);
 	TSharedRef<SDockTab> SpawnClientsTab(const FSpawnTabArgs& InTabArgs);
 	TSharedRef<SDockTab> SpawnClientDetailsTab(const FSpawnTabArgs& InTabArgs);
