Move LiveLinkHub to its own experimental plugin

diff --git a/Engine/Plugins/Animation/LiveLink/LiveLink.uplugin b/Engine/Plugins/Animation/LiveLink/LiveLink.uplugin
--- a/Engine/Plugins/Animation/LiveLink/LiveLink.uplugin
+++ b/Engine/Plugins/Animation/LiveLink/LiveLink.uplugin
@@ -41,17 +41,6 @@
 			"ProgramAllowList": [ "LiveLinkHub" ],
 			"LoadingPhase": "Default"
 		},
-		{
-			"Name": "LiveLinkHub",
-			"Type": "EditorAndProgram",
-			"ProgramAllowList": [ "LiveLinkHub" ],
-			"LoadingPhase": "Default"
-		},
-		{
-			"Name": "LiveLinkHubEditor",
-			"Type": "Editor",
-			"LoadingPhase": "Default"
-		},
 		{
 			"Name": "LiveLinkComponents",
 			"Type": "Runtime",
@@ -83,12 +72,6 @@
 			"Name": "LiveLinkMultiUser",
 			"Type": "UncookedOnly",
 			"LoadingPhase": "PostEngineInit"
-		},
-		{
-			"Name": "LiveLinkHubMessaging",
-			"Type": "Runtime",
-			"LoadingPhase": "Default",
-			"ProgramAllowList": [ "LiveLinkHub" ]
 		}
 	]
 }
diff --git a/Engine/Programs/LiveLinkHub/Config/DefaultEngine.ini b/Engine/Programs/LiveLinkHub/Config/DefaultEngine.ini
--- a/Engine/Programs/LiveLinkHub/Config/DefaultEngine.ini
+++ b/Engine/Programs/LiveLinkHub/Config/DefaultEngine.ini
@@ -3,6 +3,7 @@ GameNameShort=LLH
 
 [Plugins]
 +ProgramEnabledPlugins="LiveLink"
++ProgramEnabledPlugins="LiveLinkHub"
 +ProgramEnabledPlugins="StructUtils"
 +ProgramEnabledPlugins="ContentBrowserAssetDataSource"
 +ProgramEnabledPlugins="PropertyAccessEditor"
diff --git a/Engine/Source/Programs/LiveLinkHub/LiveLinkHub.Target.cs b/Engine/Source/Programs/LiveLinkHub/LiveLinkHub.Target.cs
--- a/Engine/Source/Programs/LiveLinkHub/LiveLinkHub.Target.cs
+++ b/Engine/Source/Programs/LiveLinkHub/LiveLinkHub.Target.cs
@@ -20,6 +20,7 @@ public class LiveLinkHubTarget : TargetRules
 		SolutionDirectory = "Programs/LiveLink";
 
 		AdditionalPlugins.Add("LiveLink");
+		AdditionalPlugins.Add("LiveLinkHub");
 		AdditionalPlugins.Add("ContentBrowserAssetDataSource");
 		AdditionalPlugins.Add("StructUtils");
 		AdditionalPlugins.Add("UdpMessaging");
