Expose protected member to LiveLinkClient that allows derived instances to notify when new frame data has been received.  This allows us to support time code providers in the Live Link Hub app

diff --git a/Engine/Plugins/Animation/LiveLink/Source/LiveLink/Private/LiveLinkClient.cpp b/Engine/Plugins/Animation/LiveLink/Source/LiveLink/Private/LiveLinkClient.cpp
--- a/Engine/Plugins/Animation/LiveLink/Source/LiveLink/Private/LiveLinkClient.cpp
+++ b/Engine/Plugins/Animation/LiveLink/Source/LiveLink/Private/LiveLinkClient.cpp
@@ -729,6 +729,16 @@ void FLiveLinkClient::PushSubjectStaticData_Internal(FPendingSubjectStatic&& Sub
 	}
 }
 
+
+void FLiveLinkClient::BroadcastFrameDataUpdate(const FLiveLinkSubjectKey& InSubjectKey, const FLiveLinkFrameDataStruct& InFrameData)
+{
+	FScopeLock BroadcastLock(&SubjectFrameReceivedHandleseCriticalSection);
+	if (const FSubjectFramesReceivedHandles* Handles = SubjectFrameReceivedHandles.Find(InSubjectKey))
+	{
+		Handles->OnFrameDataReceived.Broadcast(InFrameData);
+	}
+}
+
 void FLiveLinkClient::PushSubjectFrameData_AnyThread(const FLiveLinkSubjectKey& InSubjectKey, FLiveLinkFrameDataStruct&& InFrameData)
 {
 	FPendingSubjectFrame SubjectFrame{ InSubjectKey, MoveTemp(InFrameData) };
@@ -749,13 +759,7 @@ void FLiveLinkClient::PushSubjectFrameData_AnyThread(const FLiveLinkSubjectKey&
 
 	if (bCanPushFrame)
 	{
-		{
-			FScopeLock BroadcastLock(&SubjectFrameReceivedHandleseCriticalSection);
-			if (const FSubjectFramesReceivedHandles* Handles = SubjectFrameReceivedHandles.Find(InSubjectKey))
-			{
-				Handles->OnFrameDataReceived.Broadcast(SubjectFrame.FrameData);
-			}
-		}
+		BroadcastFrameDataUpdate(InSubjectKey, SubjectFrame.FrameData);
 			
 		// Since the lock was released between setting bCanPushFrame and adding to the array, it is possible that
 		// we exceed MaxNumBufferToCached. But this should be rare and also harmless.
diff --git a/Engine/Plugins/Animation/LiveLink/Source/LiveLink/Public/LiveLinkClient.h b/Engine/Plugins/Animation/LiveLink/Source/LiveLink/Public/LiveLinkClient.h
--- a/Engine/Plugins/Animation/LiveLink/Source/LiveLink/Public/LiveLinkClient.h
+++ b/Engine/Plugins/Animation/LiveLink/Source/LiveLink/Public/LiveLinkClient.h
@@ -242,6 +242,9 @@ private:
 	void RemoveRebroadcastedSubject(FLiveLinkSubjectKey InSubjectKey);
 
 protected:
+	/** Broadcast out to the SubjectFrameAddedHandles a frame data update. */
+	void BroadcastFrameDataUpdate(const FLiveLinkSubjectKey& InSubjectKey, const FLiveLinkFrameDataStruct& InFrameData);
+
 	/** The current collection used. */
 	TUniquePtr<FLiveLinkSourceCollection> Collection;
 private:
diff --git a/Engine/Plugins/Animation/LiveLinkHub/Source/LiveLinkHub/Private/LiveLinkHubClient.cpp b/Engine/Plugins/Animation/LiveLinkHub/Source/LiveLinkHub/Private/LiveLinkHubClient.cpp
--- a/Engine/Plugins/Animation/LiveLinkHub/Source/LiveLinkHub/Private/LiveLinkHubClient.cpp
+++ b/Engine/Plugins/Animation/LiveLinkHub/Source/LiveLinkHub/Private/LiveLinkHubClient.cpp
@@ -415,6 +415,7 @@ void FLiveLinkHubClient::PushSubjectFrameData_AnyThread(const FLiveLinkSubjectKe
 	}
 
 	OnFrameDataReceivedDelegate_AnyThread.Broadcast(SubjectKey, FrameData);
+	BroadcastFrameDataUpdate(SubjectKey, FrameData);
 }
 
 FText FLiveLinkHubClient::GetSourceStatus(FGuid InEntryGuid) const
