Fixed unreachable code warnings from MSVC

#rnx

diff --git a/Engine/Plugins/Animation/LiveLink/Source/LiveLink/Private/LiveLinkClient.cpp b/Engine/Plugins/Animation/LiveLink/Source/LiveLink/Private/LiveLinkClient.cpp
--- a/Engine/Plugins/Animation/LiveLink/Source/LiveLink/Private/LiveLinkClient.cpp
+++ b/Engine/Plugins/Animation/LiveLink/Source/LiveLink/Private/LiveLinkClient.cpp
@@ -326,8 +326,6 @@ ELiveLinkSubjectState FLiveLinkClient::GetSubjectState(FLiveLinkSubjectName InSu
 		ensureMsgf(false, TEXT("Unhandled ETimedDataInputState::%d"), InputState);
 		return ELiveLinkSubjectState::Unknown;
 	}
-
-	checkNoEntry();
 }
 
 FGuid FLiveLinkClient::AddSource(TSharedPtr<ILiveLinkSource> InSource)
diff --git a/Engine/Source/Runtime/Engine/Private/Actor.cpp b/Engine/Source/Runtime/Engine/Private/Actor.cpp
--- a/Engine/Source/Runtime/Engine/Private/Actor.cpp
+++ b/Engine/Source/Runtime/Engine/Private/Actor.cpp
@@ -6319,9 +6319,9 @@ TArray<const UDataLayerInstance*> AActor::GetDataLayerInstancesInternal(bool bUs
 		}
 		return DataLayerInstances;
 	}
-#endif
-	
+#else
 	return TArray<const UDataLayerInstance*>();
+#endif
 }
 
 bool AActor::ContainsDataLayer(const UDataLayerInstance* DataLayerInstance) const
