LiveLinkHub - Fix LiveLink subjects status disappearing in read-only mode

diff --git a/Engine/Plugins/Animation/LiveLink/Source/LiveLinkEditor/Private/LiveLinkClientPanelViews.cpp b/Engine/Plugins/Animation/LiveLink/Source/LiveLinkEditor/Private/LiveLinkClientPanelViews.cpp
--- a/Engine/Plugins/Animation/LiveLink/Source/LiveLinkEditor/Private/LiveLinkClientPanelViews.cpp
+++ b/Engine/Plugins/Animation/LiveLink/Source/LiveLinkEditor/Private/LiveLinkClientPanelViews.cpp
@@ -260,7 +260,6 @@ public:
 				return SNew(SBox)
 					.HAlign(HAlign_Center)
 					.VAlign(VAlign_Center)
-					.Visibility(this, &SLiveLinkClientPanelSubjectRow::GetVisibilityFromReadOnly)
 					[
 						SNew(STextBlock)
 						.Font(FAppStyle::Get().GetFontStyle("FontAwesome.8"))
