LiveLinkHub - Remove Virtual Subject menu in the hub since we don't use virtual subjects

diff --git a/Engine/Plugins/Animation/LiveLink/Source/LiveLinkEditor/Private/LiveLinkClientPanelToolbar.cpp b/Engine/Plugins/Animation/LiveLink/Source/LiveLinkEditor/Private/LiveLinkClientPanelToolbar.cpp
--- a/Engine/Plugins/Animation/LiveLink/Source/LiveLinkEditor/Private/LiveLinkClientPanelToolbar.cpp
+++ b/Engine/Plugins/Animation/LiveLink/Source/LiveLinkEditor/Private/LiveLinkClientPanelToolbar.cpp
@@ -44,6 +44,10 @@
 #include "Widgets/SNullWidget.h"
 #include "Widgets/SWindow.h"
 
+#ifndef WITH_LIVELINK_HUB
+#define WITH_LIVELINK_HUB 0
+#endif
+
 
 #define LOCTEXT_NAMESPACE "LiveLinkClientPanel"
 
@@ -609,6 +613,7 @@ TSharedRef<SWidget> SLiveLinkClientPanelToolbar::OnGenerateSourceMenu()
 	//	FNewMenuDelegate::CreateRaw(this, &SLiveLinkClientPanelToolbar::PopulateVirtualSubjectSourceCreationMenu)
 	//);
 
+#if !WITH_LIVELINK_HUB
 	MenuBuilder.AddMenuEntry(
 		LOCTEXT("AddVirtualSubject", "Add Virtual Subject"),
 		LOCTEXT("AddVirtualSubject_Tooltip", "Adds a new virtual subject to LiveLink. Instead of coming from a source a virtual subject is a combination of 2 or more real subjects"),
@@ -618,6 +623,7 @@ TSharedRef<SWidget> SLiveLinkClientPanelToolbar::OnGenerateSourceMenu()
 		),
 		NAME_None,
 		EUserInterfaceActionType::Button);
+#endif
 
 	MenuBuilder.EndSection();
 
