Documentation changes for LiveLinkTimecodeProvider

diff --git a/Engine/Plugins/Animation/LiveLink/Source/LiveLink/Private/LiveLinkTimecodeProvider.cpp b/Engine/Plugins/Animation/LiveLink/Source/LiveLink/Private/LiveLinkTimecodeProvider.cpp
--- a/Engine/Plugins/Animation/LiveLink/Source/LiveLink/Private/LiveLinkTimecodeProvider.cpp
+++ b/Engine/Plugins/Animation/LiveLink/Source/LiveLink/Private/LiveLinkTimecodeProvider.cpp
@@ -29,7 +29,7 @@ FQualifiedFrameTime ULiveLinkTimecodeProvider::GetQualifiedFrameTime() const
 		}
 		else
 		{
-			// Find the frame that is the closest to what we expect.
+			// Find the frame that is the closest to the current world time.
 			const double Seconds = FPlatformTime::Seconds();
 			int32 FoundIndex = SubjectFrameTimes.Num() - 1;
 			for (; FoundIndex >= 0; --FoundIndex)
@@ -62,7 +62,7 @@ FQualifiedFrameTime ULiveLinkTimecodeProvider::GetQualifiedFrameTime() const
 			}
 
 			check(IndexA != INDEX_NONE && IndexB != INDEX_NONE);
-			// Between the 2 closest frames
+			// Using the 2 closest frames
 			return LerpBetweenFrames(Seconds, IndexA, IndexB);
 		}
 	}
diff --git a/Engine/Plugins/Animation/LiveLink/Source/LiveLink/Public/LiveLinkTimecodeProvider.h b/Engine/Plugins/Animation/LiveLink/Source/LiveLink/Public/LiveLinkTimecodeProvider.h
--- a/Engine/Plugins/Animation/LiveLink/Source/LiveLink/Public/LiveLinkTimecodeProvider.h
+++ b/Engine/Plugins/Animation/LiveLink/Source/LiveLink/Public/LiveLinkTimecodeProvider.h
@@ -16,9 +16,9 @@ class IModularFeature;
 UENUM()
 enum class ELiveLinkTimecodeProviderEvaluationType
 {
-	/** Interpolate, or extrapolate, between the 2 frames that are the closest to evaluation. */
+	/** Interpolate between, or extrapolate using the 2 frames that are the closest to the current world time. */
 	Lerp,
-	/** Use the frame that is closest to evaluation. */
+	/** Use the frame that is closest to the current world time. */
 	Nearest,
 	/** Use the newest frame that was received. */
 	Latest,
@@ -65,6 +65,17 @@ public:
 	}
 private:
 	FQualifiedFrameTime ConvertTo(FQualifiedFrameTime Value) const;
+
+	/**
+	 * Infers the frame time for the current world time by either interpolating between or extrapolating a frame time value from two subject frames provided via live link.
+	 *
+	 * This method uses the FMath::Lerp function but will intentionally provide an Alpha value greater than 1.0 when extrapolation is required.
+	 * 
+	 * @param Seconds The current world time.
+	 * @param IndexA The first subject frame index.
+	 * @param IndexB The second subject frame index.
+	 * @return The inferred frame time.
+	 */
 	FQualifiedFrameTime LerpBetweenFrames(double Seconds, int32 IndexA, int32 IndexB) const;
 
 	void InitClient();
