LiveLink: Fix 1-frame delay between LiveLink evaluation and MRQ rendering. Adds an option to the LiveLink component to force the LiveLink client to Tick again immediately before evaluating its subject if there are additional subject frames that have not been processed. This only affects PIE.

diff --git a/Engine/Plugins/Animation/LiveLink/Source/LiveLink/Private/LiveLinkClient.cpp b/Engine/Plugins/Animation/LiveLink/Source/LiveLink/Private/LiveLinkClient.cpp
--- a/Engine/Plugins/Animation/LiveLink/Source/LiveLink/Private/LiveLinkClient.cpp
+++ b/Engine/Plugins/Animation/LiveLink/Source/LiveLink/Private/LiveLinkClient.cpp
@@ -1278,6 +1278,12 @@ void FLiveLinkClient::ForceTick()
 	Tick();
 }
 
+bool FLiveLinkClient::HasPendingSubjectFrames()
+{
+	FScopeLock Lock(&CollectionAccessCriticalSection);
+	return !SubjectFrameToPush.IsEmpty();
+}
+
 bool FLiveLinkClient::EvaluateFrame_AnyThread(FLiveLinkSubjectName InSubjectName, TSubclassOf<ULiveLinkRole> InDesiredRole, FLiveLinkSubjectFrameData& OutFrame)
 {
 	SCOPE_CYCLE_COUNTER(STAT_LiveLink_EvaluateFrame);
diff --git a/Engine/Plugins/Animation/LiveLink/Source/LiveLink/Public/LiveLinkClient.h b/Engine/Plugins/Animation/LiveLink/Source/LiveLink/Public/LiveLinkClient.h
--- a/Engine/Plugins/Animation/LiveLink/Source/LiveLink/Public/LiveLinkClient.h
+++ b/Engine/Plugins/Animation/LiveLink/Source/LiveLink/Public/LiveLinkClient.h
@@ -123,6 +123,7 @@ public:
 	virtual bool EvaluateFrameAtWorldTime_AnyThread(FLiveLinkSubjectName SubjectName, double WorldTime, TSubclassOf<ULiveLinkRole> DesiredRole, FLiveLinkSubjectFrameData& OutFrame) override;
 	virtual bool EvaluateFrameAtSceneTime_AnyThread(FLiveLinkSubjectName SubjectName, const FQualifiedFrameTime& FrameTime, TSubclassOf<ULiveLinkRole> DesiredRole, FLiveLinkSubjectFrameData& OutFrame) override;
 	virtual void ForceTick() override;
+	virtual bool HasPendingSubjectFrames() override;
 
 	virtual FSimpleMulticastDelegate& OnLiveLinkTicked() override;
 	virtual FSimpleMulticastDelegate& OnLiveLinkSourcesChanged() override;
diff --git a/Engine/Plugins/Animation/LiveLink/Source/LiveLinkComponents/Private/LiveLinkComponentController.cpp b/Engine/Plugins/Animation/LiveLink/Source/LiveLinkComponents/Private/LiveLinkComponentController.cpp
--- a/Engine/Plugins/Animation/LiveLink/Source/LiveLinkComponents/Private/LiveLinkComponentController.cpp
+++ b/Engine/Plugins/Animation/LiveLink/Source/LiveLinkComponents/Private/LiveLinkComponentController.cpp
@@ -171,9 +171,11 @@ void ULiveLinkComponentController::DestroyComponent(bool bPromoteChildren /*= fa
 
 void ULiveLinkComponentController::TickComponent(float DeltaTime, ELevelTick TickType, FActorComponentTickFunction* ThisTickFunction)
 {
+	UWorld* OwningWorld = GetWorld();
+
 	// Verify if we are in an editor preview world (blueprint editor). Without being able to select the desired component
 	// When you spawn a LL component, it defaults to the root component on which we can't reset the transform in case it's manipulated by LL automatically
-	if (GetWorld() && GetWorld()->WorldType == EWorldType::EditorPreview && bUpdateInPreviewEditor == false)
+	if (OwningWorld && OwningWorld->WorldType == EWorldType::EditorPreview && bUpdateInPreviewEditor == false)
 	{
 		return;
 	}
@@ -194,6 +196,15 @@ void ULiveLinkComponentController::TickComponent(float DeltaTime, ELevelTick Tic
 
 	ILiveLinkClient& LiveLinkClient = IModularFeatures::Get().GetModularFeature<ILiveLinkClient>(ILiveLinkClient::ModularFeatureName);
 
+	if (OwningWorld && OwningWorld->WorldType == EWorldType::PIE)
+	{
+		const bool bUpdateImmediatelyInPIE = GetDefault<ULiveLinkComponentSettings>()->bUpdateSubjectsImmediatelyInPIE;
+		if (bUpdateImmediatelyInPIE && LiveLinkClient.HasPendingSubjectFrames())
+		{
+			LiveLinkClient.ForceTick();
+		}
+	}
+
 	// Evaluate subject frame once and pass the data to our controllers
 	FLiveLinkSubjectFrameData SubjectData;
 
diff --git a/Engine/Plugins/Animation/LiveLink/Source/LiveLinkComponents/Public/LiveLinkComponentSettings.h b/Engine/Plugins/Animation/LiveLink/Source/LiveLinkComponents/Public/LiveLinkComponentSettings.h
--- a/Engine/Plugins/Animation/LiveLink/Source/LiveLinkComponents/Public/LiveLinkComponentSettings.h
+++ b/Engine/Plugins/Animation/LiveLink/Source/LiveLinkComponents/Public/LiveLinkComponentSettings.h
@@ -24,6 +24,10 @@ public:
 	/** Default Controller class to use for the specified role */
 	UPROPERTY(config, EditAnywhere, Category = "LiveLink", meta = (AllowAbstract = "false"))
 	TMap<TSubclassOf<ULiveLinkRole>, TSubclassOf<ULiveLinkControllerBase>> DefaultControllerForRole;
+
+	/** If true, all LiveLink subjects with pending frame data will be updated immediately before the LiveLink components execute (only applies to PIE) */
+	UPROPERTY(config, EditAnywhere, Category = "LiveLink")
+	bool bUpdateSubjectsImmediatelyInPIE = true;
 };
 
 #if UE_ENABLE_INCLUDE_ORDER_DEPRECATED_IN_5_2
diff --git a/Engine/Source/Runtime/LiveLinkInterface/Public/ILiveLinkClient.h b/Engine/Source/Runtime/LiveLinkInterface/Public/ILiveLinkClient.h
--- a/Engine/Source/Runtime/LiveLinkInterface/Public/ILiveLinkClient.h
+++ b/Engine/Source/Runtime/LiveLinkInterface/Public/ILiveLinkClient.h
@@ -349,6 +349,9 @@ public:
 	*/
 	virtual void ForceTick() = 0;
 
+	/** Returns true if the client has pending subject frames to process during its next tick */
+	virtual bool HasPendingSubjectFrames() = 0;
+
 	/** Notify when the list of sources has changed. */
 	virtual FSimpleMulticastDelegate& OnLiveLinkSourcesChanged() = 0;
 
