Fix string formatting for LiveLinkSequencer log warning.

#rnx

diff --git a/Engine/Plugins/Animation/LiveLink/Source/LiveLinkSequencer/Private/SequencerRecorderSections/MovieSceneLiveLinkTrackRecorder.cpp b/Engine/Plugins/Animation/LiveLink/Source/LiveLinkSequencer/Private/SequencerRecorderSections/MovieSceneLiveLinkTrackRecorder.cpp
--- a/Engine/Plugins/Animation/LiveLink/Source/LiveLinkSequencer/Private/SequencerRecorderSections/MovieSceneLiveLinkTrackRecorder.cpp
+++ b/Engine/Plugins/Animation/LiveLink/Source/LiveLinkSequencer/Private/SequencerRecorderSections/MovieSceneLiveLinkTrackRecorder.cpp
@@ -265,9 +265,8 @@ void UMovieSceneLiveLinkTrackRecorder::RecordSampleImpl(const FQualifiedFrameTim
 					FFrameTime FrameTime = LiveLinkFrameTime.ConvertTo(TickResolution);
 					FrameNumber = FrameTime.FrameNumber;
 
-					UE_LOG(LogLiveLinkSequencer, VeryVerbose, TEXT("LiveLinkFrameTime: [%s.%.03f] at %s for subject '%s'."), 
-						   *(FTimecode::FromFrameNumber(LiveLinkFrameTime.Time.FrameNumber, LiveLinkFrameTime.Rate).ToString()), 
-						   LiveLinkFrameTime.Time.GetSubFrame(),
+					UE_LOG(LogLiveLinkSequencer, VeryVerbose, TEXT("LiveLinkFrameTime: [%s] at %s for subject '%s'."),
+						   *(FTimecode::FromFrameTime(LiveLinkFrameTime.Time, LiveLinkFrameTime.Rate).ToString()),
 						   *(LiveLinkFrameTime.Rate.ToPrettyText().ToString()),
 						   *(SubjectName.ToString()));
 				}
@@ -287,11 +286,9 @@ void UMovieSceneLiveLinkTrackRecorder::RecordSampleImpl(const FQualifiedFrameTim
 				}
 				else
 				{
-					UE_LOG(LogLiveLinkSequencer, Warning, TEXT("Discarded buffered frame: [%s.%.03f] outside of start frame: [%s.%.03f] for subject '%s'."),
-						   *(FTimecode::FromFrameNumber(ConvertFrameTime(FrameNumber, TickResolution, DisplayRate).FrameNumber, DisplayRate).ToString()), 
-						   ConvertFrameTime(FrameNumber, TickResolution, DisplayRate).GetSubFrame(),
-						   *(FTimecode::FromFrameNumber(ConvertFrameTime(MovieSceneSection->GetInclusiveStartFrame(), TickResolution, DisplayRate).FrameNumber, DisplayRate)).ToString(), 
-						   ConvertFrameTime(MovieSceneSection->GetInclusiveStartFrame(), TickResolution, DisplayRate).GetSubFrame(),
+					UE_LOG(LogLiveLinkSequencer, Warning, TEXT("Discarded buffered frame: [%s] outside of start frame: [%s] for subject '%s'."),
+						   *(FTimecode::FromFrameTime(ConvertFrameTime(FrameNumber, TickResolution, DisplayRate), DisplayRate).ToString()),
+						   *(FTimecode::FromFrameTime(ConvertFrameTime(MovieSceneSection->GetInclusiveStartFrame(), TickResolution, DisplayRate), DisplayRate)).ToString(),
 						   *(SubjectName.ToString()));
 				}
 			}
