LiveLink - Add support for virtual subject bone attachments

Bone attachments allow you to combine subjects into a virtual subject by specifying the source and target bone of the attachment.
In this initial verison, the bone names of the combined subjects must exist on the target skeleton.

diff --git a/Engine/Plugins/Animation/LiveLink/Source/LiveLink/Private/LiveLinkAnimationVirtualSubject.cpp b/Engine/Plugins/Animation/LiveLink/Source/LiveLink/Private/LiveLinkAnimationVirtualSubject.cpp
--- a/Engine/Plugins/Animation/LiveLink/Source/LiveLink/Private/LiveLinkAnimationVirtualSubject.cpp
+++ b/Engine/Plugins/Animation/LiveLink/Source/LiveLink/Private/LiveLinkAnimationVirtualSubject.cpp
@@ -2,6 +2,7 @@
 
 #include "LiveLinkAnimationVirtualSubject.h"
 
+#include "Algo/TopologicalSort.h"
 #include "ILiveLinkClient.h"
 #include "Roles/LiveLinkAnimationRole.h"
 #include "Roles/LiveLinkAnimationTypes.h"
@@ -11,7 +12,7 @@
 
 namespace LiveLinkAnimationVirtualSubjectUtils
 {
-	void AddToBoneNames(TArray<FName>& BoneNames, const TArray<FName>& NewBoneNames, const FName Prefix)
+	void AddToBoneNames(TArray<FName>& BoneNames, const TArray<FName>& NewBoneNames, const FName Prefix, TMap<int32, FName>& NamesToOverride)
 	{
 		FString NameFormat;
 		if (Prefix != NAME_None)
@@ -23,6 +24,13 @@ namespace LiveLinkAnimationVirtualSubjectUtils
 
 		for (const FName& NewBoneName : NewBoneNames)
 		{
+			int32 Index = BoneNames.IndexOfByKey(NewBoneName);
+			if (Index != INDEX_NONE)
+			{
+				FString OverridenName = TEXT("__REPLACED_BONE_") + NewBoneName.ToString();
+				NamesToOverride.Add(Index, *OverridenName);
+			}
+
 			BoneNames.Add(*(NameFormat + NewBoneName.ToString()));
 		}
 	}
@@ -60,19 +68,35 @@ ULiveLinkAnimationVirtualSubject::ULiveLinkAnimationVirtualSubject()
 
 void ULiveLinkAnimationVirtualSubject::Update()
 {
-	Super::Update();
+	// Invalid the snapshot
+	InvalidateFrameData();
+
+	UpdateTranslatorsForThisFrame();
 
 	TArray<FLiveLinkSubjectKey> ActiveSubjects = LiveLinkClient->GetSubjects(false, false);
 
 	if (AreSubjectsValid(ActiveSubjects))
 	{
 		TArray<FLiveLinkSubjectFrameData> SubjectSnapshot;
+
+		if (bSubjectsNeedSorting)
+		{
+			SortSubjects();
+		}
+
 		if (BuildSubjectSnapshot(SubjectSnapshot))
 		{
 			BuildSkeleton(SubjectSnapshot);
 			BuildFrame(SubjectSnapshot);
 		}
-
+	}
+	else
+	{
+		for (const FLiveLinkVirtualSubjectBoneAttachment& Attachment : Attachments)
+		{
+			// Update error messages on the attachments.
+			Attachment.IsValid(Subjects);
+		}
 	}
 }
 
@@ -126,23 +150,43 @@ void ULiveLinkAnimationVirtualSubject::BuildSkeleton(const TArray<FLiveLinkSubje
 {
 	if (DoesSkeletonNeedRebuilding())
 	{
+		ChildBonesInfo.Reset();
+		BoneNameToIndex.Reset();
+
 		FLiveLinkStaticDataStruct StaticData(FLiveLinkSkeletonStaticData::StaticStruct());
 		FLiveLinkSkeletonStaticData* SkeletonData = StaticData.Cast<FLiveLinkSkeletonStaticData>();
 
-		TArray<FName> BoneNames{ TEXT("Root") };
-		TArray<int32> BoneParents{ INDEX_NONE };
-
 		check(InSubjectSnapshots.Num() == Subjects.Num());
-		for (int32 i = 0; i < InSubjectSnapshots.Num(); ++i)
+
+		TArray<FName> BoneNames{ };
+		TArray<int32> BoneParents{ };
+
+		TMap<int32, FName> NamesToOverride;
+
+		for (int32 Index = 0; Index < InSubjectSnapshots.Num(); ++Index)
 		{
-			const FLiveLinkSubjectFrameData& SubjectSnapShotData = InSubjectSnapshots[i];
+			const FLiveLinkSubjectFrameData& SubjectSnapShotData = InSubjectSnapshots[Index];
 			check(SubjectSnapShotData.StaticData.IsValid());
 			const FLiveLinkSkeletonStaticData* SubjectSkeletonData = SubjectSnapShotData.StaticData.Cast<FLiveLinkSkeletonStaticData>();
 
-			const FName BonePrefix = bAppendSubjectNameToBones ? Subjects[i] : NAME_None;
-			LiveLinkAnimationVirtualSubjectUtils::AddToBoneNames(BoneNames, SubjectSkeletonData->GetBoneNames(), BonePrefix);
+			const FName BonePrefix = bAppendSubjectNameToBones ? Subjects[Index] : NAME_None;
+
+			LiveLinkAnimationVirtualSubjectUtils::AddToBoneNames(BoneNames, SubjectSkeletonData->GetBoneNames(), BonePrefix, NamesToOverride);
 			LiveLinkAnimationVirtualSubjectUtils::AddToBoneParents(BoneParents, SubjectSkeletonData->GetBoneParents());
-			SkeletonData->PropertyNames.Append(SubjectSkeletonData->PropertyNames);
+
+			// Cache bone names to bone index, we need to use both the subject and bone names to make sure
+			// that we can find a parent bone even if it's in a name conflict.
+			for (int32 BoneIndex = 0; BoneIndex < BoneNames.Num(); BoneIndex++)
+			{
+				BoneNameToIndex.FindOrAdd({ Subjects[Index], BoneNames[BoneIndex] }) = BoneIndex;
+			}
+		}
+
+		ProcessAttachmentsForStaticData(BoneParents);
+
+		for (const TPair<int32, FName>& Names : NamesToOverride)
+		{
+			BoneNames[Names.Key] = Names.Value;
 		}
 
 		SkeletonData->SetBoneNames(BoneNames);
@@ -151,6 +195,8 @@ void ULiveLinkAnimationVirtualSubject::BuildSkeleton(const TArray<FLiveLinkSubje
 		UpdateStaticDataSnapshot(MoveTemp(StaticData));
 
 		bInvalidate = false;
+
+		PostSkeletonRebuild();
 	}
 }
 
@@ -161,14 +207,13 @@ void ULiveLinkAnimationVirtualSubject::BuildFrame(const TArray<FLiveLinkSubjectF
 	FLiveLinkAnimationFrameData* NewSnapshotFrameData = NewFrameData.Cast<FLiveLinkAnimationFrameData>();
 
 	NewSnapshotFrameData->Transforms.Reset(SnapshotSkeletonData->GetBoneNames().Num());
-	NewSnapshotFrameData->Transforms.Add(FTransform::Identity);
 	NewSnapshotFrameData->MetaData.StringMetaData.Empty();
 
 	//Go over each subject snapshot and take transforms and curves
 	check(InSubjectSnapshots.Num() == Subjects.Num());
-	for (int32 i = 0; i < InSubjectSnapshots.Num(); ++i)
+	for (int32 Index = 0; Index < InSubjectSnapshots.Num(); ++Index)
 	{
-		const FLiveLinkSubjectFrameData& SubjectSnapShotData = InSubjectSnapshots[i];
+		const FLiveLinkSubjectFrameData& SubjectSnapShotData = InSubjectSnapshots[Index];
 		check(SubjectSnapShotData.FrameData.IsValid());
 		const FLiveLinkAnimationFrameData* SubjectFrameData = SubjectSnapShotData.FrameData.Cast<FLiveLinkAnimationFrameData>();
 
@@ -176,17 +221,162 @@ void ULiveLinkAnimationVirtualSubject::BuildFrame(const TArray<FLiveLinkSubjectF
 		NewSnapshotFrameData->PropertyValues.Append(SubjectFrameData->PropertyValues);
 		for (const auto& MetaDatum : SubjectFrameData->MetaData.StringMetaData)
 		{
-			const FName QualifiedKey = FName(*(Subjects[i].ToString() + MetaDatum.Key.ToString()));
-			NewSnapshotFrameData->MetaData.StringMetaData.Emplace(Subjects[i], MetaDatum.Value);
+			NewSnapshotFrameData->MetaData.StringMetaData.Emplace(Subjects[Index], MetaDatum.Value);
 		}
 	}
 
+	ProcessAttachmentsForFrameData(NewSnapshotFrameData);
+
 	UpdateFrameDataSnapshot(MoveTemp(NewFrameData));
 }
 
+
+void ULiveLinkAnimationVirtualSubject::SortSubjects()
+{
+	if (Attachments.Num() == 0)
+	{
+		return;
+	}
+
+	bool bOneValidAttachment = false;
+	for (const FLiveLinkVirtualSubjectBoneAttachment& Attachment : Attachments)
+	{
+		if (Attachment.IsValid(Subjects))
+		{
+			bOneValidAttachment = true;
+			break;
+		}
+	}
+
+	if (!bOneValidAttachment)
+	{
+		return;
+	}
+
+	TMap<FLiveLinkSubjectName, TArray<FLiveLinkSubjectName>> ParentToChildren;
+	for (const FLiveLinkVirtualSubjectBoneAttachment& Attachment : Attachments)
+	{
+		if (Attachment.ChildSubject != FName(NAME_None))
+		{
+			ParentToChildren.FindOrAdd(Attachment.ParentSubject).Add(Attachment.ChildSubject);
+		}
+	}
+
+	auto FindDependencies = [&ParentToChildren](FLiveLinkSubjectName SubjectName) -> TArray<FLiveLinkSubjectName>
+	{
+		if (TArray<FLiveLinkSubjectName>* Children = ParentToChildren.Find(SubjectName))
+		{
+			return *Children;
+		}
+
+		return {};
+	};
+
+	if (!Algo::TopologicalSort(Subjects, FindDependencies))
+	{
+		UE_LOG(LogTemp, Warning, TEXT("Circular dependency present in attachments."));
+	}
+
+	Algo::Reverse(Subjects);
+
+	bSubjectsNeedSorting = false;
+}
+
+void ULiveLinkAnimationVirtualSubject::ProcessAttachmentsForStaticData(TArray<int32>& InOutBoneParents)
+{
+	// Compute the bone parents from the attachments and modify the bone array.
+	int32 GlobalParentIndex = INDEX_NONE; // The global index of the parent bone.
+	int32 GlobalChildIndex = INDEX_NONE; // The global index of the bone that will be attached to the parent bone.
+	TSet<FLiveLinkSubjectName> SubjectsRequiredByAttachments;
+	for (const FLiveLinkVirtualSubjectBoneAttachment& Attachment : Attachments)
+	{
+		if (Attachment.IsValid(Subjects))
+		{
+			SubjectsRequiredByAttachments.Add(Attachment.ChildSubject);
+			SubjectsRequiredByAttachments.Add(Attachment.ParentSubject);
+
+			// 1. Find global bone index for parent and child bone
+			if (int32* ParentIndex = BoneNameToIndex.Find({ Attachment.ParentSubject, Attachment.ParentBone }))
+			{
+				GlobalParentIndex = *ParentIndex;
+			}
+
+			if (int32* ChildIndex = BoneNameToIndex.Find({ Attachment.ChildSubject, Attachment.ChildBone }))
+			{
+				GlobalChildIndex = *ChildIndex;
+			}
+
+			if (GlobalParentIndex == INDEX_NONE || GlobalChildIndex == INDEX_NONE)
+			{
+				// Skip this attachment if we couldn't find either the parent or child bone.
+				continue;
+			}
+
+			// 2. Override the bone parents according to the attachments
+			InOutBoneParents[GlobalChildIndex] = GlobalParentIndex;
+
+			// 3. Store the info for the attachment child.
+			FTransform Offset = FTransform::Identity;
+			Offset.SetTranslation(Attachment.LocationOffset);
+			Offset.SetRotation(Attachment.RotationOffset.Quaternion());
+
+			FChildBoneInfo ChildBoneInfo;
+			ChildBoneInfo.Offset = Offset;
+			ChildBoneInfo.ParentBone = GlobalParentIndex;
+			ChildBoneInfo.bIgnoreParentLocation = Attachment.bIgnoreParentLocation;
+			ChildBoneInfo.bIgnoreParentRotation = Attachment.bIgnoreParentRotation;
+
+			ChildBonesInfo.Add(GlobalChildIndex, MoveTemp(ChildBoneInfo));
+		}
+	}
+
+}
+
+void ULiveLinkAnimationVirtualSubject::ProcessAttachmentsForFrameData(FLiveLinkAnimationFrameData* SnapshotFrameData)
+{
+	// Apply transformations specified by the attachments.
+	for (const TPair<int32, FChildBoneInfo>& ChildInfo : ChildBonesInfo)
+	{
+		// Apply offsets specified by the attachments.
+		FTransform Offset = ChildInfo.Value.Offset;
+
+		const FChildBoneInfo& ChildBoneInfo = ChildInfo.Value;
+
+		FTransform ParentTransform = SnapshotFrameData->Transforms[ChildBoneInfo.ParentBone];
+		FTransform ModifiedParentTransform = FTransform::Identity;
+		if (!ChildBoneInfo.bIgnoreParentLocation)
+		{
+			ModifiedParentTransform.SetLocation(ParentTransform.GetLocation());
+		}
+		if (!ChildBoneInfo.bIgnoreParentRotation)
+		{
+			ModifiedParentTransform.SetRotation(ParentTransform.GetRotation());
+		}
+
+		ModifiedParentTransform.SetScale3D(ParentTransform.GetScale3D());
+
+		// Also apply an offset from the attach target to the child bone.
+		Offset = Offset * ModifiedParentTransform;
+
+		SnapshotFrameData->Transforms[ChildInfo.Key] = SnapshotFrameData->Transforms[ChildInfo.Key] * Offset;
+	}
+}
+
 bool ULiveLinkAnimationVirtualSubject::DoesSkeletonNeedRebuilding() const
 {
-	return !HasValidStaticData() || bInvalidate;
+	if (!HasValidStaticData())
+	{
+		UE_LOG(LogTemp, Warning, TEXT("No static data"));
+		return true;
+	}
+
+	if (bInvalidate)
+	{
+		UE_LOG(LogTemp, Warning, TEXT("Invalidate"));
+		return true;
+	}
+
+	return false;
 }
 
 #if WITH_EDITOR
@@ -194,8 +384,24 @@ void ULiveLinkAnimationVirtualSubject::PostEditChangeProperty(struct FPropertyCh
 {
 	Super::PostEditChangeProperty(PropertyChangedEvent);
 
-	//If our properties have changed, force a skeleton rebuild for next frame
+	// When modifying an attachment, we only want to invalidate the static data if we modify a parent/child bone or subject.
+	if (PropertyChangedEvent.GetMemberPropertyName() == GET_MEMBER_NAME_CHECKED(ULiveLinkAnimationVirtualSubject, Attachments))
+	{
+		if (PropertyChangedEvent.GetPropertyName() == GET_MEMBER_NAME_CHECKED(FLiveLinkVirtualSubjectBoneAttachment, ParentBone)
+		|| PropertyChangedEvent.GetPropertyName() == GET_MEMBER_NAME_CHECKED(FLiveLinkVirtualSubjectBoneAttachment, ChildBone)
+		|| PropertyChangedEvent.GetPropertyName() == GET_MEMBER_NAME_CHECKED(FLiveLinkVirtualSubjectBoneAttachment, ParentSubject)
+		|| PropertyChangedEvent.GetPropertyName() == GET_MEMBER_NAME_CHECKED(FLiveLinkVirtualSubjectBoneAttachment, ChildSubject))
+		{
+			bSubjectsNeedSorting = true;
+			InvalidateStaticData();
+		}
+	}
+	else
+	{
+		InvalidateStaticData();
+	}
+
+
 	bInvalidate = true;
 }
 #endif //WITH_EDITOR
-
diff --git a/Engine/Plugins/Animation/LiveLink/Source/LiveLink/Private/LiveLinkVirtualSubjectBoneAttachment.cpp b/Engine/Plugins/Animation/LiveLink/Source/LiveLink/Private/LiveLinkVirtualSubjectBoneAttachment.cpp
new file mode 100644
--- /dev/null
+++ b/Engine/Plugins/Animation/LiveLink/Source/LiveLink/Private/LiveLinkVirtualSubjectBoneAttachment.cpp
@@ -0,0 +1,30 @@
+// Copyright Epic Games, Inc. All Rights Reserved.
+
+#include "LiveLinkVirtualSubjectBoneAttachment.h"
+
+#define LOCTEXT_NAMESPACE "VirtualSubjectBoneAttachment"
+
+bool FLiveLinkVirtualSubjectBoneAttachment::IsValid(const TArray<FLiveLinkSubjectName>& ActiveSubjects) const
+{
+	if (ParentSubject.IsNone() || ChildSubject.IsNone())
+	{
+		LastError = LOCTEXT("InvalidSubjectError", "One or more subject name is not specified.");
+	}
+	else if (ParentBone.IsNone() || ChildBone.IsNone())
+	{
+		LastError = LOCTEXT("InvalidBoneError", "One or more bone name is not specified.");
+	}
+	else if (!ActiveSubjects.Contains(ParentSubject) || !ActiveSubjects.Contains(ChildSubject))
+	{
+		LastError = LOCTEXT("DisabledSubject", "One or more subject is not enabled.");
+	}
+	else
+	{
+		LastError = FText::GetEmpty();
+		return true;
+	}
+
+	return false;
+}
+
+#undef LOCTEXT_NAMESPACE
diff --git a/Engine/Plugins/Animation/LiveLink/Source/LiveLink/Public/LiveLinkAnimationVirtualSubject.h b/Engine/Plugins/Animation/LiveLink/Source/LiveLink/Public/LiveLinkAnimationVirtualSubject.h
--- a/Engine/Plugins/Animation/LiveLink/Source/LiveLink/Public/LiveLinkAnimationVirtualSubject.h
+++ b/Engine/Plugins/Animation/LiveLink/Source/LiveLink/Public/LiveLinkAnimationVirtualSubject.h
@@ -3,6 +3,7 @@
 #pragma once
 
 #include "LiveLinkVirtualSubject.h"
+#include "LiveLinkVirtualSubjectBoneAttachment.h"
 
 
 #include "LiveLinkAnimationVirtualSubject.generated.h"
@@ -11,7 +12,7 @@ struct FLiveLinkSubjectFrameData;
 struct FLiveLinkSubjectKey;
 
 
-// A Skeleton virtual subject is an assembly of different subjects supporting the animation role
+/** A Skeleton virtual subject is an assembly of different subjects supporting the animation role */
 UCLASS(meta=(DisplayName="Animation Virtual Subject"))
 class LIVELINK_API ULiveLinkAnimationVirtualSubject : public ULiveLinkVirtualSubject
 {
@@ -29,28 +30,61 @@ public:
 	//~ End UObject interface
 
 protected:
-	// Validates current source subjects
+	/** Validates current source subjects */
 	bool AreSubjectsValid(const TArray<FLiveLinkSubjectKey>& InActiveSubjects) const;
 
 	bool BuildSubjectSnapshot(TArray<FLiveLinkSubjectFrameData>& OutSnapshot);
 
-	// Builds a new ref skeleton based on the current subject state (can early out if ref skeleton is already up to date)
-	void BuildSkeleton(const TArray<FLiveLinkSubjectFrameData>& InSubjectSnapshots);
+	/** Builds a new ref skeleton based on the current subject state(can early out if ref skeleton is already up to date) */
+	virtual void BuildSkeleton(const TArray<FLiveLinkSubjectFrameData>& InSubjectSnapshots);
 
-	void BuildFrame(const TArray<FLiveLinkSubjectFrameData>& InSubjectSnapshots);
+	virtual void BuildFrame(const TArray<FLiveLinkSubjectFrameData>& InSubjectSnapshots);
 
-	// Tests to see if current ref skeleton is up to data
+	/** Called after the update method. */
+	virtual void PostSkeletonRebuild() {}
+
+	/** Tests to see if current ref skeleton is up to data */
 	bool DoesSkeletonNeedRebuilding() const;
 	bool bInvalidate;
 
+	/** Sort the subjects according to the dependencies introduced by the attachments. */
+	void SortSubjects();
+
+	/** Build the ChildBonesInfo list from the specified attachments. */
+	void ProcessAttachmentsForStaticData(TArray<int32>& InOutBoneParents);
+
+	/**  Apply transform offsets generated by the attachments to the frame data. */
+	void ProcessAttachmentsForFrameData(struct FLiveLinkAnimationFrameData* SnapshotFrameData);
+
 public:
+	/** List of bone attachments between subjects. */
+	UPROPERTY(EditAnywhere, Category = "LiveLink")
+	TArray<FLiveLinkVirtualSubjectBoneAttachment> Attachments;
 
-	//Whether to append SubjectName to each bones part of the virtual hierarchy
+	/** Whether to append SubjectName to each bones part of the virtual hierarchy */
 	UPROPERTY(EditAnywhere, Category = "Settings")
 	bool bAppendSubjectNameToBones;
-};
 
-#if UE_ENABLE_INCLUDE_ORDER_DEPRECATED_IN_5_2
-#include "CoreMinimal.h"
-#include "ILiveLinkClient.h"
-#endif
+private:
+	/** Information about a bone that needs to be attached to a parent subject. */
+	struct FChildBoneInfo
+	{
+		/** Parent bone global index. */
+		int32 ParentBone = INDEX_NONE;
+		/** Offset of the child bone to its parent bone. */
+		FTransform Offset = FTransform::Identity;
+		/** Whether to ignore parent location when offsetting the child bone relative to its parent. */
+		bool bIgnoreParentLocation = false;
+		/** Whether to ignore parent rotation when offsetting the child bone relative to its parent. */
+		bool bIgnoreParentRotation = false;
+	};
+
+	/** Map of global child indices to their respective bone info structure.  */
+	TMap<int32, FChildBoneInfo> ChildBonesInfo;
+
+	/** Lookup table of Subject + bone name to index in the concatenated list of bone names. */
+	TMap<TPair<FName, FName>, int32> BoneNameToIndex;
+
+	/** Whether the subjects need to be sorted because an attachment has been modified. */
+	bool bSubjectsNeedSorting = false;
+};
diff --git a/Engine/Plugins/Animation/LiveLink/Source/LiveLink/Public/LiveLinkVirtualSubjectBoneAttachment.h b/Engine/Plugins/Animation/LiveLink/Source/LiveLink/Public/LiveLinkVirtualSubjectBoneAttachment.h
new file mode 100644
--- /dev/null
+++ b/Engine/Plugins/Animation/LiveLink/Source/LiveLink/Public/LiveLinkVirtualSubjectBoneAttachment.h
@@ -0,0 +1,57 @@
+// Copyright Epic Games, Inc. All Rights Reserved.
+
+#pragma once
+
+#include "LiveLinkTypes.h"
+#include "Math/MathFwd.h"
+
+
+#include "LiveLinkVirtualSubjectBoneAttachment.generated.h"
+
+/**
+ * A bone attachment used by the virtual subjects to attach 2 subjects together by specifying a child and parent bone.
+ */
+USTRUCT()
+struct LIVELINK_API FLiveLinkVirtualSubjectBoneAttachment
+{
+	GENERATED_BODY()
+
+public:
+	/** Returns whether the attachment can be used. */
+	bool IsValid(const TArray<FLiveLinkSubjectName>& ActiveSubjects) const;
+
+	/** Subject that the child subject will be attached to. */
+	UPROPERTY(EditAnywhere, DisplayName = "Parent", Category = "Live Link")
+	FLiveLinkSubjectName ParentSubject;
+
+	/** Name of the bone in the parent subject that will serve as a parent to the child bone of the child subject. */
+	UPROPERTY(EditAnywhere, Category = "Live Link")
+	FName ParentBone;
+
+	/** Subject that will be attached to the parent. */
+	UPROPERTY(EditAnywhere, DisplayName = "Child", Category = "Live Link")
+	FLiveLinkSubjectName ChildSubject;
+
+	/** Bone that will be attached to the parent bone.  */
+	UPROPERTY(EditAnywhere, Category = "Live Link")
+	FName ChildBone;
+
+	/** Location of the component relative to its parent */
+	UPROPERTY(EditAnywhere,  Category = "Live Link", meta=(LinearDeltaSensitivity = "1", Delta = "1.0"))
+	FVector LocationOffset;
+
+	/** Rotation of the component relative to its parent */
+	UPROPERTY(EditAnywhere, Category = "Live Link", meta=(LinearDeltaSensitivity = "1", Delta = "1.0"))
+	FRotator RotationOffset;
+
+	/** Whether to ignore the parent bone location when attaching the child bone. */
+	UPROPERTY(EditAnywhere, Category = "Live Link")
+    bool bIgnoreParentLocation = false;
+
+	/** Whether to ignore the parent bone rotation when attaching the child bone. */
+    UPROPERTY(EditAnywhere, Category = "Live Link")
+    bool bIgnoreParentRotation = false;
+
+	/** If not empty, hold the last error text set by the IsValid method. Used to bubble up errors to the UI. */
+	mutable FText LastError;
+};
