LiveLink Virtual Subjects - Fix temporary changes accidentally getting in the last CL

diff --git a/Engine/Plugins/Animation/LiveLink/Source/LiveLink/Private/LiveLinkAnimationVirtualSubject.cpp b/Engine/Plugins/Animation/LiveLink/Source/LiveLink/Private/LiveLinkAnimationVirtualSubject.cpp
--- a/Engine/Plugins/Animation/LiveLink/Source/LiveLink/Private/LiveLinkAnimationVirtualSubject.cpp
+++ b/Engine/Plugins/Animation/LiveLink/Source/LiveLink/Private/LiveLinkAnimationVirtualSubject.cpp
@@ -198,6 +198,11 @@ void ULiveLinkAnimationVirtualSubject::BuildSkeleton(const TArray<FLiveLinkSubje
 
 		PostSkeletonRebuild();
 	}
+	else
+	{
+		TArray<int32> BoneParents;
+		ProcessAttachmentsForStaticData(BoneParents);
+	}
 }
 
 void ULiveLinkAnimationVirtualSubject::BuildFrame(const TArray<FLiveLinkSubjectFrameData>& InSubjectSnapshots)
@@ -313,7 +318,10 @@ void ULiveLinkAnimationVirtualSubject::ProcessAttachmentsForStaticData(TArray<in
 			}
 
 			// 2. Override the bone parents according to the attachments
-			InOutBoneParents[GlobalChildIndex] = GlobalParentIndex;
+			if (InOutBoneParents.IsValidIndex(GlobalChildIndex))
+			{
+				InOutBoneParents[GlobalChildIndex] = GlobalParentIndex;
+			}
 
 			// 3. Store the info for the attachment child.
 			FTransform Offset = FTransform::Identity;
@@ -329,7 +337,6 @@ void ULiveLinkAnimationVirtualSubject::ProcessAttachmentsForStaticData(TArray<in
 			ChildBonesInfo.Add(GlobalChildIndex, MoveTemp(ChildBoneInfo));
 		}
 	}
-
 }
 
 void ULiveLinkAnimationVirtualSubject::ProcessAttachmentsForFrameData(FLiveLinkAnimationFrameData* SnapshotFrameData)
@@ -364,19 +371,7 @@ void ULiveLinkAnimationVirtualSubject::ProcessAttachmentsForFrameData(FLiveLinkA
 
 bool ULiveLinkAnimationVirtualSubject::DoesSkeletonNeedRebuilding() const
 {
-	if (!HasValidStaticData())
-	{
-		UE_LOG(LogTemp, Warning, TEXT("No static data"));
-		return true;
-	}
-
-	if (bInvalidate)
-	{
-		UE_LOG(LogTemp, Warning, TEXT("Invalidate"));
-		return true;
-	}
-
-	return false;
+	return !HasValidStaticData() || bInvalidate;
 }
 
 #if WITH_EDITOR
@@ -393,15 +388,14 @@ void ULiveLinkAnimationVirtualSubject::PostEditChangeProperty(struct FPropertyCh
 		|| PropertyChangedEvent.GetPropertyName() == GET_MEMBER_NAME_CHECKED(FLiveLinkVirtualSubjectBoneAttachment, ChildSubject))
 		{
 			bSubjectsNeedSorting = true;
+			bInvalidate = true;
 			InvalidateStaticData();
 		}
 	}
 	else
 	{
+		bInvalidate = true;
 		InvalidateStaticData();
 	}
-
-
-	bInvalidate = true;
 }
 #endif //WITH_EDITOR
