LiveLink - Add thread safe method to get subject static data

diff --git a/Engine/Plugins/Animation/LiveLink/Source/LiveLink/Private/LiveLinkClient.cpp b/Engine/Plugins/Animation/LiveLink/Source/LiveLink/Private/LiveLinkClient.cpp
--- a/Engine/Plugins/Animation/LiveLink/Source/LiveLink/Private/LiveLinkClient.cpp
+++ b/Engine/Plugins/Animation/LiveLink/Source/LiveLink/Private/LiveLinkClient.cpp
@@ -1518,6 +1518,20 @@ UObject* FLiveLinkClient::GetSubjectSettings(const FLiveLinkSubjectKey& InSubjec
 	return nullptr;
 }
 
+const FLiveLinkStaticDataStruct* FLiveLinkClient::GetSubjectStaticData_AnyThread(const FLiveLinkSubjectKey& InSubjectKey) const
+{
+	FScopeLock Lock(&CollectionAccessCriticalSection);
+	if (const FLiveLinkCollectionSubjectItem* SubjectItem = Collection->FindSubject(InSubjectKey))
+	{
+		if (FLiveLinkSubject* LiveLinkSubject = SubjectItem->GetLiveSubject())
+		{
+			return &LiveLinkSubject->GetStaticData();
+		}
+	}
+
+	return nullptr;
+}
+
 void FLiveLinkClient::RegisterForFrameDataReceived(const FLiveLinkSubjectKey& InSubjectKey, const FOnLiveLinkSubjectStaticDataReceived::FDelegate& OnStaticDataReceived_AnyThread, const FOnLiveLinkSubjectFrameDataReceived::FDelegate& OnFrameDataReceived_AnyThread, FDelegateHandle& OutStaticDataReceivedHandle, FDelegateHandle& OutFrameDataReceivedHandle)
 {
 	OutStaticDataReceivedHandle.Reset();
diff --git a/Engine/Plugins/Animation/LiveLink/Source/LiveLink/Public/LiveLinkClient.h b/Engine/Plugins/Animation/LiveLink/Source/LiveLink/Public/LiveLinkClient.h
--- a/Engine/Plugins/Animation/LiveLink/Source/LiveLink/Public/LiveLinkClient.h
+++ b/Engine/Plugins/Animation/LiveLink/Source/LiveLink/Public/LiveLinkClient.h
@@ -116,6 +116,7 @@ public:
 	virtual TArray<FLiveLinkTime> GetSubjectFrameTimes(FLiveLinkSubjectName SubjectName) const override;
 	virtual ULiveLinkSourceSettings* GetSourceSettings(const FGuid& SourceGuid) const override;
 	virtual UObject* GetSubjectSettings(const FLiveLinkSubjectKey& SubjectKey) const override;
+	virtual const FLiveLinkStaticDataStruct* GetSubjectStaticData_AnyThread(const FLiveLinkSubjectKey& InSubjectKey) const override;
 
 
 	virtual bool EvaluateFrameFromSource_AnyThread(const FLiveLinkSubjectKey& SubjectKey, TSubclassOf<ULiveLinkRole> Role, FLiveLinkSubjectFrameData& OutFrame) override;
@@ -224,9 +225,6 @@ private:
 	/** Cache the game thread values to be reused on any thread */
 	void CacheValues();
 
-	/** Registered with each subject and called when it changes */
-	void OnSubjectChangedHandler();
-
 	void PushSubjectStaticData_Internal(FPendingSubjectStatic&& SubjectStaticData);
 	void PushSubjectFrameData_Internal(FPendingSubjectFrame&& SubjectFrameData);
 
@@ -248,6 +246,7 @@ protected:
 
 	/** The current collection used. */
 	TUniquePtr<FLiveLinkSourceCollection> Collection;
+
 private:
 	/** Pending static info to add to a subject. */
 	TArray<FPendingSubjectStatic> SubjectStaticToPush;
diff --git a/Engine/Source/Runtime/LiveLinkInterface/Public/ILiveLinkClient.h b/Engine/Source/Runtime/LiveLinkInterface/Public/ILiveLinkClient.h
--- a/Engine/Source/Runtime/LiveLinkInterface/Public/ILiveLinkClient.h
+++ b/Engine/Source/Runtime/LiveLinkInterface/Public/ILiveLinkClient.h
@@ -301,6 +301,11 @@ public:
 	 */
 	virtual UObject* GetSubjectSettings(const FLiveLinkSubjectKey& SubjectKey) const = 0;
 
+	/**
+	 * Utility method to grab a subject's static data. Used by the RecordingController when static data is missing from the recording. 
+	 */
+	virtual const FLiveLinkStaticDataStruct* GetSubjectStaticData_AnyThread(const FLiveLinkSubjectKey& InSubjectKey) const = 0;
+
 	/**
 	 * Return the evaluated subject from a specific source snapshot for a specific role.
 	 * A subject could have to go through a translator to output in the desired role.
