LiveLink - Add LiveLinkTransformRoleToAnimation translator

Add LiveLinkTransformRoleToAnimation role to translate Transform frame data to skeleton data.
Remove SubjectAdded delegate from LiveLinkHub since static data broadcast is already handed elsewhere

diff --git a/Engine/Plugins/Animation/LiveLink/Source/LiveLink/Private/Translator/LiveLinkTransformRoleToAnimation.cpp b/Engine/Plugins/Animation/LiveLink/Source/LiveLink/Private/Translator/LiveLinkTransformRoleToAnimation.cpp
new file mode 100644
--- /dev/null
+++ b/Engine/Plugins/Animation/LiveLink/Source/LiveLink/Private/Translator/LiveLinkTransformRoleToAnimation.cpp
@@ -0,0 +1,100 @@
+// Copyright Epic Games, Inc. All Rights Reserved.
+
+#include "Translator/LiveLinkTransformRoleToAnimation.h"
+
+#include "Roles/LiveLinkAnimationRole.h"
+#include "Roles/LiveLinkAnimationTypes.h"
+#include "Roles/LiveLinkTransformRole.h"
+#include "Roles/LiveLinkTransformTypes.h"
+
+#include UE_INLINE_GENERATED_CPP_BY_NAME(LiveLinkTransformRoleToAnimation)
+
+
+/**
+ * ULiveLinkTransformRoleToAnimation::FLiveLinkTransformRoleToAnimationWorker
+ */
+TSubclassOf<ULiveLinkRole> ULiveLinkTransformRoleToAnimation::FLiveLinkTransformRoleToAnimationWorker::GetFromRole() const
+{
+	return ULiveLinkTransformRole::StaticClass();
+}
+
+TSubclassOf<ULiveLinkRole> ULiveLinkTransformRoleToAnimation::FLiveLinkTransformRoleToAnimationWorker::GetToRole() const
+{
+	return ULiveLinkAnimationRole::StaticClass();
+}
+
+bool ULiveLinkTransformRoleToAnimation::FLiveLinkTransformRoleToAnimationWorker::Translate(const FLiveLinkStaticDataStruct& InStaticData, const FLiveLinkFrameDataStruct& InFrameData, FLiveLinkSubjectFrameData& OutTranslatedFrame) const
+{
+	if (!InStaticData.IsValid() || !InFrameData.IsValid())
+	{
+		return false;
+	}
+
+
+	const FLiveLinkTransformStaticData* TransformStaticData = InStaticData.Cast<FLiveLinkTransformStaticData>();
+	const FLiveLinkTransformFrameData* FrameData = InFrameData.Cast<FLiveLinkTransformFrameData>();
+
+	if (TransformStaticData == nullptr || FrameData == nullptr)
+	{
+		return false;
+	}
+
+	//Allocate memory for the output translated frame with the desired type
+	OutTranslatedFrame.StaticData.InitializeWith(FLiveLinkSkeletonStaticData::StaticStruct(), nullptr);
+	OutTranslatedFrame.FrameData.InitializeWith(FLiveLinkAnimationFrameData::StaticStruct(), nullptr);
+
+	FLiveLinkSkeletonStaticData* AnimationStaticData = OutTranslatedFrame.StaticData.Cast<FLiveLinkSkeletonStaticData>();
+	FLiveLinkAnimationFrameData* AnimationFrameData = OutTranslatedFrame.FrameData.Cast<FLiveLinkAnimationFrameData>();
+	check(AnimationStaticData && AnimationFrameData);
+
+	AnimationStaticData->BoneNames.Add(OutputBoneName);
+	AnimationStaticData->BoneParents.Add(-1);
+
+	AnimationFrameData->MetaData = FrameData->MetaData;
+	AnimationFrameData->WorldTime = FrameData->WorldTime;
+	AnimationFrameData->Transforms.Reset(1);
+	AnimationFrameData->Transforms.Add(FrameData->Transform);
+
+	return true;
+}
+
+/**
+ * ULiveLinkTransformRoleToAnimation
+ */
+TSubclassOf<ULiveLinkRole> ULiveLinkTransformRoleToAnimation::GetFromRole() const
+{
+	return ULiveLinkTransformRole::StaticClass();
+}
+
+TSubclassOf<ULiveLinkRole> ULiveLinkTransformRoleToAnimation::GetToRole() const
+{
+	return ULiveLinkAnimationRole::StaticClass();
+}
+
+ULiveLinkFrameTranslator::FWorkerSharedPtr ULiveLinkTransformRoleToAnimation::FetchWorker()
+{
+	if (OutputBoneName.IsNone())
+	{
+		Instance.Reset();
+	}
+	else if (!Instance.IsValid())
+	{
+		Instance = MakeShared<FLiveLinkTransformRoleToAnimationWorker>();
+		Instance->OutputBoneName = OutputBoneName;
+	}
+
+	return Instance;
+}
+
+#if WITH_EDITOR
+void ULiveLinkTransformRoleToAnimation::PostEditChangeChainProperty(struct FPropertyChangedChainEvent& PropertyChangedEvent)
+{
+	if (PropertyChangedEvent.Property->GetFName() == GET_MEMBER_NAME_CHECKED(ThisClass, OutputBoneName))
+	{
+		Instance.Reset();
+	}
+
+	Super::PostEditChangeChainProperty(PropertyChangedEvent);
+}
+#endif //WITH_EDITOR
+
diff --git a/Engine/Plugins/Animation/LiveLink/Source/LiveLink/Public/Translator/LiveLinkTransformRoleToAnimation.h b/Engine/Plugins/Animation/LiveLink/Source/LiveLink/Public/Translator/LiveLinkTransformRoleToAnimation.h
new file mode 100644
--- /dev/null
+++ b/Engine/Plugins/Animation/LiveLink/Source/LiveLink/Public/Translator/LiveLinkTransformRoleToAnimation.h
@@ -0,0 +1,46 @@
+// Copyright Epic Games, Inc. All Rights Reserved.
+
+#pragma once
+
+#include "LiveLinkFrameTranslator.h"
+#include "LiveLinkTransformRoleToAnimation.generated.h"
+
+
+/**
+ * LiveLink Translator used to convert Transform frame data to Animation (Skeletal) frame data.
+ */
+UCLASS(meta=(DisplayName="Transform to Animation"))
+class LIVELINK_API ULiveLinkTransformRoleToAnimation : public ULiveLinkFrameTranslator
+{
+	GENERATED_BODY()
+
+public:
+	class FLiveLinkTransformRoleToAnimationWorker : public ILiveLinkFrameTranslatorWorker
+	{
+	public:
+		FName OutputBoneName;
+		virtual TSubclassOf<ULiveLinkRole> GetFromRole() const override;
+		virtual TSubclassOf<ULiveLinkRole> GetToRole() const override;
+		virtual bool Translate(const FLiveLinkStaticDataStruct& InStaticData, const FLiveLinkFrameDataStruct& InFrameData, FLiveLinkSubjectFrameData& OutTranslatedFrame) const override;
+	};
+
+protected:
+	/** Name of the resulting bone. Defaults to "Root" */
+	UPROPERTY(EditAnywhere, Category="LiveLink")
+	FName OutputBoneName = "Root";
+
+public:
+	virtual TSubclassOf<ULiveLinkRole> GetFromRole() const override;
+	virtual TSubclassOf<ULiveLinkRole> GetToRole() const override;
+	virtual ULiveLinkFrameTranslator::FWorkerSharedPtr FetchWorker() override;
+
+public:
+	//~ UObject interface
+#if WITH_EDITOR
+	virtual void PostEditChangeChainProperty(struct FPropertyChangedChainEvent& PropertyChangedEvent) override;
+#endif //WITH_EDITOR
+	//~ End UObject interface
+
+private:
+	TSharedPtr<FLiveLinkTransformRoleToAnimationWorker> Instance;
+};
diff --git a/Engine/Plugins/Animation/LiveLinkHub/Source/LiveLinkHub/Private/LiveLinkHub.cpp b/Engine/Plugins/Animation/LiveLinkHub/Source/LiveLinkHub/Private/LiveLinkHub.cpp
--- a/Engine/Plugins/Animation/LiveLinkHub/Source/LiveLinkHub/Private/LiveLinkHub.cpp
+++ b/Engine/Plugins/Animation/LiveLinkHub/Source/LiveLinkHub/Private/LiveLinkHub.cpp
@@ -77,7 +77,6 @@ void FLiveLinkHub::Initialize()
 	LiveLinkHubClient->OnStaticDataReceived_AnyThread().AddSP(this, &FLiveLinkHub::OnStaticDataReceived_AnyThread);
 	LiveLinkHubClient->OnFrameDataReceived_AnyThread().AddSP(this, &FLiveLinkHub::OnFrameDataReceived_AnyThread);
 	LiveLinkHubClient->OnSubjectMarkedPendingKill_AnyThread().AddSP(this, &FLiveLinkHub::OnSubjectMarkedPendingKill_AnyThread);
-	LiveLinkHubClient->OnLiveLinkSubjectAdded().AddSP(this, &FLiveLinkHub::OnSubjectAdded);
 
 	RegisterLiveLinkHubSettings();
 
@@ -93,7 +92,6 @@ FLiveLinkHub::~FLiveLinkHub()
 	RecordingController.Reset();
 	PlaybackController.Reset();
 
-	LiveLinkHubClient->OnLiveLinkSubjectAdded().RemoveAll(this);
 	LiveLinkHubClient->OnSubjectMarkedPendingKill_AnyThread().RemoveAll(this);
 	LiveLinkHubClient->OnFrameDataReceived_AnyThread().RemoveAll(this);
 	LiveLinkHubClient->OnStaticDataReceived_AnyThread().RemoveAll(this);
@@ -183,27 +181,6 @@ void FLiveLinkHub::OnFrameDataReceived_AnyThread(const FLiveLinkSubjectKey& InSu
 	}
 }
 
-void FLiveLinkHub::OnSubjectAdded(FLiveLinkSubjectKey InSubjectKey) const
-{
-	// Send an update to connected clients as well.
-	ULiveLinkSubjectSettings* SubjectSettings = Cast<ULiveLinkSubjectSettings>(LiveLinkHubClient->GetSubjectSettings(InSubjectKey));
-
-	if (const FLiveLinkStaticDataStruct* StaticData = LiveLinkHubClient->GetSubjectStaticData(InSubjectKey))
-	{
-		FLiveLinkStaticDataStruct StaticDataCopy;
-		StaticDataCopy.InitializeWith(*StaticData);
-
-		const FName OverridenName = GetSubjectNameOverride(InSubjectKey);
-
-		UE_LOG(LogLiveLinkHub, Verbose, TEXT("Pushed static subject %s"), *InSubjectKey.SubjectName.ToString());
-		LiveLinkProvider->UpdateSubjectStaticData(OverridenName, SubjectSettings->Role, MoveTemp(StaticDataCopy));
-	}
-	else
-	{
-		UE_LOG(LogLiveLinkHub, Verbose, TEXT("Failed to push static subject %s, static data doesn't exist"), *InSubjectKey.SubjectName.ToString());
-	}
-}
-
 void FLiveLinkHub::OnSubjectMarkedPendingKill_AnyThread(const FLiveLinkSubjectKey& InSubjectKey) const
 {
 	UE_LOG(LogLiveLinkHub, Verbose, TEXT("Removed subject %s"), *InSubjectKey.SubjectName.ToString());
diff --git a/Engine/Plugins/Animation/LiveLinkHub/Source/LiveLinkHub/Private/LiveLinkHub.h b/Engine/Plugins/Animation/LiveLinkHub/Source/LiveLinkHub/Private/LiveLinkHub.h
--- a/Engine/Plugins/Animation/LiveLinkHub/Source/LiveLinkHub/Private/LiveLinkHub.h
+++ b/Engine/Plugins/Animation/LiveLinkHub/Source/LiveLinkHub/Private/LiveLinkHub.h
@@ -75,7 +75,6 @@ private:
 	void OnStaticDataReceived_AnyThread(const FLiveLinkSubjectKey& InSubjectKey, TSubclassOf<ULiveLinkRole> InRole, const FLiveLinkStaticDataStruct& InStaticDataStruct) const;
 	void OnFrameDataReceived_AnyThread(const FLiveLinkSubjectKey& InSubjectKey, FLiveLinkFrameDataStruct& InFrameDataStruct) const;
 	void OnSubjectMarkedPendingKill_AnyThread(const FLiveLinkSubjectKey& InSubjectKey) const;
-	void OnSubjectAdded(FLiveLinkSubjectKey InSubjectKey) const;
 	//~ LiveLink Client delegates
 
 	/** Bind all available live link hub commands. */
