LiveLink - Allow setting rebroadcast provider from child classes

Fix use after movetemp in livelink client

diff --git a/Engine/Plugins/Animation/LiveLink/Source/LiveLink/Private/LiveLinkSourceCollection.cpp b/Engine/Plugins/Animation/LiveLink/Source/LiveLink/Private/LiveLinkSourceCollection.cpp
--- a/Engine/Plugins/Animation/LiveLink/Source/LiveLink/Private/LiveLinkSourceCollection.cpp
+++ b/Engine/Plugins/Animation/LiveLink/Source/LiveLink/Private/LiveLinkSourceCollection.cpp
@@ -213,8 +213,9 @@ const FLiveLinkCollectionSourceItem* FLiveLinkSourceCollection::FindVirtualSourc
 
 void FLiveLinkSourceCollection::AddSubject(FLiveLinkCollectionSubjectItem InSubject)
 {
+	FLiveLinkSubjectKey Key = InSubject.Key;
 	Subjects.Add(MoveTemp(InSubject));
-	OnLiveLinkSubjectAdded().Broadcast(InSubject.Key);
+	OnLiveLinkSubjectAdded().Broadcast(Key);
 	OnLiveLinkSubjectsChanged().Broadcast();
 }
 
diff --git a/Engine/Plugins/Animation/LiveLink/Source/LiveLink/Public/LiveLinkClient.h b/Engine/Plugins/Animation/LiveLink/Source/LiveLink/Public/LiveLinkClient.h
--- a/Engine/Plugins/Animation/LiveLink/Source/LiveLink/Public/LiveLinkClient.h
+++ b/Engine/Plugins/Animation/LiveLink/Source/LiveLink/Public/LiveLinkClient.h
@@ -247,6 +247,12 @@ protected:
 	/** The current collection used. */
 	TUniquePtr<FLiveLinkSourceCollection> Collection;
 
+	/** LiveLink Provider for rebroadcasting */
+	TSharedPtr<ILiveLinkProvider> RebroadcastLiveLinkProvider;
+
+	/** Lock to stop multiple threads accessing the Subjects from the collection at the same time */
+    mutable FCriticalSection CollectionAccessCriticalSection;
+
 private:
 	/** Pending static info to add to a subject. */
 	TArray<FPendingSubjectStatic> SubjectStaticToPush;
@@ -257,9 +263,6 @@ private:
 	/** Key funcs for looking up a set of cached keys by its layout element */
 	TMap<FLiveLinkSubjectName, FLiveLinkSubjectKey> EnabledSubjects;
 
-	/** Lock to stop multiple threads accessing the Subjects from the collection at the same time */
-	mutable FCriticalSection CollectionAccessCriticalSection;
-
 	struct FSubjectFramesAddedHandles
 	{
 		FOnLiveLinkSubjectStaticDataAdded OnStaticDataAdded;
@@ -284,8 +287,6 @@ private:
 	/** Delegate when LiveLinkClient has ticked. */
 	FSimpleMulticastDelegate OnLiveLinkTickedDelegate;
 
-	/** LiveLink Provider for rebroadcasting */
-	TSharedPtr<ILiveLinkProvider> RebroadcastLiveLinkProvider;
 	FString RebroadcastLiveLinkProviderName;
 	TSet<FLiveLinkSubjectKey> RebroadcastedSubjects;
 
