LiveLink - Add UI customizations for virtual subjects bone attachments

#okforgitub public, internal

diff --git a/Engine/Plugins/Animation/LiveLink/Source/LiveLink/Public/Translator/LiveLinkTransformRoleToAnimation.h b/Engine/Plugins/Animation/LiveLink/Source/LiveLink/Public/Translator/LiveLinkTransformRoleToAnimation.h
--- a/Engine/Plugins/Animation/LiveLink/Source/LiveLink/Public/Translator/LiveLinkTransformRoleToAnimation.h
+++ b/Engine/Plugins/Animation/LiveLink/Source/LiveLink/Public/Translator/LiveLinkTransformRoleToAnimation.h
@@ -24,7 +24,7 @@ public:
 		virtual bool Translate(const FLiveLinkStaticDataStruct& InStaticData, const FLiveLinkFrameDataStruct& InFrameData, FLiveLinkSubjectFrameData& OutTranslatedFrame) const override;
 	};
 
-protected:
+public:
 	/** Name of the resulting bone. Defaults to "Root" */
 	UPROPERTY(EditAnywhere, Category="LiveLink")
 	FName OutputBoneName = "Root";
diff --git a/Engine/Plugins/Animation/LiveLink/Source/LiveLinkEditor/Private/LiveLinkBoneAttachmentDetailCustomization.cpp b/Engine/Plugins/Animation/LiveLink/Source/LiveLinkEditor/Private/LiveLinkBoneAttachmentDetailCustomization.cpp
new file mode 100644
--- /dev/null
+++ b/Engine/Plugins/Animation/LiveLink/Source/LiveLinkEditor/Private/LiveLinkBoneAttachmentDetailCustomization.cpp
@@ -0,0 +1,352 @@
+// Copyright Epic Games, Inc. All Rights Reserved.
+
+#include "LiveLinkBoneAttachmentDetailCustomization.h"
+
+#include "DetailWidgetRow.h"
+#include "Features/IModularFeatures.h"
+#include "IDetailChildrenBuilder.h"
+#include "IPropertyUtilities.h"
+#include "ILiveLinkClient.h"
+#include "LiveLinkEditorPrivate.h"
+#include "LiveLinkVirtualSubject.h"
+#include "LiveLinkVirtualSubjectBoneAttachment.h"
+#include "Misc/Guid.h"
+#include "Roles/LiveLinkAnimationRole.h"
+#include "SLiveLinkBoneSelectionWidget.h"
+#include "Templates/SubclassOf.h"
+#include "Widgets/Images/SImage.h"
+#include "Widgets/Text/STextBlock.h"
+
+
+#define LOCTEXT_NAMESPACE "LiveLinkBoneAttachmentDetailCustomization"
+
+TSharedRef<IPropertyTypeCustomization> FLiveLinkBoneAttachmentDetailCustomization::MakeInstance()
+{
+	return MakeShared<FLiveLinkBoneAttachmentDetailCustomization>();
+}
+
+void FLiveLinkBoneAttachmentDetailCustomization::CustomizeHeader(TSharedRef<IPropertyHandle> InPropertyHandle, FDetailWidgetRow& HeaderRow, IPropertyTypeCustomizationUtils& CustomizationUtils)
+{
+	StructPropertyHandle = InPropertyHandle;
+	TSharedPtr<IPropertyUtilities> PropertyUtils = CustomizationUtils.GetPropertyUtilities();
+
+	check(CastFieldChecked<FStructProperty>(StructPropertyHandle->GetProperty())->Struct == FLiveLinkVirtualSubjectBoneAttachment::StaticStruct());
+	
+	HeaderRow.NameContent()
+	[
+		StructPropertyHandle->CreatePropertyNameWidget()
+	]
+	.ValueContent()
+	[
+		SNew(SHorizontalBox)
+		+ SHorizontalBox::Slot()
+		.AutoWidth()
+		[
+			StructPropertyHandle->CreatePropertyValueWidget()
+		]
+		+ SHorizontalBox::Slot()
+		.Padding(FMargin(4.0f, 0.0f, 0.0f, 0.0f))
+		.HAlign(HAlign_Center)
+		.VAlign(VAlign_Center)
+		.AutoWidth()
+		[
+			SNew(SImage)
+				.Image(FLiveLinkEditorPrivate::GetStyleSet()->GetBrush("LiveLinkController.WarningIcon"))
+				.ToolTipText_Raw(this, &FLiveLinkBoneAttachmentDetailCustomization::GetWarningTooltip)
+				.Visibility_Raw(this, &FLiveLinkBoneAttachmentDetailCustomization::HandleWarningVisibility)
+		]
+	].IsEnabled(MakeAttributeLambda([=] { return !InPropertyHandle->IsEditConst() && PropertyUtils->IsPropertyEditingEnabled(); }));
+}
+
+void FLiveLinkBoneAttachmentDetailCustomization::CustomizeChildren(TSharedRef<IPropertyHandle> PropertyHandle, IDetailChildrenBuilder& ChildBuilder, IPropertyTypeCustomizationUtils& CustomizationUtils)
+{
+	TSharedPtr<IPropertyUtilities> PropertyUtils = CustomizationUtils.GetPropertyUtilities();
+	uint32 NumberOfChild;
+
+	FName ParentSubjectPropertyName = GET_MEMBER_NAME_CHECKED(FLiveLinkVirtualSubjectBoneAttachment, ParentSubject);
+	FName ParentBonePropertyName = GET_MEMBER_NAME_CHECKED(FLiveLinkVirtualSubjectBoneAttachment, ParentBone);
+	FName ChildSubjectPropertyName = GET_MEMBER_NAME_CHECKED(FLiveLinkVirtualSubjectBoneAttachment, ChildSubject);
+	FName ChildBonePropertyName = GET_MEMBER_NAME_CHECKED(FLiveLinkVirtualSubjectBoneAttachment, ChildBone);
+
+	TSharedPtr<IPropertyHandle> ParentSubjectHandle = PropertyHandle->GetChildHandle(ParentSubjectPropertyName);
+	TSharedPtr<IPropertyHandle> ParentBoneHandle = PropertyHandle->GetChildHandle(ParentBonePropertyName);
+	TSharedPtr<IPropertyHandle> ChildSubjectHandle = PropertyHandle->GetChildHandle(ChildSubjectPropertyName);
+	TSharedPtr<IPropertyHandle> ChildBoneHandle = PropertyHandle->GetChildHandle(ChildBonePropertyName);
+
+	// Since bone properties are displayed inline, reset them when we reset the subject.
+	ParentSubjectHandle->SetOnPropertyResetToDefault(FSimpleDelegate::CreateLambda([ParentBoneHandle]() 
+		{
+			if (ParentBoneHandle)
+			{
+				ParentBoneHandle->ResetToDefault();
+			}
+		}));
+
+	ChildSubjectHandle->SetOnPropertyResetToDefault(FSimpleDelegate::CreateLambda([ChildBoneHandle]()
+		{
+			if (ChildBoneHandle)
+			{
+				ChildBoneHandle->ResetToDefault();
+			}
+		}));
+
+	if (PropertyHandle->GetNumChildren(NumberOfChild) == FPropertyAccess::Success)
+	{
+		for (uint32 Index = 0; Index < NumberOfChild; ++Index)
+		{
+			TSharedRef<IPropertyHandle> ChildPropertyHandle = PropertyHandle->GetChildHandle(Index).ToSharedRef();
+
+			if (ChildPropertyHandle->GetProperty()->GetFName() != ParentBonePropertyName &&
+								ChildPropertyHandle->GetProperty()->GetFName() != ChildBonePropertyName)
+			{
+				IDetailPropertyRow& DetailRow = ChildBuilder.AddProperty(ChildPropertyHandle)
+					.ShowPropertyButtons(true)
+					.IsEnabled(MakeAttributeLambda([=] { return !PropertyHandle->IsEditConst() && PropertyUtils->IsPropertyEditingEnabled(); }));
+
+				if (ChildPropertyHandle->GetProperty()->GetFName() == ParentSubjectPropertyName)
+				{
+
+					CustomizeSubjectRow(ChildPropertyHandle, ParentBoneHandle, DetailRow, CustomizationUtils, BonePickerWidgets.ParentWidget);
+				}
+				else if (ChildPropertyHandle->GetProperty()->GetFName() == ChildSubjectPropertyName)
+				{
+					CustomizeSubjectRow(ChildPropertyHandle, ChildBoneHandle, DetailRow, CustomizationUtils, BonePickerWidgets.ChildWidget);
+				}
+			}
+		}
+	}
+}
+
+void FLiveLinkBoneAttachmentDetailCustomization::CustomizeSubjectRow(TSharedPtr<IPropertyHandle> SubjectPropertyHandle, TSharedPtr<IPropertyHandle> BonePropertyHandle, IDetailPropertyRow& PropertyRow, IPropertyTypeCustomizationUtils& StructCustomizationUtils, TSharedPtr<SLiveLinkBoneSelectionWidget>& BonePickerWidget)
+{
+	FLiveLinkSubjectKey SubjectKey;
+	
+	TArray<const void*> RawData;
+	SubjectPropertyHandle->AccessRawData(RawData);
+	if (RawData.Num() && RawData[0])
+	{
+		const FLiveLinkSubjectName* SubjectName = reinterpret_cast<const FLiveLinkSubjectName*>(RawData[0]);
+
+		constexpr bool bIncludeDisabledSubject = false;
+		constexpr bool bIncludeVirtualSubject = true;
+
+		ILiveLinkClient& LiveLinkClient = IModularFeatures::Get().GetModularFeature<ILiveLinkClient>(ILiveLinkClient::ModularFeatureName);
+		TArray<FLiveLinkSubjectKey> SubjectKeys = LiveLinkClient.GetSubjects(bIncludeDisabledSubject, bIncludeVirtualSubject);
+
+		if (FLiveLinkSubjectKey* Subject = SubjectKeys.FindByPredicate([Name = *SubjectName](const FLiveLinkSubjectKey& Other) { return Other.SubjectName.Name == Name; }))
+		{
+			SubjectKey = *Subject;
+		}
+	}
+
+	BonePickerWidget = SNew(SLiveLinkBoneSelectionWidget, SubjectKey)
+					.OnBoneSelectionChanged(this, &FLiveLinkBoneAttachmentDetailCustomization::OnBoneSelected, BonePropertyHandle)
+					.OnGetSelectedBone(this, &FLiveLinkBoneAttachmentDetailCustomization::GetSelectedBone, BonePropertyHandle);
+
+	PropertyRow
+		.CustomWidget()
+		.NameContent()
+		[
+			SubjectPropertyHandle->CreatePropertyNameWidget()
+		]
+		.ValueContent()
+		[
+			SNew(SHorizontalBox)
+			+ SHorizontalBox::Slot()
+			.Padding(0.0f, 2.0f)
+			.VAlign(VAlign_Center)
+			.AutoWidth()
+			[
+				SNew(STextBlock)
+				.Font(StructCustomizationUtils.GetRegularFont())
+				.Text(LOCTEXT("SubjectLabel", "Subject"))
+			]
+			+ SHorizontalBox::Slot()
+			.Padding(4.0f, 2.0f)
+			.AutoWidth()
+			[
+				SNew(SLiveLinkSubjectRepresentationPicker)
+					.Font(StructCustomizationUtils.GetRegularFont())
+					.Value(this, &FLiveLinkBoneAttachmentDetailCustomization::GetValue, SubjectPropertyHandle)
+					.OnValueChanged(this, &FLiveLinkBoneAttachmentDetailCustomization::SetValue, SubjectPropertyHandle, BonePickerWidget)
+					.OnGetSubjects(this, &FLiveLinkBoneAttachmentDetailCustomization::GetSubjectsForPicker)
+			]
+			+ SHorizontalBox::Slot()
+			.VAlign(VAlign_Center)
+			.Padding(5.0f, 2.0f, 2.0f, 2.0f)
+			.AutoWidth()
+			[
+				SNew(STextBlock)
+					.Font(StructCustomizationUtils.GetRegularFont())
+					.Text(LOCTEXT("BoneLabel", "Bone"))
+			]
+			+ SHorizontalBox::Slot()
+			.Padding(4.0f, 2.0f)
+			.AutoWidth()
+			[
+				BonePickerWidget.ToSharedRef()
+			]
+		];
+}
+
+SLiveLinkSubjectRepresentationPicker::FLiveLinkSourceSubjectRole FLiveLinkBoneAttachmentDetailCustomization::GetValue(TSharedPtr<IPropertyHandle> Handle) const
+{
+	TArray<const void*> RawData;
+	Handle->AccessRawData(RawData);
+
+	for (const void* RawPtr : RawData)
+	{
+		if (RawPtr)
+		{
+			FLiveLinkSubjectName SubjectName = *reinterpret_cast<const FLiveLinkSubjectName*>(RawPtr);
+			return SLiveLinkSubjectRepresentationPicker::FLiveLinkSourceSubjectRole(FGuid(), SubjectName, nullptr);
+		}
+	}
+
+	return SLiveLinkSubjectRepresentationPicker::FLiveLinkSourceSubjectRole();
+}
+
+void FLiveLinkBoneAttachmentDetailCustomization::SetValue(SLiveLinkSubjectRepresentationPicker::FLiveLinkSourceSubjectRole NewValue, TSharedPtr<IPropertyHandle> Handle, TSharedPtr<SLiveLinkBoneSelectionWidget> BonePickerWidget)
+{
+	FStructProperty* StructProperty = CastFieldChecked<FStructProperty>(Handle->GetProperty());
+
+	FProperty* Property = Handle->GetProperty();
+
+	TArray<void*> RawData;
+	Handle->AccessRawData(RawData);
+	FLiveLinkSubjectKey NewSubjectKey = NewValue.ToSubjectKey();
+
+	if (BonePickerWidget)
+	{
+		BonePickerWidget->SetSubject(NewSubjectKey);
+	}
+
+	TArray<UObject*> Objects;
+	Handle->GetOuterObjects(Objects);
+
+	if (Objects.Num() == 1)
+	{
+		FLiveLinkSubjectName* Key = (FLiveLinkSubjectName*) Handle->GetValueBaseAddress((uint8*)Objects[0]);
+		*Key = NewSubjectKey.SubjectName;
+	}
+}
+
+FText FLiveLinkBoneAttachmentDetailCustomization::GetWarningTooltip() const
+{
+	if (StructPropertyHandle)
+	{
+		TArray<const void*> RawData;
+		StructPropertyHandle->AccessRawData(RawData);
+
+		for (const void* RawPtr : RawData)
+		{
+			if (RawPtr)
+			{
+				const FLiveLinkVirtualSubjectBoneAttachment* Attachment = reinterpret_cast<const FLiveLinkVirtualSubjectBoneAttachment*>(RawPtr);
+				return Attachment->LastError;
+			}
+		}
+	}
+
+	return FText::GetEmpty();
+}
+
+EVisibility FLiveLinkBoneAttachmentDetailCustomization::HandleWarningVisibility() const
+{
+	if (StructPropertyHandle)
+	{
+		TArray<const void*> RawData;
+		StructPropertyHandle->AccessRawData(RawData);
+
+		for (const void* RawPtr : RawData)
+		{
+			if (RawPtr)
+			{
+				const FLiveLinkVirtualSubjectBoneAttachment* Attachment = reinterpret_cast<const FLiveLinkVirtualSubjectBoneAttachment*>(RawPtr);
+				return Attachment->LastError.IsEmpty() ? EVisibility::Collapsed : EVisibility::Visible;
+			}
+		}
+	}
+
+	return EVisibility::Collapsed;
+}
+
+void FLiveLinkBoneAttachmentDetailCustomization::GetSubjectsForPicker(TArray<FLiveLinkSubjectKey>& OutSubjectKeys)
+{
+	ILiveLinkClient& LiveLinkClient = IModularFeatures::Get().GetModularFeature<ILiveLinkClient>(ILiveLinkClient::ModularFeatureName);
+
+	constexpr bool bIncludeDisabledSubject = true;
+	constexpr bool bIncludeVirtualSubject = true;
+	TArray<FLiveLinkSubjectKey> SubjectKeys = LiveLinkClient.GetSubjects(bIncludeDisabledSubject, bIncludeVirtualSubject);
+	TArray<FLiveLinkSubjectKey> FilteredSubjectKeys;
+
+	FLiveLinkSubjectName ParentSubjectName;
+	FLiveLinkSubjectName ChildSubjectName;
+
+	// Key of the virtual subject that owns the bone attachment.
+	FLiveLinkSubjectKey VirtualSubjectKey;
+
+	if (StructPropertyHandle)
+	{
+		TArray<UObject*> Objects;
+		StructPropertyHandle->GetOuterObjects(Objects);
+
+		if (Objects.Num() && Objects[0])
+		{
+			VirtualSubjectKey = Cast<ULiveLinkVirtualSubject>(Objects[0])->GetSubjectKey();
+		}
+
+		TArray<const void*> RawData;
+		StructPropertyHandle->AccessRawData(RawData);
+
+		if (RawData.Num() && RawData[0])
+		{
+			const FLiveLinkVirtualSubjectBoneAttachment* Attachment = reinterpret_cast<const FLiveLinkVirtualSubjectBoneAttachment*>(RawData[0]);
+			ParentSubjectName = Attachment->ParentSubject;
+			ChildSubjectName = Attachment->ChildSubject;
+		}
+	}
+
+	for (const FLiveLinkSubjectKey& SubjectKey : SubjectKeys)
+	{
+		bool bAddKey = true;
+
+		if (!LiveLinkClient.DoesSubjectSupportsRole_AnyThread(SubjectKey, ULiveLinkAnimationRole::StaticClass()))
+		{
+			bAddKey = false;
+		}
+
+		bAddKey = bAddKey && SubjectKey != VirtualSubjectKey;
+
+		bAddKey = bAddKey && SubjectKey.SubjectName != ChildSubjectName
+			&& SubjectKey.SubjectName != ParentSubjectName;
+
+		if (bAddKey)
+		{
+			OutSubjectKeys.Add(SubjectKey);
+		}
+	}
+}
+
+void FLiveLinkBoneAttachmentDetailCustomization::OnBoneSelected(FName SelectedBone, TSharedPtr<IPropertyHandle> BonePropertyHandle) const
+{
+	if (BonePropertyHandle)
+	{
+		BonePropertyHandle->SetValue(SelectedBone.ToString());
+	}
+}
+
+FName FLiveLinkBoneAttachmentDetailCustomization::GetSelectedBone(TSharedPtr<IPropertyHandle> BonePropertyHandle) const
+{
+	FName BoneName;
+	if (BonePropertyHandle)
+	{
+		FString BoneNameStr;
+		BonePropertyHandle->GetValueAsFormattedString(BoneNameStr);
+		BoneName = *BoneNameStr;
+	}
+
+	return BoneName;
+}
+
+#undef LOCTEXT_NAMESPACE
diff --git a/Engine/Plugins/Animation/LiveLink/Source/LiveLinkEditor/Private/LiveLinkBoneAttachmentDetailCustomization.h b/Engine/Plugins/Animation/LiveLink/Source/LiveLinkEditor/Private/LiveLinkBoneAttachmentDetailCustomization.h
new file mode 100644
--- /dev/null
+++ b/Engine/Plugins/Animation/LiveLink/Source/LiveLinkEditor/Private/LiveLinkBoneAttachmentDetailCustomization.h
@@ -0,0 +1,51 @@
+// Copyright Epic Games, Inc. All Rights Reserved.
+
+#pragma once
+
+#include "IPropertyTypeCustomization.h"
+#include "SLiveLinkSubjectRepresentationPicker.h"
+
+class SLiveLinkBoneSelectionWidget;
+
+class FLiveLinkBoneAttachmentDetailCustomization : public IPropertyTypeCustomization
+{
+public:
+	/** Makes a new instance of this detail layout class for a specific detail view requesting it. */
+	static TSharedRef<IPropertyTypeCustomization> MakeInstance();
+
+public:
+	// IDetailCustomization interface
+	virtual void CustomizeHeader(TSharedRef<IPropertyHandle> PropertyHandle, FDetailWidgetRow& HeaderRow, IPropertyTypeCustomizationUtils& CustomizationUtils) override;
+	virtual void CustomizeChildren(TSharedRef<IPropertyHandle> PropertyHandle, IDetailChildrenBuilder& ChildBuilder, IPropertyTypeCustomizationUtils& CustomizationUtils) override;
+
+private:
+	/** Customize a subject row to combine the subject and bone pickers. */
+	void CustomizeSubjectRow(TSharedPtr<IPropertyHandle> SubjectPropertyHandle, TSharedPtr<IPropertyHandle> BonePropertyHandle, class IDetailPropertyRow& PropertyRow, class IPropertyTypeCustomizationUtils& StructCustomizationUtils, TSharedPtr<SLiveLinkBoneSelectionWidget>& BonePickerWidget);
+	/** When the warning icon is visible, returns the text that will be shown in its tooltip. */
+	FText GetWarningTooltip() const;
+	/** Returns whether the warning symbol should be visible according to the status of the attachment. */
+	EVisibility HandleWarningVisibility() const;
+	/** Retrieves the value of a subject property on the attachment. */
+	SLiveLinkSubjectRepresentationPicker::FLiveLinkSourceSubjectRole GetValue(TSharedPtr<IPropertyHandle> Handle) const;
+	/** Sets the value of a subject property on the attachment. */
+	void SetValue(SLiveLinkSubjectRepresentationPicker::FLiveLinkSourceSubjectRole NewValue, TSharedPtr<IPropertyHandle> Handle, TSharedPtr<SLiveLinkBoneSelectionWidget> BonePickerWidget);
+	/** Populate a filtered list of subjects for the subject representation. */
+	void GetSubjectsForPicker(TArray<FLiveLinkSubjectKey>& OutSubjectKeys);
+	/** Update a bone attachment after a bone is selected. */
+	void OnBoneSelected(FName SelectedBone, TSharedPtr<IPropertyHandle> BonePropertyHandle) const;
+	/** Returns the bone from the property handle. */
+	FName GetSelectedBone(TSharedPtr<IPropertyHandle> BonePropertyHandle) const;
+
+private:
+	/** Handle to the attachment struct. */
+	TSharedPtr<IPropertyHandle> StructPropertyHandle;
+
+	struct FBonePickerWidgets
+	{
+		TSharedPtr<SLiveLinkBoneSelectionWidget> ParentWidget;
+		TSharedPtr<SLiveLinkBoneSelectionWidget> ChildWidget;
+	};
+
+	/** Holds the bone picker widgets for the parent and child bone. */
+	FBonePickerWidgets BonePickerWidgets;
+};
diff --git a/Engine/Plugins/Animation/LiveLink/Source/LiveLinkEditor/Private/LiveLinkEditorModule.cpp b/Engine/Plugins/Animation/LiveLink/Source/LiveLinkEditor/Private/LiveLinkEditorModule.cpp
--- a/Engine/Plugins/Animation/LiveLink/Source/LiveLinkEditor/Private/LiveLinkEditorModule.cpp
+++ b/Engine/Plugins/Animation/LiveLink/Source/LiveLinkEditor/Private/LiveLinkEditorModule.cpp
@@ -18,6 +18,7 @@
 #include "WorkspaceMenuStructure.h"
 #include "WorkspaceMenuStructureModule.h"
 
+#include "LiveLinkBoneAttachmentDetailCustomization.h"
 #include "LiveLinkComponentController.h"
 #include "LiveLinkComponentDetailCustomization.h"
 #include "LiveLinkComponentSettings.h"
@@ -35,9 +36,11 @@
 #include "LiveLinkSubjectNameDetailCustomization.h"
 #include "LiveLinkSubjectRepresentationDetailCustomization.h"
 #include "LiveLinkVirtualSubjectDetailCustomization.h"
+#include "LiveLinkVirtualSubjectBoneAttachment.h"
 #include "Kismet2/KismetEditorUtilities.h"
 #include "VirtualSubjects/LiveLinkBlueprintVirtualSubject.h"
 
+
 /**
  * Implements the LiveLinkEditor module.
  */
@@ -270,6 +273,7 @@ private:
 		PropertyEditorModule.RegisterCustomPropertyTypeLayout(FLiveLinkSubjectKey::StaticStruct()->GetFName(), FOnGetPropertyTypeCustomizationInstance::CreateStatic(&FLiveLinkSubjectKeyDetailCustomization::MakeInstance));
 		PropertyEditorModule.RegisterCustomPropertyTypeLayout(FLiveLinkSubjectRepresentation::StaticStruct()->GetFName(), FOnGetPropertyTypeCustomizationInstance::CreateStatic(&FLiveLinkSubjectRepresentationDetailCustomization::MakeInstance));
 		PropertyEditorModule.RegisterCustomPropertyTypeLayout(FLiveLinkSubjectName::StaticStruct()->GetFName(), FOnGetPropertyTypeCustomizationInstance::CreateStatic(&FLiveLinkSubjectNameDetailCustomization::MakeInstance));
+		PropertyEditorModule.RegisterCustomPropertyTypeLayout(FLiveLinkVirtualSubjectBoneAttachment::StaticStruct()->GetFName(), FOnGetPropertyTypeCustomizationInstance::CreateStatic(&FLiveLinkBoneAttachmentDetailCustomization::MakeInstance));
 		PropertyEditorModule.RegisterCustomClassLayout(ULiveLinkSourceSettings::StaticClass()->GetFName(), FOnGetDetailCustomizationInstance::CreateStatic(&FLiveLinkSourceSettingsDetailCustomization::MakeInstance));
 		PropertyEditorModule.RegisterCustomClassLayout(ULiveLinkComponentController::StaticClass()->GetFName(), FOnGetDetailCustomizationInstance::CreateStatic(&FLiveLinkComponentDetailCustomization::MakeInstance));
 		PropertyEditorModule.RegisterCustomClassLayout(ULiveLinkControllerBase::StaticClass()->GetFName(), FOnGetDetailCustomizationInstance::CreateStatic(&FLiveLinkControllerBaseDetailCustomization::MakeInstance));
@@ -289,6 +293,7 @@ private:
 				PropertyEditorModule->UnregisterCustomClassLayout(ULiveLinkComponentController::StaticClass()->GetFName());
 				PropertyEditorModule->UnregisterCustomClassLayout(ULiveLinkControllerBase::StaticClass()->GetFName());
 				PropertyEditorModule->UnregisterCustomClassLayout(ULiveLinkSourceSettings::StaticClass()->GetFName());
+				PropertyEditorModule->UnregisterCustomClassLayout(FLiveLinkVirtualSubjectBoneAttachment::StaticStruct()->GetFName());
 				PropertyEditorModule->UnregisterCustomPropertyTypeLayout(FLiveLinkSubjectName::StaticStruct()->GetFName());
 				PropertyEditorModule->UnregisterCustomPropertyTypeLayout(FLiveLinkSubjectRepresentation::StaticStruct()->GetFName());
 				PropertyEditorModule->UnregisterCustomPropertyTypeLayout(FLiveLinkSubjectKey::StaticStruct()->GetFName());
diff --git a/Engine/Plugins/Animation/LiveLink/Source/LiveLinkEditor/Private/LiveLinkSubjectNameDetailCustomization.cpp b/Engine/Plugins/Animation/LiveLink/Source/LiveLinkEditor/Private/LiveLinkSubjectNameDetailCustomization.cpp
--- a/Engine/Plugins/Animation/LiveLink/Source/LiveLinkEditor/Private/LiveLinkSubjectNameDetailCustomization.cpp
+++ b/Engine/Plugins/Animation/LiveLink/Source/LiveLinkEditor/Private/LiveLinkSubjectNameDetailCustomization.cpp
@@ -62,12 +62,16 @@ void FLiveLinkSubjectNameDetailCustomization::SetValue(SLiveLinkSubjectRepresent
 
 	TArray<void*> RawData;
 	StructPropertyHandle->AccessRawData(RawData);
-	FLiveLinkSubjectName* PreviousValue = reinterpret_cast<FLiveLinkSubjectName*>(RawData[0]);
-	FLiveLinkSubjectName NewSubjectNameValue = NewValue.Subject;
 
-	FString TextValue;
-	StructProperty->Struct->ExportText(TextValue, &NewSubjectNameValue, PreviousValue, nullptr, EPropertyPortFlags::PPF_None, nullptr);
-	ensure(StructPropertyHandle->SetValueFromFormattedString(TextValue, EPropertyValueSetFlags::DefaultFlags) == FPropertyAccess::Result::Success);
+	if (ensure(RawData.Num()))
+	{
+		FLiveLinkSubjectName* PreviousValue = reinterpret_cast<FLiveLinkSubjectName*>(RawData[0]);
+		FLiveLinkSubjectName NewSubjectNameValue = NewValue.Subject;
+
+		FString TextValue;
+		StructProperty->Struct->ExportText(TextValue, &NewSubjectNameValue, PreviousValue, nullptr, EPropertyPortFlags::PPF_None, nullptr);
+		ensure(StructPropertyHandle->SetValueFromFormattedString(TextValue, EPropertyValueSetFlags::DefaultFlags) == FPropertyAccess::Result::Success);
+	}
 }
 
 bool FLiveLinkSubjectNameDetailCustomization::HasMultipleValues() const
diff --git a/Engine/Plugins/Animation/LiveLink/Source/LiveLinkEditor/Private/SLiveLinkBoneSelectionWidget.cpp b/Engine/Plugins/Animation/LiveLink/Source/LiveLinkEditor/Private/SLiveLinkBoneSelectionWidget.cpp
new file mode 100644
--- /dev/null
+++ b/Engine/Plugins/Animation/LiveLink/Source/LiveLinkEditor/Private/SLiveLinkBoneSelectionWidget.cpp
@@ -0,0 +1,332 @@
+// Copyright Epic Games, Inc. All Rights Reserved.
+
+#include "SLiveLinkBoneSelectionWidget.h"
+
+#include "DetailLayoutBuilder.h"
+#include "Engine/SkeletalMeshSocket.h"
+#include "IEditableSkeleton.h"
+#include "ILiveLinkClient.h"
+#include "LiveLinkSubjectSettings.h"
+#include "LiveLinkVirtualSubject.h"
+#include "Features/IModularFeatures.h"
+#include "Roles/LiveLinkAnimationRole.h"
+#include "Styling/AppStyle.h"
+#include "Translator/LiveLinkTransformRoleToAnimation.h"
+#include "Widgets/Layout/SSeparator.h"
+#include "Widgets/Input/SComboButton.h"
+#include "Widgets/Input/SSearchBox.h"
+#include "Widgets/Views/STreeView.h"
+
+
+#define LOCTEXT_NAMESPACE "SLiveLinkBoneSelectionWidget"
+
+/////////////////////////////////////////////////////
+void SLiveLinkBoneTreeMenu::Construct(const FArguments& InArgs, FLiveLinkSkeletonStaticData InSkeletonStaticData)
+{
+	OnSelectionChangedDelegate = InArgs._OnBoneSelectionChanged;
+
+	SkeletonStaticData = MoveTemp(InSkeletonStaticData);
+
+	FText TitleToUse = !InArgs._Title.IsEmpty() ? InArgs._Title  : LOCTEXT("BonePickerTitle", "Select...");
+
+	SAssignNew(TreeView, STreeView<TSharedPtr<FBoneNameInfo>>)
+		.TreeItemsSource(&SkeletonTreeInfo)
+		.OnGenerateRow(this, &SLiveLinkBoneTreeMenu::MakeTreeRowWidget)
+		.OnGetChildren(this, &SLiveLinkBoneTreeMenu::GetChildrenForInfo)
+		.OnSelectionChanged(this, &SLiveLinkBoneTreeMenu::OnSelectionChanged)
+		.SelectionMode(ESelectionMode::Single);
+
+	RebuildBoneList(InArgs._SelectedBone);
+
+	ChildSlot
+	[
+		SNew(SBorder)
+		.Padding(6.f)
+		.BorderImage(FAppStyle::GetBrush("NoBorder"))
+		.Content()
+		[
+			SNew(SBox)
+			.WidthOverride(300.f)
+			.HeightOverride(512.f)
+			.Content()
+			[
+				SNew(SVerticalBox)
+				+ SVerticalBox::Slot()
+				.AutoHeight()
+				[
+					SNew(STextBlock)
+					.Font(FAppStyle::GetFontStyle("BoldFont"))
+					.Text(TitleToUse)
+				]
+				+ SVerticalBox::Slot()
+				.AutoHeight()
+				[
+					SNew(SSeparator)
+					.SeparatorImage(FAppStyle::GetBrush("Menu.Separator"))
+					.Orientation(Orient_Horizontal)
+				]
+				+ SVerticalBox::Slot()
+				.AutoHeight()
+				[
+					SAssignNew(FilterTextWidget, SSearchBox)
+					.SelectAllTextWhenFocused(true)
+					.OnTextChanged(this, &SLiveLinkBoneTreeMenu::OnFilterTextChanged)
+					.HintText(NSLOCTEXT("BonePicker", "Search", "Search..."))
+				]
+				+ SVerticalBox::Slot()
+				[
+					TreeView->AsShared()
+				]
+			]
+		]
+	];
+}
+
+TSharedPtr<SWidget> SLiveLinkBoneTreeMenu::GetFilterTextWidget()
+{
+	return FilterTextWidget;
+}
+
+TSharedRef<ITableRow> SLiveLinkBoneTreeMenu::MakeTreeRowWidget(TSharedPtr<FBoneNameInfo> InInfo, const TSharedRef<STableViewBase>& OwnerTable)
+{
+	return SNew(STableRow<TSharedPtr<FBoneNameInfo>>, OwnerTable)
+		.Content()
+		[
+			SNew(STextBlock)
+			.HighlightText(FilterText)
+			.Text(FText::FromName(InInfo->BoneName))
+		];
+}
+
+void SLiveLinkBoneTreeMenu::GetChildrenForInfo(TSharedPtr<FBoneNameInfo> InInfo, TArray< TSharedPtr<FBoneNameInfo> >& OutChildren)
+{
+	OutChildren = InInfo->Children;
+}
+
+void SLiveLinkBoneTreeMenu::OnFilterTextChanged(const FText& InFilterText)
+{
+	FilterText = InFilterText;
+
+	RebuildBoneList(NAME_None);
+}
+
+void SLiveLinkBoneTreeMenu::OnSelectionChanged(TSharedPtr<SLiveLinkBoneTreeMenu::FBoneNameInfo> BoneInfo, ESelectInfo::Type SelectInfo)
+{
+	//Because we recreate all our items on tree refresh we will get a spurious null selection event initially.
+	if (BoneInfo.IsValid() && SelectInfo == ESelectInfo::OnMouseClick)
+	{
+		SelectBone(BoneInfo);
+	}
+}
+
+void SLiveLinkBoneTreeMenu::SelectBone(TSharedPtr<SLiveLinkBoneTreeMenu::FBoneNameInfo> BoneInfo)
+{
+	OnSelectionChangedDelegate.ExecuteIfBound(BoneInfo->BoneName);
+}
+
+FReply SLiveLinkBoneTreeMenu::OnKeyDown(const FGeometry& MyGeometry, const FKeyEvent& InKeyEvent)
+{
+	if(InKeyEvent.GetKey() == EKeys::Enter)
+	{
+		TArray<TSharedPtr<SLiveLinkBoneTreeMenu::FBoneNameInfo>> SelectedItems = TreeView->GetSelectedItems();
+		if(SelectedItems.Num() > 0)
+		{
+			SelectBone(SelectedItems[0]);
+			return FReply::Handled();
+		}
+	}
+
+	return FReply::Unhandled();
+}
+
+void SLiveLinkBoneTreeMenu::RebuildBoneList(const FName& SelectedBone)
+{
+	SkeletonTreeInfo.Empty();
+	SkeletonTreeInfoFlat.Empty();
+
+	const int32 MaxBone = SkeletonStaticData.BoneNames.Num(); 
+
+	for (int32 BoneIdx = 0; BoneIdx < MaxBone; ++BoneIdx)
+	{
+		const FName BoneName = SkeletonStaticData.BoneNames[BoneIdx];
+		TSharedRef<FBoneNameInfo> BoneInfo = MakeShared<FBoneNameInfo>(BoneName);
+
+		// Filter if Necessary
+		if (!FilterText.IsEmpty() && !BoneInfo->BoneName.ToString().Contains(FilterText.ToString()))
+		{
+			continue;
+		}
+
+		int32 ParentIdx = SkeletonStaticData.BoneParents[BoneIdx];
+		bool bAddToParent = false;
+
+		if (ParentIdx != INDEX_NONE && FilterText.IsEmpty())
+		{
+			// We have a parent, search for it in the flat list
+			FName ParentName = SkeletonStaticData.BoneNames[ParentIdx];
+
+			for (int32 FlatListIdx = 0; FlatListIdx < SkeletonTreeInfoFlat.Num(); ++FlatListIdx)
+			{
+				TSharedPtr<FBoneNameInfo> InfoEntry = SkeletonTreeInfoFlat[FlatListIdx];
+				if (InfoEntry->BoneName == ParentName)
+				{
+					bAddToParent = true;
+					ParentIdx = FlatListIdx;
+					break;
+				}
+			}
+
+			if (bAddToParent)
+			{
+				SkeletonTreeInfoFlat[ParentIdx]->Children.Add(BoneInfo);
+			}
+			else
+			{
+				SkeletonTreeInfo.Add(BoneInfo);
+			}
+		}
+		else
+		{
+			SkeletonTreeInfo.Add(BoneInfo);
+		}
+
+		SkeletonTreeInfoFlat.Add(BoneInfo);
+		TreeView->SetItemExpansion(BoneInfo, true);
+		if (BoneName == SelectedBone)
+		{
+			TreeView->SetItemSelection(BoneInfo, true);
+			TreeView->RequestScrollIntoView(BoneInfo);
+		}
+	}
+
+	TreeView->RequestTreeRefresh();
+}
+
+/////////////////////////////////////////////////////
+
+void SLiveLinkBoneSelectionWidget::Construct(const FArguments& InArgs, const FLiveLinkSubjectKey& InSubjectKey)
+{
+	OnBoneSelectionChanged = InArgs._OnBoneSelectionChanged;
+	OnGetSelectedBone = InArgs._OnGetSelectedBone;
+	SubjectKey = InSubjectKey;
+
+	ChildSlot
+	[
+		SAssignNew(BonePickerButton, SComboButton)
+		.OnGetMenuContent(FOnGetContent::CreateSP(this, &SLiveLinkBoneSelectionWidget::CreateSkeletonWidgetMenu))
+		.ContentPadding(FMargin(4.0f, 2.0f, 4.0f, 2.0f))
+		.ButtonContent()
+		[
+			SNew(STextBlock)
+			.Text(this, &SLiveLinkBoneSelectionWidget::GetCurrentBoneName)
+			.Font(IDetailLayoutBuilder::GetDetailFont())
+		]
+	];
+}
+
+TSharedRef<SWidget> SLiveLinkBoneSelectionWidget::CreateSkeletonWidgetMenu()
+{
+	FName CurrentBoneName;
+	if (OnGetSelectedBone.IsBound())
+	{
+		CurrentBoneName = OnGetSelectedBone.Execute();
+	}
+
+	ILiveLinkClient* LiveLinkClient = &IModularFeatures::Get().GetModularFeature<ILiveLinkClient>(ILiveLinkClient::ModularFeatureName);
+	bool bIncludeVirtualSubject = true;
+	bool bIncludeDisabledSubject = true;
+
+	FName TranslatorOutputBoneName;
+	TOptional<FLiveLinkSkeletonStaticData> SkeletonData;
+
+	/** Get a list of name of subjects supporting a certain role */
+	TArray<FLiveLinkSubjectKey> Subjects = LiveLinkClient->GetSubjectsSupportingRole(ULiveLinkAnimationRole::StaticClass(), bIncludeDisabledSubject, bIncludeVirtualSubject);
+
+	if (FLiveLinkSubjectKey* Key = Subjects.FindByPredicate([this](const FLiveLinkSubjectKey& Other) { return Other.SubjectName == SubjectKey.SubjectName; }))
+	{
+		if (const FLiveLinkStaticDataStruct* StaticData = LiveLinkClient->GetSubjectStaticData_AnyThread(*Key))
+		{
+			if (const FLiveLinkSkeletonStaticData* SkeletonStaticData = StaticData->Cast<FLiveLinkSkeletonStaticData>())
+			{
+				SkeletonData = *SkeletonStaticData;
+			}
+		}
+
+		if (!SkeletonData)
+		{
+			SkeletonData = MakeStaticDataFromTranslator(*Key);
+		}
+	}
+
+
+	TSharedRef<SLiveLinkBoneTreeMenu> MenuWidget = SNew(SLiveLinkBoneTreeMenu, MoveTemp(*SkeletonData))
+		.OnBoneSelectionChanged(this, &SLiveLinkBoneSelectionWidget::OnSelectionChanged)
+		.SelectedBone(CurrentBoneName);
+
+	BonePickerButton->SetMenuContentWidgetToFocus(MenuWidget->GetFilterTextWidget());
+
+	return MenuWidget;
+}
+
+void SLiveLinkBoneSelectionWidget::OnSelectionChanged(FName BoneName)
+{
+	//Because we recreate all our items on tree refresh we will get a spurious null selection event initially.
+	if (OnBoneSelectionChanged.IsBound())
+	{
+		OnBoneSelectionChanged.Execute(BoneName);
+	}
+
+	BonePickerButton->SetIsOpen(false);
+}
+
+FText SLiveLinkBoneSelectionWidget::GetCurrentBoneName() const
+{
+	if (OnGetSelectedBone.IsBound())
+	{
+		FName Name = OnGetSelectedBone.Execute();
+		return FText::FromName(Name);
+	}
+
+	return FText::GetEmpty();
+}
+
+void SLiveLinkBoneSelectionWidget::SetSubject(const FLiveLinkSubjectKey& InSubjectKey)
+{
+	SubjectKey = InSubjectKey;
+}
+
+
+FLiveLinkSkeletonStaticData SLiveLinkBoneSelectionWidget::MakeStaticDataFromTranslator(const FLiveLinkSubjectKey& InSubjectKey) const
+{
+	ILiveLinkClient* LiveLinkClient = &IModularFeatures::Get().GetModularFeature<ILiveLinkClient>(ILiveLinkClient::ModularFeatureName);
+	UObject* Settings = LiveLinkClient->GetSubjectSettings(InSubjectKey);
+	TArray<ULiveLinkFrameTranslator*> Translators;
+
+	if (ULiveLinkSubjectSettings* SubjectSettings = Cast<ULiveLinkSubjectSettings>(Settings))
+	{
+		Translators = SubjectSettings->Translators;
+	}
+	else if (ULiveLinkVirtualSubject* VirtualSubject = Cast<ULiveLinkVirtualSubject>(Settings))
+	{
+		Translators = VirtualSubject->GetTranslators();
+	}
+
+	FName OutputBoneName;
+	for (ULiveLinkFrameTranslator* Translator : Translators)
+	{
+		if (ULiveLinkTransformRoleToAnimation* TransformTranslator = Cast<ULiveLinkTransformRoleToAnimation>(Translator))
+		{
+			OutputBoneName = TransformTranslator->OutputBoneName;
+			break;
+		}
+	}
+
+	FLiveLinkSkeletonStaticData StaticData;
+	StaticData.BoneParents = { INDEX_NONE };
+	StaticData.BoneNames = { OutputBoneName };
+
+	return StaticData;
+}
+
+
+#undef LOCTEXT_NAMESPACE
diff --git a/Engine/Plugins/Animation/LiveLink/Source/LiveLinkEditor/Private/SLiveLinkSubjectRepresentationPicker.cpp b/Engine/Plugins/Animation/LiveLink/Source/LiveLinkEditor/Private/SLiveLinkSubjectRepresentationPicker.cpp
--- a/Engine/Plugins/Animation/LiveLink/Source/LiveLinkEditor/Private/SLiveLinkSubjectRepresentationPicker.cpp
+++ b/Engine/Plugins/Animation/LiveLink/Source/LiveLinkEditor/Private/SLiveLinkSubjectRepresentationPicker.cpp
@@ -104,6 +104,8 @@ void SLiveLinkSubjectRepresentationPicker::Construct(const FArguments& InArgs)
 	HasMultipleValuesAttribute = InArgs._HasMultipleValues;
 	bShowSource = InArgs._ShowSource;
 	bShowRole = InArgs._ShowRole;
+	GetSubjectsDelegate = InArgs._OnGetSubjects;
+
 
 	SubjectRepData.Reset();
 	SelectedLiveLinkPreset.Reset();
@@ -508,7 +510,16 @@ void SLiveLinkSubjectRepresentationPicker::BuildSubjectRepDataList()
 	else if (IModularFeatures::Get().IsModularFeatureAvailable(ILiveLinkClient::ModularFeatureName))
 	{
 		ILiveLinkClient& LiveLinkClient = IModularFeatures::Get().GetModularFeature<ILiveLinkClient>(ILiveLinkClient::ModularFeatureName);
-		TArray<FLiveLinkSubjectKey> SubjectKeys = LiveLinkClient.GetSubjects(true, true);
+		TArray<FLiveLinkSubjectKey> SubjectKeys;
+
+		if (GetSubjectsDelegate.IsBound())
+		{
+			GetSubjectsDelegate.Execute(SubjectKeys);
+		}
+		else
+		{
+			SubjectKeys = LiveLinkClient.GetSubjects(true, true);
+		}
 
 		if (bShowSource)
 		{
diff --git a/Engine/Plugins/Animation/LiveLink/Source/LiveLinkEditor/Public/SLiveLinkBoneSelectionWidget.h b/Engine/Plugins/Animation/LiveLink/Source/LiveLinkEditor/Public/SLiveLinkBoneSelectionWidget.h
new file mode 100644
--- /dev/null
+++ b/Engine/Plugins/Animation/LiveLink/Source/LiveLinkEditor/Public/SLiveLinkBoneSelectionWidget.h
@@ -0,0 +1,144 @@
+/**  Copyright Epic Games, Inc. All Rights Reserved. */
+
+#pragma once
+
+#include "CoreMinimal.h"
+#include "Roles/LiveLinkAnimationTypes.h"
+#include "SlateFwd.h"
+#include "Widgets/DeclarativeSyntaxSupport.h"
+#include "Widgets/SWidget.h"
+#include "Widgets/SCompoundWidget.h"
+#include "Widgets/Views/STableViewBase.h"
+#include "Widgets/Views/STableRow.h"
+#include "Widgets/Views/STreeView.h"
+
+
+class IEditableSkeleton;
+class SComboButton;
+
+DECLARE_DELEGATE_OneParam(FOnBoneSelectionChanged, FName);
+DECLARE_DELEGATE_RetVal(FName, FGetSelectedBone);
+
+class LIVELINKEDITOR_API SLiveLinkBoneTreeMenu : public SCompoundWidget
+{
+public:
+	/** Storage object for bone hierarchy */
+	struct FBoneNameInfo
+	{
+		FBoneNameInfo(FName Name) : BoneName(Name) {}
+
+		FName BoneName;
+		TArray<TSharedPtr<FBoneNameInfo>> Children;
+	};
+
+	SLATE_BEGIN_ARGS(SLiveLinkBoneTreeMenu)
+		: _OnBoneSelectionChanged()
+		{}
+
+		SLATE_ARGUMENT(FText, Title)
+		SLATE_ARGUMENT(FName, SelectedBone)
+		SLATE_EVENT(FOnBoneSelectionChanged, OnBoneSelectionChanged)
+	SLATE_END_ARGS();
+
+	/**
+	* Construct this widget
+	*
+	* @param	InArgs	The declaration data for this widget
+	*/
+	void Construct(const FArguments& InArgs, FLiveLinkSkeletonStaticData SkeletonStaticData);
+
+	/** Get the filter text widget, e.g. for focus */
+	TSharedPtr<SWidget> GetFilterTextWidget();
+
+private:
+	/** SWidget interface */
+	virtual FReply OnKeyDown(const FGeometry& MyGeometry, const FKeyEvent& InKeyEvent) override;
+
+	/** Using the current filter, repopulate the tree view */
+	void RebuildBoneList(const FName& SelectedBone);
+
+	/** Make a single tree row widget */
+	TSharedRef<ITableRow> MakeTreeRowWidget(TSharedPtr<FBoneNameInfo> InInfo, const TSharedRef<STableViewBase>& OwnerTable);
+
+	/** Get the children for the provided bone info */
+	void GetChildrenForInfo(TSharedPtr<FBoneNameInfo> InInfo, TArray< TSharedPtr<FBoneNameInfo> >& OutChildren);
+
+	/** Called when the user changes the search filter */
+	void OnFilterTextChanged(const FText& InFilterText);
+
+	/** Handle the tree view selection changing */
+	void OnSelectionChanged(TSharedPtr<SLiveLinkBoneTreeMenu::FBoneNameInfo> BoneInfo, ESelectInfo::Type SelectInfo);
+
+	/** Select a specific bone, helper for UI handler functions */
+	void SelectBone(TSharedPtr<SLiveLinkBoneTreeMenu::FBoneNameInfo> BoneInfo);
+	
+	/** Tree info entries for bone picker */
+	TArray<TSharedPtr<FBoneNameInfo>> SkeletonTreeInfo;
+	/** Mirror of SkeletonTreeInfo but flattened for searching */
+	TArray<TSharedPtr<FBoneNameInfo>> SkeletonTreeInfoFlat;
+
+	/** Text to filter bone tree with */
+	FText FilterText;
+
+	/** Tree view used in the button menu */
+	TSharedPtr<STreeView<TSharedPtr<FBoneNameInfo>>> TreeView;
+
+	/** Filter text widget */
+	TSharedPtr<SSearchBox> FilterTextWidget;
+
+	/** Delegate called when a bone is selected. */
+	FOnBoneSelectionChanged OnSelectionChangedDelegate;
+
+	/** Static data used to populate the bone list. */
+	FLiveLinkSkeletonStaticData SkeletonStaticData;
+};
+
+class LIVELINKEDITOR_API SLiveLinkBoneSelectionWidget : public SCompoundWidget
+{
+public: 
+
+	SLATE_BEGIN_ARGS(SLiveLinkBoneSelectionWidget)
+		: _OnBoneSelectionChanged()
+		, _OnGetSelectedBone()
+	{}
+
+		/** set selected bone name */
+		SLATE_EVENT(FOnBoneSelectionChanged, OnBoneSelectionChanged);
+
+		/** get selected bone name **/
+		SLATE_EVENT(FGetSelectedBone, OnGetSelectedBone);
+	SLATE_END_ARGS();
+
+	/**
+	 * Construct this widget
+	 *
+	 * @param	InArgs	The declaration data for this widget
+	 */
+	void Construct(const FArguments& InArgs, const FLiveLinkSubjectKey& InSubjectKey);
+
+	/** Set the subject for which we will list the bones. */
+	void SetSubject(const FLiveLinkSubjectKey& InSubjectKey);
+
+	/** Utility method to create fake static data from a frame translator. */
+	FLiveLinkSkeletonStaticData MakeStaticDataFromTranslator(const FLiveLinkSubjectKey& SubjectKey) const;
+
+private: 
+
+	/** Creates the combo button menu when clicked */
+	TSharedRef<SWidget> CreateSkeletonWidgetMenu();
+	/** Called when the user selects a bone name */
+	void OnSelectionChanged(FName BoneName);
+
+	/** Gets the current bone name, used to get the right name for the combo button */
+	FText GetCurrentBoneName() const;
+
+	/** Base combo button */
+	TSharedPtr<SComboButton> BonePickerButton;
+
+	/** Subject that for which we are selecting a bone. */
+	FLiveLinkSubjectKey SubjectKey;
+
+	// Delegates
+	FOnBoneSelectionChanged OnBoneSelectionChanged;
+	FGetSelectedBone OnGetSelectedBone;
+};
diff --git a/Engine/Plugins/Animation/LiveLink/Source/LiveLinkEditor/Public/SLiveLinkSubjectRepresentationPicker.h b/Engine/Plugins/Animation/LiveLink/Source/LiveLinkEditor/Public/SLiveLinkSubjectRepresentationPicker.h
--- a/Engine/Plugins/Animation/LiveLink/Source/LiveLinkEditor/Public/SLiveLinkSubjectRepresentationPicker.h
+++ b/Engine/Plugins/Animation/LiveLink/Source/LiveLinkEditor/Public/SLiveLinkSubjectRepresentationPicker.h
@@ -25,6 +25,8 @@ class STableViewBase;
 
 typedef TSharedPtr<FLiveLinkSubjectRepresentationPickerEntry> FLiveLinkSubjectRepresentationPickerEntryPtr;
 
+DECLARE_DELEGATE_OneParam(FOnGetSubjects, TArray<FLiveLinkSubjectKey>& /**OutSubjectsList*/);
+
 
 /**
  * A widget which allows the user to enter a subject name or discover it from a drop menu.
@@ -96,6 +98,9 @@ public:
 		/** Sets the font used to draw the text on the button */
 		SLATE_ATTRIBUTE(FSlateFontInfo, Font)
 
+		/** (Optional) Delegate that can be specified in order to provide a custom list of subjects. */
+		SLATE_EVENT(FOnGetSubjects, OnGetSubjects);
+
 	SLATE_END_ARGS()
 
 	/**
@@ -143,6 +148,9 @@ private:
 	TAttribute<bool> HasMultipleValuesAttribute;
 	bool bShowSource;
 	bool bShowRole;
+
+	/** Delegate used to retrieve a custom list of subjects. */
+	FOnGetSubjects GetSubjectsDelegate;
 };
 
 #if UE_ENABLE_INCLUDE_ORDER_DEPRECATED_IN_5_2
