LiveLink - Fix wrong copyright notice and crash when accessing optional that wasn't set in BoneSelectionWidget

diff --git a/Engine/Plugins/Animation/LiveLink/Source/LiveLinkEditor/Private/SLiveLinkBoneSelectionWidget.cpp b/Engine/Plugins/Animation/LiveLink/Source/LiveLinkEditor/Private/SLiveLinkBoneSelectionWidget.cpp
--- a/Engine/Plugins/Animation/LiveLink/Source/LiveLinkEditor/Private/SLiveLinkBoneSelectionWidget.cpp
+++ b/Engine/Plugins/Animation/LiveLink/Source/LiveLinkEditor/Private/SLiveLinkBoneSelectionWidget.cpp
@@ -21,11 +21,14 @@
 #define LOCTEXT_NAMESPACE "SLiveLinkBoneSelectionWidget"
 
 /////////////////////////////////////////////////////
-void SLiveLinkBoneTreeMenu::Construct(const FArguments& InArgs, FLiveLinkSkeletonStaticData InSkeletonStaticData)
+void SLiveLinkBoneTreeMenu::Construct(const FArguments& InArgs, TOptional<FLiveLinkSkeletonStaticData> InSkeletonStaticData)
 {
 	OnSelectionChangedDelegate = InArgs._OnBoneSelectionChanged;
-
-	SkeletonStaticData = MoveTemp(InSkeletonStaticData);
+	
+	if (InSkeletonStaticData)
+	{
+		SkeletonStaticData = MoveTemp(*InSkeletonStaticData);
+	}
 
 	FText TitleToUse = !InArgs._Title.IsEmpty() ? InArgs._Title  : LOCTEXT("BonePickerTitle", "Select...");
 
@@ -258,8 +261,7 @@ TSharedRef<SWidget> SLiveLinkBoneSelectionWidget::CreateSkeletonWidgetMenu()
 		}
 	}
 
-
-	TSharedRef<SLiveLinkBoneTreeMenu> MenuWidget = SNew(SLiveLinkBoneTreeMenu, MoveTemp(*SkeletonData))
+	TSharedRef<SLiveLinkBoneTreeMenu> MenuWidget = SNew(SLiveLinkBoneTreeMenu, MoveTemp(SkeletonData))
 		.OnBoneSelectionChanged(this, &SLiveLinkBoneSelectionWidget::OnSelectionChanged)
 		.SelectedBone(CurrentBoneName);
 
diff --git a/Engine/Plugins/Animation/LiveLink/Source/LiveLinkEditor/Public/SLiveLinkBoneSelectionWidget.h b/Engine/Plugins/Animation/LiveLink/Source/LiveLinkEditor/Public/SLiveLinkBoneSelectionWidget.h
--- a/Engine/Plugins/Animation/LiveLink/Source/LiveLinkEditor/Public/SLiveLinkBoneSelectionWidget.h
+++ b/Engine/Plugins/Animation/LiveLink/Source/LiveLinkEditor/Public/SLiveLinkBoneSelectionWidget.h
@@ -1,4 +1,4 @@
-/**  Copyright Epic Games, Inc. All Rights Reserved. */
+// Copyright Epic Games, Inc. All Rights Reserved.
 
 #pragma once
 
@@ -45,7 +45,7 @@ public:
 	*
 	* @param	InArgs	The declaration data for this widget
 	*/
-	void Construct(const FArguments& InArgs, FLiveLinkSkeletonStaticData SkeletonStaticData);
+	void Construct(const FArguments& InArgs, TOptional<FLiveLinkSkeletonStaticData> SkeletonStaticData);
 
 	/** Get the filter text widget, e.g. for focus */
 	TSharedPtr<SWidget> GetFilterTextWidget();
