LiveLink Virtual Subjects - When combining subjects, allow specifying if the location/rotation should be fetched from the parent subject, child subject or if it should be a combination of them.

diff --git a/Engine/Plugins/Animation/LiveLink/Source/LiveLink/Private/LiveLinkAnimationVirtualSubject.cpp b/Engine/Plugins/Animation/LiveLink/Source/LiveLink/Private/LiveLinkAnimationVirtualSubject.cpp
--- a/Engine/Plugins/Animation/LiveLink/Source/LiveLink/Private/LiveLinkAnimationVirtualSubject.cpp
+++ b/Engine/Plugins/Animation/LiveLink/Source/LiveLink/Private/LiveLinkAnimationVirtualSubject.cpp
@@ -331,8 +331,6 @@ void ULiveLinkAnimationVirtualSubject::ProcessAttachmentsForStaticData(TArray<in
 			FChildBoneInfo ChildBoneInfo;
 			ChildBoneInfo.Offset = Offset;
 			ChildBoneInfo.ParentBone = GlobalParentIndex;
-			ChildBoneInfo.bIgnoreParentLocation = Attachment.bIgnoreParentLocation;
-			ChildBoneInfo.bIgnoreParentRotation = Attachment.bIgnoreParentRotation;
 
 			ChildBonesInfo.Add(GlobalChildIndex, MoveTemp(ChildBoneInfo));
 		}
@@ -349,23 +347,43 @@ void ULiveLinkAnimationVirtualSubject::ProcessAttachmentsForFrameData(FLiveLinkA
 
 		const FChildBoneInfo& ChildBoneInfo = ChildInfo.Value;
 
-		FTransform ParentTransform = SnapshotFrameData->Transforms[ChildBoneInfo.ParentBone];
-		FTransform ModifiedParentTransform = FTransform::Identity;
-		if (!ChildBoneInfo.bIgnoreParentLocation)
+		FTransform& ParentBoneTransform = SnapshotFrameData->Transforms[ChildInfo.Value.ParentBone];
+		FTransform& ChildBoneTransform = SnapshotFrameData->Transforms[ChildInfo.Key];
+
+		FTransform FinalBoneTransform;
+		switch(LocationBehavior)
 		{
-			ModifiedParentTransform.SetLocation(ParentTransform.GetLocation());
+			using enum EBoneTransformResolution;
+			case KeepParent:
+				FinalBoneTransform.SetLocation(ParentBoneTransform.GetLocation());
+				break;
+			case KeepChild:
+			{
+				FinalBoneTransform.SetLocation(ChildBoneTransform.GetLocation());
+				break;
+			}
+			case Combine:
+			{
+				FinalBoneTransform.SetLocation(ChildBoneTransform.TransformPosition(ParentBoneTransform.GetLocation()));
+				break;
+			}
 		}
-		if (!ChildBoneInfo.bIgnoreParentRotation)
+
+		switch(RotationBehavior)
 		{
-			ModifiedParentTransform.SetRotation(ParentTransform.GetRotation());
+			using enum EBoneTransformResolution;
+			case KeepParent:
+				FinalBoneTransform.SetRotation(ParentBoneTransform.GetRotation());
+				break;
+			case KeepChild:
+				FinalBoneTransform.SetRotation(ChildBoneTransform.GetRotation());
+				break;
+			case Combine:
+				FinalBoneTransform.SetRotation(ChildBoneTransform.TransformRotation(ParentBoneTransform.GetRotation()));
+				break;
 		}
 
-		ModifiedParentTransform.SetScale3D(ParentTransform.GetScale3D());
-
-		// Also apply an offset from the attach target to the child bone.
-		Offset = Offset * ModifiedParentTransform;
-
-		SnapshotFrameData->Transforms[ChildInfo.Key] = SnapshotFrameData->Transforms[ChildInfo.Key] * Offset;
+		ChildBoneTransform = FinalBoneTransform * Offset;
 	}
 }
 
diff --git a/Engine/Plugins/Animation/LiveLink/Source/LiveLink/Public/LiveLinkAnimationVirtualSubject.h b/Engine/Plugins/Animation/LiveLink/Source/LiveLink/Public/LiveLinkAnimationVirtualSubject.h
--- a/Engine/Plugins/Animation/LiveLink/Source/LiveLink/Public/LiveLinkAnimationVirtualSubject.h
+++ b/Engine/Plugins/Animation/LiveLink/Source/LiveLink/Public/LiveLinkAnimationVirtualSubject.h
@@ -11,6 +11,17 @@
 struct FLiveLinkSubjectFrameData;
 struct FLiveLinkSubjectKey;
 
+/** What action should be taken when a bone name conflict happens between a parent and a child subject. */
+UENUM()
+enum class EBoneTransformResolution
+{
+	/** Keep parent bone transform. */
+	KeepParent,
+	/** Keep child bone transform. */
+	KeepChild,
+	/** Combine the child and parent's bone transforms. */
+	Combine
+};
 
 /** A Skeleton virtual subject is an assembly of different subjects supporting the animation role */
 UCLASS(meta=(DisplayName="Animation Virtual Subject"))
@@ -61,6 +72,14 @@ public:
 	UPROPERTY(EditAnywhere, Category = "LiveLink")
 	TArray<FLiveLinkVirtualSubjectBoneAttachment> Attachments;
 
+	/** What should happen to the location of a bone when there's a conflict between the child and parent subject. */
+	UPROPERTY(EditAnywhere, Category = "LiveLink")
+	EBoneTransformResolution LocationBehavior = EBoneTransformResolution::Combine;
+
+	/** What should happen to the rotation of a bone when there's a conflict between the child and parent subject. */
+	UPROPERTY(EditAnywhere, Category = "LiveLink")
+	EBoneTransformResolution RotationBehavior = EBoneTransformResolution::Combine;
+
 	/** Whether to append SubjectName to each bones part of the virtual hierarchy */
 	UPROPERTY(EditAnywhere, Category = "Settings")
 	bool bAppendSubjectNameToBones;
@@ -73,10 +92,6 @@ private:
 		int32 ParentBone = INDEX_NONE;
 		/** Offset of the child bone to its parent bone. */
 		FTransform Offset = FTransform::Identity;
-		/** Whether to ignore parent location when offsetting the child bone relative to its parent. */
-		bool bIgnoreParentLocation = false;
-		/** Whether to ignore parent rotation when offsetting the child bone relative to its parent. */
-		bool bIgnoreParentRotation = false;
 	};
 
 	/** Map of global child indices to their respective bone info structure.  */
