LiveLinkHub - Experimental - Add option to tick LiveLinkHub's client outside of the game thread.

This allows the hub to process virtual subjects without being affected by hitches on the game or render thread. This change should also eventually allow us to re-use more  of the base LiveLinkClient rather than LiveLinkHubClient.

diff --git a/Engine/Plugins/Animation/LiveLink/Source/LiveLink/Private/LiveLinkClient.cpp b/Engine/Plugins/Animation/LiveLink/Source/LiveLink/Private/LiveLinkClient.cpp
--- a/Engine/Plugins/Animation/LiveLink/Source/LiveLink/Private/LiveLinkClient.cpp
+++ b/Engine/Plugins/Animation/LiveLink/Source/LiveLink/Private/LiveLinkClient.cpp
@@ -2,6 +2,8 @@
 
 #include "LiveLinkClient.h"
 
+#include "Algo/ForEach.h"
+#include "Algo/Transform.h"
 #include "Engine/Blueprint.h"
 #include "Engine/Engine.h"
 #include "LiveLinkAnimationVirtualSubject.h"
@@ -51,6 +53,7 @@ static TAutoConsoleVariable<int32> CVarMaxNewFrameDataPerUpdate(
 
 
 FLiveLinkClient::FLiveLinkClient()
+	: Collection(MakeUnique<FLiveLinkSourceCollection>())
 {
 	LLM_SCOPE_BYTAG(LiveLink_LiveLinkClient);
 
@@ -58,7 +61,6 @@ FLiveLinkClient::FLiveLinkClient()
 	CachedEngineTime = 0.0;
 #endif
 
-	Collection = MakeUnique<FLiveLinkSourceCollection>();
 	OnLiveLinkSubjectRemoved().AddRaw(this, &FLiveLinkClient::OnSubjectRemovedCallback);
 	FCoreDelegates::OnPreExit.AddRaw(this, &FLiveLinkClient::Shutdown);
 
@@ -84,7 +86,6 @@ void FLiveLinkClient::Tick()
 {
 	SCOPE_CYCLE_COUNTER(STAT_LiveLink_Client_Tick);
 
-	FScopeLock Lock(&CollectionAccessCriticalSection);
 	DoPendingWork();
 	CacheValues();
 	UpdateSources();
@@ -95,6 +96,8 @@ void FLiveLinkClient::Tick()
 
 void FLiveLinkClient::DoPendingWork()
 {
+	FScopeLock PendingFramesLock(&PendingFramesCriticalSection);
+
 	check(Collection);
 
 	// Remove Sources and Subjects
@@ -119,10 +122,10 @@ void FLiveLinkClient::DoPendingWork()
 
 void FLiveLinkClient::UpdateSources()
 {
-	for (FLiveLinkCollectionSourceItem& SourceItem : Collection->GetSources())
+	Collection->ForEachSource([](const FLiveLinkCollectionSourceItem& SourceItem)
 	{
 		SourceItem.Source->Update();
-	}
+	});
 }
 
 void FLiveLinkClient::BuildThisTicksSubjectSnapshot()
@@ -134,14 +137,42 @@ void FLiveLinkClient::BuildThisTicksSubjectSnapshot()
 	TSet<FGuid> TaggedSources;
 
 	// Update the Live Subject before the Virtual Subject
-	for (const FLiveLinkCollectionSubjectItem& SubjectItem : Collection->GetSubjects())
+	Collection->ForEachSubject([this, &TaggedSources](FLiveLinkCollectionSourceItem& SourceItem, FLiveLinkCollectionSubjectItem& SubjectItem)
 	{
 		if (FLiveLinkSubject* LiveSubject = SubjectItem.GetLiveSubject())
 		{
 			if (SubjectItem.bEnabled)
 			{
-				ULiveLinkSourceSettings* SourceSettings = GetSourceSettings(SubjectItem.Key.Source);
+				ULiveLinkSourceSettings* SourceSettings = SourceItem.Setting.Get();
 				ULiveLinkSubjectSettings* SubjectSettings = SubjectItem.GetLinkSettings();
+
+				LiveSubject->CacheSettings(SourceSettings, SubjectSettings);
+				LiveSubject->Update();
+				EnabledSubjects.Add(SubjectItem.Key.SubjectName, SubjectItem.Key);
+
+				// Update Source FrameRate from first enabled subject with valid data.
+				if (LiveSubject->HasValidFrameSnapshot() && SubjectSettings->FrameRate.IsValid() && !TaggedSources.Contains(SubjectItem.Key.Source))
+				{
+					SourceSettings->BufferSettings.DetectedFrameRate = SubjectSettings->FrameRate;
+					TaggedSources.Add(SubjectItem.Key.Source);
+				}
+			}
+			else
+			{
+				LiveSubject->ClearFrames();
+			}
+		}
+	});
+
+	Collection->ForEachSubject([this, &TaggedSources](FLiveLinkCollectionSourceItem& SourceItem, FLiveLinkCollectionSubjectItem& SubjectItem)
+	{
+		if (FLiveLinkSubject* LiveSubject = SubjectItem.GetLiveSubject())
+		{
+			if (SubjectItem.bEnabled)
+			{
+				ULiveLinkSourceSettings* SourceSettings = SourceItem.Setting.Get();
+				ULiveLinkSubjectSettings* SubjectSettings = SubjectItem.GetLinkSettings();
+
 				LiveSubject->CacheSettings(SourceSettings, SubjectSettings);
 				LiveSubject->Update();
 				EnabledSubjects.Add(SubjectItem.Key.SubjectName, SubjectItem.Key);
@@ -158,11 +189,10 @@ void FLiveLinkClient::BuildThisTicksSubjectSnapshot()
 				LiveSubject->ClearFrames();
 			}
 		}
-	}
+	});
 
-	for (const FLiveLinkCollectionSubjectItem& SubjectItem : Collection->GetSubjects())
+	Collection->ForEachSubject([this](FLiveLinkCollectionSourceItem& SourceItem, FLiveLinkCollectionSubjectItem& SubjectItem)
 	{
-		
 		if (ULiveLinkVirtualSubject* VSubject = SubjectItem.GetVirtualSubject())
 		{
 			if (SubjectItem.bEnabled)
@@ -176,7 +206,7 @@ void FLiveLinkClient::BuildThisTicksSubjectSnapshot()
 				VSubject->ClearFrames();
 			}
 		}
-	}
+	});
 }
 
 void FLiveLinkClient::CacheValues()
@@ -215,7 +245,6 @@ void FLiveLinkClient::Shutdown()
 		bool bContinue = true;
 		while(bContinue)
 		{
-			FScopeLock Lock(&CollectionAccessCriticalSection);
 			bContinue = !Collection->RequestShutdown();
 
 			if (FPlatformTime::Seconds() - StartShutdownSeconds > Timeout)
@@ -345,7 +374,8 @@ FGuid FLiveLinkClient::AddSource(TSharedPtr<ILiveLinkSource> InSource)
 		{
 			UClass* SourceSettingsClass = InSource->GetSettingsClass().Get();
 			UClass* SettingsClass = SourceSettingsClass ? SourceSettingsClass : ULiveLinkSourceSettings::StaticClass();
-			Settings = Data.Setting = NewObject<ULiveLinkSourceSettings>(GetTransientPackage(), SettingsClass);
+			Data.Setting = TStrongObjectPtr(NewObject<ULiveLinkSourceSettings>(GetTransientPackage(), SettingsClass));
+			Settings = Data.Setting.Get();
 		}
 		Collection->AddSource(MoveTemp(Data));
 
@@ -372,7 +402,7 @@ FGuid FLiveLinkClient::AddVirtualSubjectSource(FName SourceName)
 
 		ULiveLinkVirtualSubjectSourceSettings* NewSettings = NewObject<ULiveLinkVirtualSubjectSourceSettings>(GetTransientPackage(), ULiveLinkVirtualSubjectSourceSettings::StaticClass());
 		NewSettings->SourceName = SourceName;
-		Data.Setting = NewSettings;
+		Data.Setting = TStrongObjectPtr(NewSettings);
 		Data.bIsVirtualSource = true;
 		Collection->AddSource(MoveTemp(Data));
 
@@ -453,11 +483,12 @@ bool FLiveLinkClient::CreateSource(const FLiveLinkSourcePreset& InSourcePreset)
 		FLiveLinkLog::Info(TEXT("Creating Source '%s' from Preset: Settings class '%s' is not what is expected ('%s'). Updating to new class."), *InSourcePreset.SourceType.ToString(), *InSourcePreset.Settings->GetClass()->GetName(), *SourceSettingsClass->GetName());
 		Setting = NewObject<ULiveLinkSourceSettings>(GetTransientPackage(), SourceSettingsClass);
 		UEngine::CopyPropertiesForUnrelatedObjects(InSourcePreset.Settings, Setting);
-		Data.Setting = Setting;
+		Data.Setting = TStrongObjectPtr(Setting);
 	}
 	else
 	{
-		Setting = Data.Setting = DuplicateObject<ULiveLinkSourceSettings>(InSourcePreset.Settings, GetTransientPackage());
+		Data.Setting = TStrongObjectPtr(DuplicateObject<ULiveLinkSourceSettings>(InSourcePreset.Settings, GetTransientPackage()));
+		Setting = Data.Setting.Get();
 	}
 
 	Collection->AddSource(MoveTemp(Data));
@@ -488,16 +519,16 @@ void FLiveLinkClient::RemoveSource(FGuid InEntryGuid)
 void FLiveLinkClient::RemoveAllSources()
 {
 	check(Collection);
-	for (FLiveLinkCollectionSourceItem& SourceItem : Collection->GetSources())
+	Collection->ForEachSource([](FLiveLinkCollectionSourceItem& SourceItem)
 	{
 		SourceItem.bPendingKill = true;
-	}
+	});
 }
 
 bool FLiveLinkClient::HasSourceBeenAdded(TSharedPtr<ILiveLinkSource> InSource) const
 {
 	check(Collection);
-	if (FLiveLinkCollectionSourceItem* SourceItem = Collection->FindSource(InSource))
+	if (const FLiveLinkCollectionSourceItem* SourceItem = Collection->FindSource(InSource))
 	{
 		return !SourceItem->bPendingKill;
 	}
@@ -509,13 +540,14 @@ TArray<FGuid> FLiveLinkClient::GetSources(bool bEvenIfPendingKill) const
 	check(Collection);
 
 	TArray<FGuid> Result;
-	for (const FLiveLinkCollectionSourceItem& SourceItem : Collection->GetSources())
+	Collection->ForEachSource([&Result, bEvenIfPendingKill](const FLiveLinkCollectionSourceItem& SourceItem)
 	{
 		if ((!SourceItem.bPendingKill || bEvenIfPendingKill) && !SourceItem.IsVirtualSource())
 		{
 			Result.Add(SourceItem.Guid);
 		}
-	}
+	});
+
 	return Result;
 }
 
@@ -524,13 +556,15 @@ TArray<FGuid> FLiveLinkClient::GetVirtualSources(bool bEvenIfPendingKill) const
 	check(Collection);
 
 	TArray<FGuid> Result;
-	for (const FLiveLinkCollectionSourceItem& SourceItem : Collection->GetSources())
+
+	Collection->ForEachSource([&Result, bEvenIfPendingKill](const FLiveLinkCollectionSourceItem& SourceItem)
 	{
 		if ((!SourceItem.bPendingKill || bEvenIfPendingKill) && SourceItem.IsVirtualSource())
 		{
 			Result.Add(SourceItem.Guid);
 		}
-	}
+	});
+
 	return Result;
 }
 
@@ -541,13 +575,13 @@ FLiveLinkSourcePreset FLiveLinkClient::GetSourcePreset(FGuid InSourceGuid, UObje
 	UObject* DuplicatedObjectOuter = InDuplicatedObjectOuter ? InDuplicatedObjectOuter : GetTransientPackage();
 
 	FLiveLinkSourcePreset SourcePreset;
-	if (FLiveLinkCollectionSourceItem* SourceItem = Collection->FindSource(InSourceGuid))
+	if (const FLiveLinkCollectionSourceItem* SourceItem = Collection->FindSource(InSourceGuid))
 	{
 		if (SourceItem->Guid != FLiveLinkSourceCollection::DefaultVirtualSubjectGuid && SourceItem->Setting && SourceItem->Source)
 		{
 			SourcePreset.Guid = SourceItem->Guid;
 			SourcePreset.SourceType = SourceItem->Source->CanBeDisplayedInUI() ? SourceItem->Source->GetSourceType() : FText::GetEmpty();
-			SourcePreset.Settings = DuplicateObject<ULiveLinkSourceSettings>(SourceItem->Setting, DuplicatedObjectOuter);
+			SourcePreset.Settings = DuplicateObject<ULiveLinkSourceSettings>(SourceItem->Setting.Get(), DuplicatedObjectOuter);
 		}
 	}
 	return SourcePreset;
@@ -559,7 +593,7 @@ void FLiveLinkClient::PushSubjectStaticData_AnyThread(const FLiveLinkSubjectKey&
 	const int32 MaxNumBufferToCached = CVarMaxNewStaticDataPerUpdate.GetValueOnAnyThread();
 	bool bLogError = false;
 	{
-		FScopeLock Lock(&CollectionAccessCriticalSection);
+		FScopeLock Lock(&PendingFramesCriticalSection);
 		if (SubjectStaticToPush.Num() > MaxNumBufferToCached) // Something is wrong somewhere. Warn the user and discard the new Static Data.
 		{
 			bLogError = true;
@@ -626,33 +660,41 @@ void FLiveLinkClient::PushSubjectStaticData_Internal(FPendingSubjectStatic&& Sub
 	}
 
 	FLiveLinkSubject* LiveLinkSubject = nullptr;
-	if (FLiveLinkCollectionSubjectItem* SubjectItem = Collection->FindSubject(SubjectStaticData.SubjectKey))
 	{
-		if (!SubjectItem->bPendingKill)
+		if (FLiveLinkCollectionSubjectItem* SubjectItem = Collection->FindSubject(SubjectStaticData.SubjectKey))
 		{
-			LiveLinkSubject = SubjectItem->GetLiveSubject();
-
-			if (LiveLinkSubject->GetRole() != SubjectStaticData.Role)
+			if (!SubjectItem->bPendingKill)
 			{
-				FLiveLinkLog::Warning(TEXT("Subject '%s' of role '%s' is changing its role to '%s'. Current subject will be removed and a new one will be created"), *SubjectStaticData.SubjectKey.SubjectName.ToString(), *LiveLinkSubject->GetRole().GetDefaultObject()->GetDisplayName().ToString(), *SubjectStaticData.Role.GetDefaultObject()->GetDisplayName().ToString());
+				LiveLinkSubject = SubjectItem->GetLiveSubject();
 
-				Collection->RemoveSubject(SubjectStaticData.SubjectKey);
-				LiveLinkSubject = nullptr;
-			}
-			else
-			{
-				LiveLinkSubject->ClearFrames();
+				if (LiveLinkSubject->GetRole() != SubjectStaticData.Role)
+				{
+					FLiveLinkLog::Warning(TEXT("Subject '%s' of role '%s' is changing its role to '%s'. Current subject will be removed and a new one will be created"), *SubjectStaticData.SubjectKey.SubjectName.ToString(), *LiveLinkSubject->GetRole().GetDefaultObject()->GetDisplayName().ToString(), *SubjectStaticData.Role.GetDefaultObject()->GetDisplayName().ToString());
+
+					Collection->RemoveSubject(SubjectStaticData.SubjectKey);
+					LiveLinkSubject = nullptr;
+				}
+				else
+				{
+					LiveLinkSubject->ClearFrames();
+				}
 			}
 		}
 	}
 
-	if(LiveLinkSubject == nullptr)
+
+	// Prevent GC while we're creating UObjects since LL Client can potentially be ticked outside of the game thread.
+	FGCScopeGuard Guard;
+
+	if (LiveLinkSubject == nullptr)
 	{
 		const ULiveLinkSettings* LiveLinkSettings = GetDefault<ULiveLinkSettings>();
 		const FLiveLinkRoleProjectSetting DefaultSetting = LiveLinkSettings->GetDefaultSettingForRole(SubjectStaticData.Role.Get());
 
+		// We use a strong object ptr to prevent GC for this object in case this was created outside the game thread.
+		TStrongObjectPtr<ULiveLinkSubjectSettings> SubjectSettings = nullptr;
+
 		// Setting class should always be valid
-		ULiveLinkSubjectSettings* SubjectSettings = nullptr;
 		{
 			UClass* SettingClass = DefaultSetting.SettingClass.Get();
 			if (SettingClass == nullptr)
@@ -660,7 +702,9 @@ void FLiveLinkClient::PushSubjectStaticData_Internal(FPendingSubjectStatic&& Sub
 				SettingClass = ULiveLinkSubjectSettings::StaticClass();
 			}
 
-			SubjectSettings = NewObject<ULiveLinkSubjectSettings>(GetTransientPackage(), SettingClass);
+			SubjectSettings = TStrongObjectPtr{NewObject<ULiveLinkSubjectSettings>(GetTransientPackage(), SettingClass)};
+			SubjectSettings->Initialize(SubjectStaticData.SubjectKey);
+
 			SubjectSettings->Role = SubjectStaticData.Role;
 
 			UClass* FrameInterpolationProcessorClass = DefaultSetting.FrameInterpolationProcessor.Get();
@@ -669,7 +713,9 @@ void FLiveLinkClient::PushSubjectStaticData_Internal(FPendingSubjectStatic&& Sub
 				UClass* InterpolationRole = FrameInterpolationProcessorClass->GetDefaultObject<ULiveLinkFrameInterpolationProcessor>()->GetRole();
 				if (SubjectStaticData.Role->IsChildOf(InterpolationRole))
 				{
-					SubjectSettings->InterpolationProcessor = NewObject<ULiveLinkFrameInterpolationProcessor>(SubjectSettings, FrameInterpolationProcessorClass);
+					SubjectSettings->InterpolationProcessor = NewObject<ULiveLinkFrameInterpolationProcessor>(SubjectSettings.Get(), FrameInterpolationProcessorClass);
+					// Clear async flag since this might've been created outside the game thread.
+					SubjectSettings->InterpolationProcessor->AtomicallyClearInternalFlags(EInternalObjectFlags::Async);
 				}
 				else
 				{
@@ -685,7 +731,9 @@ void FLiveLinkClient::PushSubjectStaticData_Internal(FPendingSubjectStatic&& Sub
 					UClass* InterpolationRole = FallbackInterpolationProcessorClass->GetDefaultObject<ULiveLinkFrameInterpolationProcessor>()->GetRole();
 					if (SubjectStaticData.Role->IsChildOf(InterpolationRole))
 					{
-						SubjectSettings->InterpolationProcessor = NewObject<ULiveLinkFrameInterpolationProcessor>(SubjectSettings, FallbackInterpolationProcessorClass);
+						SubjectSettings->InterpolationProcessor = NewObject<ULiveLinkFrameInterpolationProcessor>(SubjectSettings.Get(), FallbackInterpolationProcessorClass);
+						// Clear async flag since this might've been created outside the game thread.
+						SubjectSettings->InterpolationProcessor->AtomicallyClearInternalFlags(EInternalObjectFlags::Async);
 					}
 				}
 			}
@@ -697,7 +745,9 @@ void FLiveLinkClient::PushSubjectStaticData_Internal(FPendingSubjectStatic&& Sub
 					UClass* PreProcessorRole = PreProcessor->GetDefaultObject<ULiveLinkFramePreProcessor>()->GetRole();
 					if (SubjectStaticData.Role->IsChildOf(PreProcessorRole))
 					{
-						SubjectSettings->PreProcessors.Add(NewObject<ULiveLinkFramePreProcessor>(SubjectSettings, PreProcessor.Get()));
+						TObjectPtr<ULiveLinkFramePreProcessor>& FramePreprocessor = SubjectSettings->PreProcessors.Add_GetRef(NewObject<ULiveLinkFramePreProcessor>(SubjectSettings.Get(), PreProcessor.Get()));
+						// Clear async flag since this might've been created outside the game thread.
+						FramePreprocessor->AtomicallyClearInternalFlags(EInternalObjectFlags::Async);
 					}
 					else
 					{
@@ -708,9 +758,12 @@ void FLiveLinkClient::PushSubjectStaticData_Internal(FPendingSubjectStatic&& Sub
 		}
 
 		bool bEnabled = Collection->FindEnabledSubject(SubjectStaticData.SubjectKey.SubjectName) == nullptr;
-		FLiveLinkCollectionSubjectItem CollectionSubjectItem(SubjectStaticData.SubjectKey, MakeUnique<FLiveLinkSubject>(SourceItem->TimedData), SubjectSettings, bEnabled);
+		FLiveLinkCollectionSubjectItem CollectionSubjectItem(SubjectStaticData.SubjectKey, MakeUnique<FLiveLinkSubject>(SourceItem->TimedData), SubjectSettings.Get(), bEnabled);
 		CollectionSubjectItem.GetLiveSubject()->Initialize(SubjectStaticData.SubjectKey, SubjectStaticData.Role.Get(), this);
 
+		// Clear the async flag since we've passed  the SubjectSettings to the subject collection item.
+		SubjectSettings->AtomicallyClearInternalFlags(EInternalObjectFlags::Async);
+
 		LiveLinkSubject = CollectionSubjectItem.GetLiveSubject();
 
 		Collection->AddSubject(MoveTemp(CollectionSubjectItem));
@@ -722,6 +775,10 @@ void FLiveLinkClient::PushSubjectStaticData_Internal(FPendingSubjectStatic&& Sub
 		{
 			Handles->OnStaticDataAdded.Broadcast(SubjectStaticData.SubjectKey, SubjectStaticData.Role, SubjectStaticData.StaticData);
 		}
+		else if (const FSubjectFramesAddedHandles* AllSubjectsHandler = SubjectFrameAddedHandles.Find(ALL_SUBJECTS_DELEGATE_TOKEN))
+		{
+			AllSubjectsHandler->OnStaticDataAdded.Broadcast(SubjectStaticData.SubjectKey, SubjectStaticData.Role, SubjectStaticData.StaticData);
+		}
 
 		LiveLinkSubject->SetStaticData(SubjectStaticData.Role, MoveTemp(SubjectStaticData.StaticData));
 	}
@@ -750,7 +807,7 @@ void FLiveLinkClient::PushSubjectFrameData_AnyThread(const FLiveLinkSubjectKey&
 
 	bool bCanPushFrame = true;
 	{
-		FScopeLock Lock(&CollectionAccessCriticalSection);
+		FScopeLock Lock(&PendingFramesCriticalSection);
 
 		if (SubjectFrameToPush.Num() > MaxNumBufferToCached) // Something is wrong somewhere. Warn the user and discard the new Frame Data.
 		{
@@ -769,7 +826,7 @@ void FLiveLinkClient::PushSubjectFrameData_AnyThread(const FLiveLinkSubjectKey&
 		// The lock was released so that OnFrameDataReceived doesn't need to be called with the lock on,
 		// which can hang the game thread when it calls EvaluateFrame if the broadcast takes longer than usual.
 
-		FScopeLock Lock(&CollectionAccessCriticalSection);
+		FScopeLock Lock(&PendingFramesCriticalSection);
 		SubjectFrameToPush.Add(MoveTemp(SubjectFrame));
 	}
 
@@ -792,7 +849,7 @@ void FLiveLinkClient::PushSubjectFrameData_Internal(FPendingSubjectFrame&& Subje
 		return;
 	}
 
-	//To add a frame data, we need to find our subject but also have a static data associated to it. 
+	//To add a frame data, we need to find our subject but also have a static data associated to it.
 	//With presets, the subject could exist but no static data received yet.
 	FLiveLinkCollectionSubjectItem* SubjectItem = Collection->FindSubject(SubjectFrameData.SubjectKey);
 	if (SubjectItem == nullptr)
@@ -804,15 +861,15 @@ void FLiveLinkClient::PushSubjectFrameData_Internal(FPendingSubjectFrame&& Subje
 	{
 		return;
 	}
-	
+
 	FLiveLinkSubject* LinkSubject = SubjectItem->GetLiveSubject();
 	if (LinkSubject == nullptr)
 	{
 		FLiveLinkLog::Error(TEXT("The Subject is not allowed to push to a virtual subject."));
 		return;
 	}
-	
-	
+
+
 	if (!LinkSubject->HasStaticData())
 	{
 		return;
@@ -835,7 +892,7 @@ void FLiveLinkClient::PushSubjectFrameData_Internal(FPendingSubjectFrame&& Subje
 		return;
 	}
 
-	//Stamp arrival time of each packet to track clock difference when it is effectively added to the stash. 
+	//Stamp arrival time of each packet to track clock difference when it is effectively added to the stash.
 	//Doing it in the Add_AnyThread would mean that we stamp it up to 1 frame time behind, causing the offset to always be 1 frame behind
 	//and requiring 2.5 frames or so to have a valid smooth offset
 	if (SubjectFrameData.FrameData.GetBaseData())
@@ -847,7 +904,7 @@ void FLiveLinkClient::PushSubjectFrameData_Internal(FPendingSubjectFrame&& Subje
 			SubjectFrameData.FrameData.GetBaseData()->ArrivalTime.SceneTime = *CurrentTime;
 		}
 	}
-	
+
 	//Let source data know about this new frame to get latest clock offset
 	SourceItem->TimedData->ProcessNewFrameTimingInfo(*SubjectFrameData.FrameData.GetBaseData());
 
@@ -855,6 +912,11 @@ void FLiveLinkClient::PushSubjectFrameData_Internal(FPendingSubjectFrame&& Subje
 	{
 		Handles->OnFrameDataAdded.Broadcast(SubjectItem->Key, Role, SubjectFrameData.FrameData);
 	}
+	else if (const FSubjectFramesAddedHandles* AllSubjectsHandler = SubjectFrameAddedHandles.Find(ALL_SUBJECTS_DELEGATE_TOKEN))
+	{
+		// NAME_None means we registered for all subjects update.
+		AllSubjectsHandler->OnFrameDataAdded.Broadcast(SubjectItem->Key, Role, SubjectFrameData.FrameData);
+	}
 
 	HandleSubjectRebroadcast(LinkSubject, SubjectFrameData.FrameData);
 
@@ -894,7 +956,6 @@ bool FLiveLinkClient::CreateSubject(const FLiveLinkSubjectPreset& InSubjectPrese
 	FLiveLinkCollectionSubjectItem* SubjectItem = Collection->FindSubject(InSubjectPreset.Key);
 	if (SubjectItem != nullptr)
 	{
-		FScopeLock Lock(&CollectionAccessCriticalSection);
 		if (SubjectItem->bPendingKill)
 		{
 			Collection->RemoveSubject(InSubjectPreset.Key);
@@ -913,7 +974,6 @@ bool FLiveLinkClient::CreateSubject(const FLiveLinkSubjectPreset& InSubjectPrese
 		FLiveLinkCollectionSubjectItem VSubjectData(InSubjectPreset.Key, VSubject, bEnabled);
 		VSubject->Initialize(VSubjectData.Key, VSubject->GetRole(), this);
 
-		FScopeLock Lock(&CollectionAccessCriticalSection);
 		Collection->AddSubject(MoveTemp(VSubjectData));
 		Collection->SetSubjectEnabled(InSubjectPreset.Key, InSubjectPreset.bEnabled);
 	}
@@ -933,7 +993,6 @@ bool FLiveLinkClient::CreateSubject(const FLiveLinkSubjectPreset& InSubjectPrese
 		FLiveLinkCollectionSubjectItem CollectionSubjectItem(InSubjectPreset.Key, MakeUnique<FLiveLinkSubject>(SourceItem->TimedData), SubjectSettings, bEnabled);
 		CollectionSubjectItem.GetLiveSubject()->Initialize(InSubjectPreset.Key, InSubjectPreset.Role.Get(), this);
 
-		FScopeLock Lock(&CollectionAccessCriticalSection);
 		Collection->AddSubject(MoveTemp(CollectionSubjectItem));
 		Collection->SetSubjectEnabled(InSubjectPreset.Key, InSubjectPreset.bEnabled);
 	}
@@ -942,8 +1001,6 @@ bool FLiveLinkClient::CreateSubject(const FLiveLinkSubjectPreset& InSubjectPrese
 
 void FLiveLinkClient::RemoveSubject_AnyThread(const FLiveLinkSubjectKey& InSubjectKey)
 {
-	FScopeLock Lock(&CollectionAccessCriticalSection);
-
 	if (Collection)
 	{
 		if (FLiveLinkCollectionSubjectItem* SubjectItem = Collection->FindSubject(InSubjectKey))
@@ -966,12 +1023,9 @@ bool FLiveLinkClient::AddVirtualSubject(const FLiveLinkSubjectKey& InVirtualSubj
 		}
 		else
 		{
-			FScopeLock Lock(&CollectionAccessCriticalSection);
-			const bool bFoundVirtualSubject = nullptr != Collection->GetSubjects().FindByPredicate(
-				[InVirtualSubjectKey](const FLiveLinkCollectionSubjectItem& Other)
-			{
-				return Other.Key == InVirtualSubjectKey && Other.GetVirtualSubject() != nullptr;
-			});
+
+			const FLiveLinkCollectionSubjectItem* SubjectItem = Collection->FindSubject(InVirtualSubjectKey);
+			const bool bFoundVirtualSubject = SubjectItem && SubjectItem->GetVirtualSubject();
 
 			if (!bFoundVirtualSubject)
 			{
@@ -990,7 +1044,7 @@ bool FLiveLinkClient::AddVirtualSubject(const FLiveLinkSubjectKey& InVirtualSubj
 					{
 						Blueprint->OnCompiled().AddLambda([this, SubjectKey = VSubjectData.Key](UBlueprint* BP) {
 							this->ReinitializeVirtualSubject(SubjectKey);
-						});
+							});
 					}
 				}
 #endif
@@ -1011,7 +1065,6 @@ bool FLiveLinkClient::AddVirtualSubject(const FLiveLinkSubjectKey& InVirtualSubj
 
 void FLiveLinkClient::RemoveVirtualSubject(const FLiveLinkSubjectKey& InVirtualSubjectKey)
 {
-	FScopeLock Lock(&CollectionAccessCriticalSection);
 	if (Collection)
 	{
 		Collection->RemoveSubject(InVirtualSubjectKey);
@@ -1020,8 +1073,6 @@ void FLiveLinkClient::RemoveVirtualSubject(const FLiveLinkSubjectKey& InVirtualS
 
 void FLiveLinkClient::ClearSubjectsFrames_AnyThread(FLiveLinkSubjectName InSubjectName)
 {
-	FScopeLock Lock(&CollectionAccessCriticalSection);
-
 	// Use the subject enabled for at this frame
 	if (FLiveLinkSubjectKey* SubjectKey = EnabledSubjects.Find(InSubjectName))
 	{
@@ -1031,27 +1082,25 @@ void FLiveLinkClient::ClearSubjectsFrames_AnyThread(FLiveLinkSubjectName InSubje
 
 void FLiveLinkClient::ClearSubjectsFrames_AnyThread(const FLiveLinkSubjectKey& InSubjectKey)
 {
-	FScopeLock Lock(&CollectionAccessCriticalSection);
-
 	if (Collection)
 	{
 		if (FLiveLinkCollectionSubjectItem* SubjectItem = Collection->FindSubject(InSubjectKey))
 		{
 			SubjectItem->GetSubject()->ClearFrames();
+
+
 		}
 	}
 }
 
 void FLiveLinkClient::ClearAllSubjectsFrames_AnyThread()
 {
-	FScopeLock Lock(&CollectionAccessCriticalSection);
-
 	if (Collection)
 	{
-		for (const FLiveLinkCollectionSubjectItem& SubjectItem : Collection->GetSubjects())
+		Collection->ForEachSubject([](FLiveLinkCollectionSourceItem& SourceItem, FLiveLinkCollectionSubjectItem& SubjectItem)
 		{
 			SubjectItem.GetSubject()->ClearFrames();
-		}
+		});
 	}
 }
 
@@ -1076,6 +1125,7 @@ FLiveLinkSubjectPreset FLiveLinkClient::GetSubjectPreset(const FLiveLinkSubjectK
 	UObject* DuplicatedObjectOuter = InDuplicatedObjectOuter ? InDuplicatedObjectOuter : GetTransientPackage();
 
 	FLiveLinkSubjectPreset SubjectPreset;
+
 	if (const FLiveLinkCollectionSubjectItem* SubjectItem = Collection->FindSubject(InSubjectKey))
 	{
 		SubjectPreset.Key = SubjectItem->Key;
@@ -1090,21 +1140,24 @@ FLiveLinkSubjectPreset FLiveLinkClient::GetSubjectPreset(const FLiveLinkSubjectK
 			SubjectPreset.Settings = DuplicateObject<ULiveLinkSubjectSettings>(SubjectItem->GetLinkSettings(), DuplicatedObjectOuter);
 		}
 	}
+
+
 	return SubjectPreset;
 }
 
 TArray<FLiveLinkSubjectKey> FLiveLinkClient::GetSubjects(bool bIncludeDisabledSubject, bool bIncludeVirtualSubject) const
 {
 	TArray<FLiveLinkSubjectKey> SubjectEntries;
-	SubjectEntries.Reserve(Collection->GetSubjects().Num());
 
-	for (const FLiveLinkCollectionSubjectItem& SubjectItem : Collection->GetSubjects())
+	SubjectEntries.Reserve(Collection->NumSubjects());
+
+	Collection->ForEachSubject([this, bIncludeDisabledSubject, bIncludeVirtualSubject, &SubjectEntries](const FLiveLinkCollectionSourceItem& SourceItem, const FLiveLinkCollectionSubjectItem& SubjectItem)
 	{
 		if ((SubjectItem.bEnabled || bIncludeDisabledSubject) && (bIncludeVirtualSubject || SubjectItem.GetVirtualSubject() == nullptr))
 		{
 			SubjectEntries.Add(SubjectItem.Key);
 		}
-	}
+	});
 
 	return SubjectEntries;
 }
@@ -1143,6 +1196,7 @@ bool FLiveLinkClient::IsSubjectEnabled(const FLiveLinkSubjectKey& InSubjectKey,
 		}
 		return false;
 	}
+
 	return Collection->IsSubjectEnabled(InSubjectKey);
 }
 
@@ -1182,8 +1236,6 @@ bool FLiveLinkClient::IsSubjectTimeSynchronized(FLiveLinkSubjectName InSubjectNa
 
 TSubclassOf<ULiveLinkRole> FLiveLinkClient::GetSubjectRole_AnyThread(const FLiveLinkSubjectKey& InSubjectKey) const
 {
-	FScopeLock Lock(&CollectionAccessCriticalSection);
-
 	if (const FLiveLinkCollectionSubjectItem* SubjectItem = Collection->FindSubject(InSubjectKey))
 	{
 		return SubjectItem->GetSubject()->GetRole();
@@ -1194,8 +1246,6 @@ TSubclassOf<ULiveLinkRole> FLiveLinkClient::GetSubjectRole_AnyThread(const FLive
 
 TSubclassOf<ULiveLinkRole> FLiveLinkClient::GetSubjectRole_AnyThread(FLiveLinkSubjectName InSubjectName) const
 {
-	FScopeLock Lock(&CollectionAccessCriticalSection);
-
 	if (const FLiveLinkCollectionSubjectItem* SubjectItem = Collection->FindEnabledSubject(InSubjectName))
 	{
 		return SubjectItem->GetSubject()->GetRole();
@@ -1206,8 +1256,6 @@ TSubclassOf<ULiveLinkRole> FLiveLinkClient::GetSubjectRole_AnyThread(FLiveLinkSu
 
 bool FLiveLinkClient::DoesSubjectSupportsRole_AnyThread(const FLiveLinkSubjectKey& InSubjectKey, TSubclassOf<ULiveLinkRole> InSupportedRole) const
 {
-	FScopeLock Lock(&CollectionAccessCriticalSection);
-
 	if (const FLiveLinkCollectionSubjectItem* SubjectItem = Collection->FindSubject(InSubjectKey))
 	{
 		return SubjectItem->GetSubject()->SupportsRole(InSupportedRole);
@@ -1218,8 +1266,6 @@ bool FLiveLinkClient::DoesSubjectSupportsRole_AnyThread(const FLiveLinkSubjectKe
 
 bool FLiveLinkClient::DoesSubjectSupportsRole_AnyThread(FLiveLinkSubjectName InSubjectName, TSubclassOf<ULiveLinkRole> InSupportedRole) const
 {
-	FScopeLock Lock(&CollectionAccessCriticalSection);
-
 	if (const FLiveLinkCollectionSubjectItem* SubjectItem = Collection->FindEnabledSubject(InSubjectName))
 	{
 		return SubjectItem->GetSubject()->SupportsRole(InSupportedRole);
@@ -1250,7 +1296,8 @@ TArray<FLiveLinkTime> FLiveLinkClient::GetSubjectFrameTimes(FLiveLinkSubjectName
 TArray<FLiveLinkSubjectKey> FLiveLinkClient::GetSubjectsSupportingRole(TSubclassOf<ULiveLinkRole> InSupportedRole, bool bIncludeDisabledSubject, bool bIncludeVirtualSubject) const
 {
 	TArray<FLiveLinkSubjectKey> SubjectKeys;
-	for (const FLiveLinkCollectionSubjectItem& SubjectItem : Collection->GetSubjects())
+
+	Collection->ForEachSubject([this, &SubjectKeys, InSupportedRole, bIncludeDisabledSubject, bIncludeVirtualSubject](const FLiveLinkCollectionSourceItem& SourceItem, const FLiveLinkCollectionSubjectItem& SubjectItem)
 	{
 		if (SubjectItem.GetSubject()->SupportsRole(InSupportedRole))
 		{
@@ -1259,7 +1306,8 @@ TArray<FLiveLinkSubjectKey> FLiveLinkClient::GetSubjectsSupportingRole(TSubclass
 				SubjectKeys.Add(SubjectItem.Key);
 			}
 		}
-	}
+	});
+
 	return SubjectKeys;
 }
 
@@ -1267,8 +1315,6 @@ bool FLiveLinkClient::EvaluateFrameFromSource_AnyThread(const FLiveLinkSubjectKe
 {
 	SCOPE_CYCLE_COUNTER(STAT_LiveLink_EvaluateFrame);
 
-	FScopeLock Lock(&CollectionAccessCriticalSection);
-
 	if (const FLiveLinkCollectionSubjectItem* SubjectItem = Collection->FindSubject(InSubjectKey))
 	{
 		return SubjectItem->GetSubject()->EvaluateFrame(InDesiredRole, OutFrame);
@@ -1285,7 +1331,7 @@ void FLiveLinkClient::ForceTick()
 
 bool FLiveLinkClient::HasPendingSubjectFrames()
 {
-	FScopeLock Lock(&CollectionAccessCriticalSection);
+	FScopeLock PendingFramesLock(&PendingFramesCriticalSection);
 	return !SubjectFrameToPush.IsEmpty();
 }
 
@@ -1293,8 +1339,6 @@ bool FLiveLinkClient::EvaluateFrame_AnyThread(FLiveLinkSubjectName InSubjectName
 {
 	SCOPE_CYCLE_COUNTER(STAT_LiveLink_EvaluateFrame);
 
-	FScopeLock Lock(&CollectionAccessCriticalSection);
-
 	bool bResult = false;
 
 	// Used the cached enabled list
@@ -1330,8 +1374,6 @@ bool FLiveLinkClient::EvaluateFrameAtWorldTime_AnyThread(FLiveLinkSubjectName In
 {
 	SCOPE_CYCLE_COUNTER(STAT_LiveLink_EvaluateFrame);
 
-	FScopeLock Lock(&CollectionAccessCriticalSection);
-
 	bool bResult = false;
 
 	// Used the cached enabled list
@@ -1374,7 +1416,7 @@ bool FLiveLinkClient::EvaluateFrameAtSceneTime_AnyThread(FLiveLinkSubjectName In
 {
 	SCOPE_CYCLE_COUNTER(STAT_LiveLink_EvaluateFrame);
 
-	FScopeLock Lock(&CollectionAccessCriticalSection);
+
 
 	bool bResult = false;
 
@@ -1423,24 +1465,19 @@ TArray<FGuid> FLiveLinkClient::GetDisplayableSources(bool bIncludeVirtualSources
 {
 	TArray<FGuid> Results;
 
-	const TArray<FLiveLinkCollectionSourceItem>& PresetSources = Collection->GetSources();
-	Results.Reserve(PresetSources.Num());
-
-	for (const FLiveLinkCollectionSourceItem& Data : PresetSources)
+	Collection->ForEachSource([&Results, bIncludeVirtualSources](const FLiveLinkCollectionSourceItem& SourceItem)
 	{
-		if (Data.Source->CanBeDisplayedInUI() || (bIncludeVirtualSources && Data.IsVirtualSource()))
+		if (SourceItem.Source->CanBeDisplayedInUI() || (bIncludeVirtualSources && SourceItem.IsVirtualSource()))
 		{
-			Results.Add(Data.Guid);
+			Results.Add(SourceItem.Guid);
 		}
-	}
+	});
 
 	return Results;
 }
 
 FLiveLinkSubjectTimeSyncData FLiveLinkClient::GetTimeSyncData(FLiveLinkSubjectName InSubjectName)
 {
-	FScopeLock Lock(&CollectionAccessCriticalSection);
-
 	if (const FLiveLinkCollectionSubjectItem* SubjectItem = Collection->FindEnabledSubject(InSubjectName))
 	{
 		if (FLiveLinkSubject* LinkSubject = SubjectItem->GetLiveSubject())
@@ -1499,9 +1536,10 @@ bool FLiveLinkClient::IsVirtualSubject(const FLiveLinkSubjectKey& InSubjectKey)
 
 void FLiveLinkClient::OnPropertyChanged(FGuid InEntryGuid, const FPropertyChangedEvent& InPropertyChangedEvent)
 {
+
 	if (const FLiveLinkCollectionSourceItem* SourceItem = Collection->FindSource(InEntryGuid))
 	{
-		SourceItem->Source->OnSettingsChanged(SourceItem->Setting, InPropertyChangedEvent);
+		SourceItem->Source->OnSettingsChanged(SourceItem->Setting.Get(), InPropertyChangedEvent);
 	}
 }
 
@@ -1509,7 +1547,7 @@ ULiveLinkSourceSettings* FLiveLinkClient::GetSourceSettings(const FGuid& InEntry
 {
 	if (const FLiveLinkCollectionSourceItem* SourceItem = Collection->FindSource(InEntryGuid))
 	{
-		return SourceItem->Setting;
+		return SourceItem->Setting.Get();
 	}
 	return nullptr;
 }
@@ -1525,7 +1563,6 @@ UObject* FLiveLinkClient::GetSubjectSettings(const FLiveLinkSubjectKey& InSubjec
 
 const FLiveLinkStaticDataStruct* FLiveLinkClient::GetSubjectStaticData_AnyThread(const FLiveLinkSubjectKey& InSubjectKey) const
 {
-	FScopeLock Lock(&CollectionAccessCriticalSection);
 	if (const FLiveLinkCollectionSubjectItem* SubjectItem = Collection->FindSubject(InSubjectKey))
 	{
 		if (FLiveLinkSubject* LiveLinkSubject = SubjectItem->GetLiveSubject())
@@ -1617,7 +1654,6 @@ FSimpleMulticastDelegate& FLiveLinkClient::OnLiveLinkSubjectsChanged()
 	return Collection->OnLiveLinkSubjectsChanged();
 }
 
-
 FOnLiveLinkSourceChangedDelegate& FLiveLinkClient::OnLiveLinkSourceAdded()
 {
 	return Collection->OnLiveLinkSourceAdded();
@@ -1651,18 +1687,14 @@ PRAGMA_DISABLE_DEPRECATION_WARNINGS
  */
 void FLiveLinkClient::AquireLock_Deprecation() 
 {
-	CollectionAccessCriticalSection.Lock();
 }
 
 void FLiveLinkClient::ReleaseLock_Deprecation()
 {
-	CollectionAccessCriticalSection.Unlock();
 }
 
 void FLiveLinkClient::ClearFrames_Deprecation(const FLiveLinkSubjectKey& InSubjectKey)
 {
-	FScopeLock Lock(&CollectionAccessCriticalSection);
-
 	if (Collection)
 	{
 		if (FLiveLinkCollectionSubjectItem* SubjectItem = Collection->FindSubject(InSubjectKey))
@@ -1674,8 +1706,6 @@ void FLiveLinkClient::ClearFrames_Deprecation(const FLiveLinkSubjectKey& InSubje
 
 FLiveLinkSkeletonStaticData* FLiveLinkClient::GetSubjectAnimationStaticData_Deprecation(const FLiveLinkSubjectKey& InSubjectKey)
 {
-	FScopeLock Lock(&CollectionAccessCriticalSection);
-
 	if (Collection)
 	{
 		if (FLiveLinkCollectionSubjectItem* SubjectItem = Collection->FindSubject(InSubjectKey))
@@ -1690,6 +1720,23 @@ FLiveLinkSkeletonStaticData* FLiveLinkClient::GetSubjectAnimationStaticData_Depr
 	return nullptr;
 }
 
+void FLiveLinkClient::UnregisterGlobalSubjectFramesDelegate(FDelegateHandle& InStaticDataAddedHandle, FDelegateHandle& InFrameDataAddedHandle)
+{
+	if (FSubjectFramesAddedHandles* Handles = SubjectFrameAddedHandles.Find(ALL_SUBJECTS_DELEGATE_TOKEN))
+	{
+		Handles->OnStaticDataAdded.Remove(InStaticDataAddedHandle);
+		Handles->OnFrameDataAdded.Remove(InFrameDataAddedHandle);
+	}
+}
+
+bool FLiveLinkClient::RegisterGlobalSubjectFramesDelegate(const FOnLiveLinkSubjectStaticDataAdded::FDelegate& InOnStaticDataAdded, const FOnLiveLinkSubjectFrameDataAdded::FDelegate& InOnFrameDataAdded, FDelegateHandle& OutStaticDataAddedHandle, FDelegateHandle& OutFrameDataAddedHandle)
+{
+	FSubjectFramesAddedHandles& Handles = SubjectFrameAddedHandles.FindOrAdd(ALL_SUBJECTS_DELEGATE_TOKEN);
+	OutStaticDataAddedHandle = Handles.OnStaticDataAdded.Add(InOnStaticDataAdded);
+	OutFrameDataAddedHandle = Handles.OnFrameDataAdded.Add(InOnFrameDataAdded);
+
+	return true;
+}
 
 /**
  * Function that are now deprecated
diff --git a/Engine/Plugins/Animation/LiveLink/Source/LiveLink/Private/LiveLinkSourceCollection.cpp b/Engine/Plugins/Animation/LiveLink/Source/LiveLink/Private/LiveLinkSourceCollection.cpp
--- a/Engine/Plugins/Animation/LiveLink/Source/LiveLink/Private/LiveLinkSourceCollection.cpp
+++ b/Engine/Plugins/Animation/LiveLink/Source/LiveLink/Private/LiveLinkSourceCollection.cpp
@@ -3,6 +3,7 @@
 #include "LiveLinkSourceCollection.h"
 
 #include "EngineAnalytics.h"
+#include "Experimental/Async/MultiUniqueLock.h"
 #include "IAnalyticsProviderET.h"
 #include "LiveLinkVirtualSource.h"
 #include "UObject/Package.h"
@@ -50,6 +51,7 @@ FLiveLinkCollectionSubjectItem::FLiveLinkCollectionSubjectItem(FLiveLinkSubjectK
 {
 }
 
+
 FLiveLinkCollectionSubjectItem::FLiveLinkCollectionSubjectItem(FLiveLinkSubjectKey InKey, ULiveLinkVirtualSubject* InVirtualSubject, bool bInEnabled)
 	: Key(InKey)
 	, bEnabled(bInEnabled)
@@ -59,7 +61,6 @@ FLiveLinkCollectionSubjectItem::FLiveLinkCollectionSubjectItem(FLiveLinkSubjectK
 {
 }
 
-
 const FGuid FLiveLinkSourceCollection::DefaultVirtualSubjectGuid{ 0x4ed2dc4e, 0xcc5911ce, 0x4af0635d, 0xa8b24a5a };
 
 
@@ -70,29 +71,22 @@ FLiveLinkSourceCollection::FLiveLinkSourceCollection()
 	Data.Guid = DefaultVirtualSubjectGuid;
 	ULiveLinkVirtualSubjectSourceSettings* NewSettings = NewObject<ULiveLinkVirtualSubjectSourceSettings>(GetTransientPackage(), ULiveLinkVirtualSubjectSourceSettings::StaticClass());
 	NewSettings->SourceName = TEXT("DefaultVirtualSource");
-	Data.Setting = NewSettings;
+	Data.Setting = TStrongObjectPtr(NewSettings);
 	Data.bIsVirtualSource = true;
 	Data.Source->InitializeSettings(NewSettings);
 
-	Sources.Add(MoveTemp(Data));
-}
-
-void FLiveLinkSourceCollection::AddReferencedObjects(FReferenceCollector & Collector)
-{
-	for (FLiveLinkCollectionSourceItem& Item : Sources)
+	if (!IsInGameThread())
 	{
-		Collector.AddReferencedObject(Item.Setting);
+		// If the settings object was created outside of the game thread, we need to clear the async flag to allow the object to be garbage collected.
+		Data.Setting->AtomicallyClearInternalFlags(EInternalObjectFlags::Async);
 	}
 
-	for (FLiveLinkCollectionSubjectItem& Item : Subjects)
-	{
-		Collector.AddReferencedObject(Item.VirtualSubject);
-		Collector.AddReferencedObject(Item.Setting);
-	}
+	Sources.Add(MoveTemp(Data));
 }
 
 void FLiveLinkSourceCollection::AddSource(FLiveLinkCollectionSourceItem InSource)
 {
+	UE::TUniqueLock Lock(SourcesLock);
 	FLiveLinkCollectionSourceItem& SourceItem = Sources.Add_GetRef(MoveTemp(InSource));
 
 	OnLiveLinkSourceAdded().Broadcast(SourceItem.Guid);
@@ -106,6 +100,8 @@ void FLiveLinkSourceCollection::RemoveSource(FGuid InSourceGuid)
 {
 	if (InSourceGuid != FLiveLinkSourceCollection::DefaultVirtualSubjectGuid)
 	{
+		UE::TMultiUniqueLock<UE::FRecursiveMutex> MultiLock({&SubjectsLock, &SourcesLock});
+
 		int32 SourceIndex = Sources.IndexOfByPredicate([InSourceGuid](const FLiveLinkCollectionSourceItem& Other) { return Other.Guid == InSourceGuid; });
 		if (SourceIndex != INDEX_NONE)
 		{
@@ -118,7 +114,7 @@ void FLiveLinkSourceCollection::RemoveSource(FGuid InSourceGuid)
 					bRemovedSubject = true;
 					FLiveLinkSubjectKey Key = Subjects[SubjectIndex].Key;
 					Subjects.RemoveAtSwap(SubjectIndex);
-					OnLiveLinkSubjectRemoved().Broadcast(Key);
+					BroadcastOnGameThread(OnLiveLinkSubjectRemoved(), Key);
 				}
 			}
 
@@ -137,38 +133,48 @@ void FLiveLinkSourceCollection::RemoveSource(FGuid InSourceGuid)
 
 void FLiveLinkSourceCollection::RemoveAllSources()
 {
-	const bool bHasSubjects = Subjects.Num() > 0;
-	for (int32 Index = Subjects.Num() - 1; Index >= 0; --Index)
+	bool bHasRemovedSubject = false;
 	{
-		FLiveLinkSubjectKey Key = Subjects[Index].Key;
-		Subjects.RemoveAtSwap(Index);
-		OnLiveLinkSubjectRemoved().Broadcast(Key);
+		UE::TUniqueLock Lock(SubjectsLock);
+		for (int32 Index = Subjects.Num() - 1; Index >= 0; --Index)
+		{
+			bHasRemovedSubject = true;
+			FLiveLinkSubjectKey Key = Subjects[Index].Key;
+			Subjects.RemoveAtSwap(Index);
+			BroadcastOnGameThread(OnLiveLinkSubjectRemoved(), Key);
+		}
 	}
-	if (bHasSubjects)
+
+	if (bHasRemovedSubject)
 	{
-		OnLiveLinkSubjectsChanged().Broadcast();
+		BroadcastOnGameThread(OnLiveLinkSubjectsChanged());
 	}
 
 	bool bHasRemovedSource = false;
-	for (int32 Index = Sources.Num() - 1; Index >= 0; --Index)
 	{
-		if (Sources[Index].Guid != FLiveLinkSourceCollection::DefaultVirtualSubjectGuid)
+		UE::TUniqueLock Lock(SubjectsLock);
+		for (int32 Index = Sources.Num() - 1; Index >= 0; --Index)
 		{
-			bHasRemovedSource = true;
-			FGuid Key = Sources[Index].Guid;
-			Sources.RemoveAtSwap(Index);
-			OnLiveLinkSourceRemoved().Broadcast(Key);
+			if (Sources[Index].Guid != FLiveLinkSourceCollection::DefaultVirtualSubjectGuid)
+			{
+				bHasRemovedSource = true;
+				FGuid Key = Sources[Index].Guid;
+				Sources.RemoveAtSwap(Index);
+				BroadcastOnGameThread(OnLiveLinkSourceRemoved(), Key);
+			}
 		}
 	}
+
 	if (bHasRemovedSource)
 	{
-		OnLiveLinkSourcesChanged().Broadcast();
+		BroadcastOnGameThread(OnLiveLinkSourcesChanged());
 	}
 }
 
 
 FLiveLinkCollectionSourceItem* FLiveLinkSourceCollection::FindSource(TSharedPtr<ILiveLinkSource> InSource)
 {
+	UE::TUniqueLock Lock(SourcesLock);
 	return Sources.FindByPredicate([InSource](const FLiveLinkCollectionSourceItem& Other) { return Other.Source == InSource; });
 }
 
@@ -181,6 +187,7 @@ const FLiveLinkCollectionSourceItem* FLiveLinkSourceCollection::FindSource(TShar
 
 FLiveLinkCollectionSourceItem* FLiveLinkSourceCollection::FindSource(FGuid InSourceGuid)
 {
+	UE::TUniqueLock Lock(SourcesLock);
 	return Sources.FindByPredicate([InSourceGuid](const FLiveLinkCollectionSourceItem& Other) { return Other.Guid == InSourceGuid; });
 }
 
@@ -193,16 +200,17 @@ const FLiveLinkCollectionSourceItem* FLiveLinkSourceCollection::FindSource(FGuid
 
 FLiveLinkCollectionSourceItem* FLiveLinkSourceCollection::FindVirtualSource(FName VirtualSourceName)
 {
+	UE::TUniqueLock Lock(SourcesLock);
 	return Sources.FindByPredicate([VirtualSourceName](const FLiveLinkCollectionSourceItem& Other)
 	{
 		if (Other.IsVirtualSource())
 		{
-			if (ULiveLinkVirtualSubjectSourceSettings* VirtualSubjectSettings = Cast<ULiveLinkVirtualSubjectSourceSettings>(Other.Setting))
+			if (ULiveLinkVirtualSubjectSourceSettings* VirtualSubjectSettings = Cast<ULiveLinkVirtualSubjectSourceSettings>(Other.Setting.Get()))
 			{
 				return VirtualSubjectSettings->SourceName == VirtualSourceName;
 			}
 		}
-		return false; 
+		return false;
 	});
 }
 
@@ -211,29 +219,45 @@ const FLiveLinkCollectionSourceItem* FLiveLinkSourceCollection::FindVirtualSourc
 	return const_cast<FLiveLinkSourceCollection*>(this)->FindVirtualSource(VirtualSourceName);
 }
 
+int32 FLiveLinkSourceCollection::NumSources() const
+{
+	UE::TUniqueLock Lock(SourcesLock);
+	return Sources.Num();
+}
+
 void FLiveLinkSourceCollection::AddSubject(FLiveLinkCollectionSubjectItem InSubject)
 {
 	FLiveLinkSubjectKey Key = InSubject.Key;
-	Subjects.Add(MoveTemp(InSubject));
-	OnLiveLinkSubjectAdded().Broadcast(Key);
-	OnLiveLinkSubjectsChanged().Broadcast();
+
+	{
+		UE::TUniqueLock Lock(SubjectsLock);
+		Subjects.Add(MoveTemp(InSubject));
+	}
+
+	BroadcastOnGameThread(OnLiveLinkSubjectAdded(), Key);
+	BroadcastOnGameThread(OnLiveLinkSubjectsChanged());
 }
 
 
 void FLiveLinkSourceCollection::RemoveSubject(FLiveLinkSubjectKey InSubjectKey)
 {
-	int32 IndexOf = Subjects.IndexOfByPredicate([InSubjectKey](const FLiveLinkCollectionSubjectItem& Other) { return Other.Key == InSubjectKey; });
-	if (IndexOf != INDEX_NONE)
 	{
-		Subjects.RemoveAtSwap(IndexOf);
-		OnLiveLinkSubjectRemoved().Broadcast(InSubjectKey);
-		OnLiveLinkSubjectsChanged().Broadcast();
+		UE::TUniqueLock Lock(SubjectsLock);
+		int32 IndexOf = Subjects.IndexOfByPredicate([InSubjectKey](const FLiveLinkCollectionSubjectItem& Other) { return Other.Key == InSubjectKey; });
+		if (IndexOf != INDEX_NONE)
+		{
+			Subjects.RemoveAtSwap(IndexOf);
+		}
 	}
+
+	BroadcastOnGameThread(OnLiveLinkSubjectRemoved(), InSubjectKey);
+    BroadcastOnGameThread(OnLiveLinkSubjectsChanged());
 }
 
 
 FLiveLinkCollectionSubjectItem* FLiveLinkSourceCollection::FindSubject(FLiveLinkSubjectKey InSubjectKey)
 {
+	UE::TUniqueLock Lock(SubjectsLock);
 	return Subjects.FindByPredicate([InSubjectKey](const FLiveLinkCollectionSubjectItem& Other) { return Other.Key == InSubjectKey; });
 }
 
@@ -243,15 +267,22 @@ const FLiveLinkCollectionSubjectItem* FLiveLinkSourceCollection::FindSubject(FLi
 	return const_cast<FLiveLinkSourceCollection*>(this)->FindSubject(InSubjectKey);
 }
 
-
 const FLiveLinkCollectionSubjectItem* FLiveLinkSourceCollection::FindEnabledSubject(FLiveLinkSubjectName InSubjectName) const
 {
+	UE::TUniqueLock Lock(SubjectsLock);
 	return Subjects.FindByPredicate([InSubjectName](const FLiveLinkCollectionSubjectItem& Other) { return Other.Key.SubjectName == InSubjectName && Other.bEnabled && !Other.bPendingKill; });
 }
 
+int32 FLiveLinkSourceCollection::NumSubjects() const
+{
+	UE::TUniqueLock Lock(SubjectsLock);
+	return Subjects.Num();
+}
+
 
 bool FLiveLinkSourceCollection::IsSubjectEnabled(FLiveLinkSubjectKey InSubjectKey) const
 {
+	UE::TUniqueLock Lock(SubjectsLock);
 	if (const FLiveLinkCollectionSubjectItem* Item = FindSubject(InSubjectKey))
 	{
 		return Item->bEnabled;
@@ -267,6 +298,8 @@ void FLiveLinkSourceCollection::SetSubjectEnabled(FLiveLinkSubjectKey InSubjectK
 		// clear all bEnabled only if found
 		if (FLiveLinkCollectionSubjectItem* NewEnabledItem = FindSubject(InSubjectKey))
 		{
+			UE::TUniqueLock Lock(SubjectsLock);
+
 			NewEnabledItem->bEnabled = true;
 
 			for (FLiveLinkCollectionSubjectItem& SubjectItem : Subjects)
@@ -280,6 +313,7 @@ void FLiveLinkSourceCollection::SetSubjectEnabled(FLiveLinkSubjectKey InSubjectK
 	}
 	else
 	{
+		UE::TUniqueLock Lock(SubjectsLock);
 		for (FLiveLinkCollectionSubjectItem& SubjectItem : Subjects)
 		{
 			if (SubjectItem.Key.SubjectName == InSubjectKey.SubjectName)
@@ -290,9 +324,11 @@ void FLiveLinkSourceCollection::SetSubjectEnabled(FLiveLinkSubjectKey InSubjectK
 	}
 }
 
+
 void FLiveLinkSourceCollection::RemovePendingKill()
 {
-	// Remove Sources that are pending kill
+	UE::TMultiUniqueLock<UE::FRecursiveMutex> MultiLock({&SubjectsLock, &SourcesLock});
+
 	for (int32 SourceIndex = Sources.Num() - 1; SourceIndex >= 0; --SourceIndex)
 	{
 		FLiveLinkCollectionSourceItem& SourceItem = Sources[SourceIndex];
@@ -328,11 +364,15 @@ void FLiveLinkSourceCollection::RemovePendingKill()
 	}
 }
 
+
 bool FLiveLinkSourceCollection::RequestShutdown()
 {
-	bool bHadSubject = Subjects.Num() > 0;
-	Subjects.Reset();
+	{
+		UE::TUniqueLock Lock(SubjectsLock);
+		Subjects.Reset();
+	}
 
+	UE::TUniqueLock Lock(SourcesLock);
 	for (int32 SourceIndex = Sources.Num() - 1; SourceIndex >= 0; --SourceIndex)
 	{
 		FLiveLinkCollectionSourceItem& SourceItem = Sources[SourceIndex];
@@ -343,6 +383,55 @@ bool FLiveLinkSourceCollection::RequestShutdown()
 	}
 
 	// No callback when we shutdown
-
 	return Sources.Num() == 0;
 }
+
+
+void FLiveLinkSourceCollection::ForEachSubject(TFunctionRef<void(FLiveLinkCollectionSourceItem&, FLiveLinkCollectionSubjectItem&)> VisitorFunc)
+{
+	UE::TMultiUniqueLock<UE::FRecursiveMutex> MultiLock({&SubjectsLock, &SourcesLock});
+
+	for (FLiveLinkCollectionSubjectItem& Subject : Subjects)
+	{
+		if (FLiveLinkCollectionSourceItem* SourceItem = FindSource(Subject.Key.Source))
+		{
+			VisitorFunc(*SourceItem, Subject);
+		}
+	}
+}
+
+
+void FLiveLinkSourceCollection::ForEachSubject(TFunctionRef<void(const FLiveLinkCollectionSourceItem&, const FLiveLinkCollectionSubjectItem&)> VisitorFunc) const
+{
+	UE::TMultiUniqueLock<UE::FRecursiveMutex> MultiLock({ &SubjectsLock, &SourcesLock });
+
+	for (const FLiveLinkCollectionSubjectItem& Subject : Subjects)
+	{
+		if (const FLiveLinkCollectionSourceItem* SourceItem = FindSource(Subject.Key.Source))
+		{
+			VisitorFunc(*SourceItem, Subject);
+		}
+	}
+}
+
+
+void FLiveLinkSourceCollection::ForEachSource(TFunctionRef<void(FLiveLinkCollectionSourceItem&)> VisitorFunc)
+{
+	UE::TUniqueLock Lock(SourcesLock);
+
+	for (FLiveLinkCollectionSourceItem& Source : Sources)
+	{
+		VisitorFunc(Source);
+	}
+}
+
+
+void FLiveLinkSourceCollection::ForEachSource(TFunctionRef<void(const FLiveLinkCollectionSourceItem&)> VisitorFunc) const
+{
+	UE::TUniqueLock Lock(SourcesLock);
+
+	for (const FLiveLinkCollectionSourceItem& Source : Sources)
+	{
+		VisitorFunc(Source);
+	}
+}
diff --git a/Engine/Plugins/Animation/LiveLink/Source/LiveLink/Private/LiveLinkSubject.cpp b/Engine/Plugins/Animation/LiveLink/Source/LiveLink/Private/LiveLinkSubject.cpp
--- a/Engine/Plugins/Animation/LiveLink/Source/LiveLink/Private/LiveLinkSubject.cpp
+++ b/Engine/Plugins/Animation/LiveLink/Source/LiveLink/Private/LiveLinkSubject.cpp
@@ -56,6 +56,8 @@ FLiveLinkSubject::~FLiveLinkSubject()
 
 void FLiveLinkSubject::Update()
 {
+	FScopeLock Lock(&SettingsCriticalSection);
+
 	// Clear all frames that are too old
 	if (FrameData.Num() > CachedSettings.BufferSettings.MaxNumberOfFrameToBuffered)
 	{
@@ -334,8 +336,6 @@ bool FLiveLinkSubject::HasStaticData() const
 
 void FLiveLinkSubject::AddFrameData(FLiveLinkFrameDataStruct&& InFrameData)
 {
-	check(IsInGameThread());
-
 	if (!ValidateFrameData(InFrameData))
 	{
 		return;
@@ -1017,8 +1017,6 @@ void FLiveLinkSubject::RemoveFrames(int32 InCount)
 
 void FLiveLinkSubject::SetStaticData(TSubclassOf<ULiveLinkRole> InRole, FLiveLinkStaticDataStruct&& InStaticData)
 {
-	check(IsInGameThread());
-
 	if (Role == nullptr)
 	{
 		static const FName NAME_NoRoleForSubject = "LiveLinkSubject_NoRoleForSubject";
@@ -1026,7 +1024,7 @@ void FLiveLinkSubject::SetStaticData(TSubclassOf<ULiveLinkRole> InRole, FLiveLin
 		return;
 	}
 
-	if(Role == InRole)
+	if (Role == InRole)
 	{
 		//Set initial blending processor to the role's default one. User will be able to modify it afterwards.
 		FrameData.Reset();
@@ -1042,7 +1040,7 @@ void FLiveLinkSubject::SetStaticData(TSubclassOf<ULiveLinkRole> InRole, FLiveLin
 
 void FLiveLinkSubject::CacheSettings(ULiveLinkSourceSettings* SourceSetting, ULiveLinkSubjectSettings* SubjectSetting)
 {
-	check(IsInGameThread());
+	FScopeLock Lock(&SettingsCriticalSection);
 
 	if (SourceSetting)
 	{
diff --git a/Engine/Plugins/Animation/LiveLink/Source/LiveLink/Private/LiveLinkTimedDataInput.cpp b/Engine/Plugins/Animation/LiveLink/Source/LiveLink/Private/LiveLinkTimedDataInput.cpp
--- a/Engine/Plugins/Animation/LiveLink/Source/LiveLink/Private/LiveLinkTimedDataInput.cpp
+++ b/Engine/Plugins/Animation/LiveLink/Source/LiveLink/Private/LiveLinkTimedDataInput.cpp
@@ -189,7 +189,7 @@ int32 FLiveLinkTimedDataInput::GetDataBufferSize() const
 
 void FLiveLinkTimedDataInput::SetDataBufferSize(int32 BufferSize)
 {
-	int32 NewSize = FMath::Clamp(BufferSize, LiveLinkTimedDataInput::CVarLiveLinkMinBufferSize.GetValueOnGameThread(), LiveLinkTimedDataInput::CVarLiveLinkMaxBufferSize.GetValueOnGameThread());
+	const int32 NewSize = FMath::Clamp(BufferSize, LiveLinkTimedDataInput::CVarLiveLinkMinBufferSize.GetValueOnAnyThread(), LiveLinkTimedDataInput::CVarLiveLinkMaxBufferSize.GetValueOnAnyThread());
 	if (ULiveLinkSourceSettings* Settings = LiveLinkClient->GetSourceSettings(Source))
 	{
 		Settings->BufferSettings.MaxNumberOfFrameToBuffered = NewSize;
@@ -205,7 +205,7 @@ const FSlateBrush* FLiveLinkTimedDataInput::GetDisplayIcon() const
 
 void FLiveLinkTimedDataInput::ProcessNewFrameTimingInfo(FLiveLinkBaseFrameData& NewFrameData)
 {
-	bool bUpdateContinuousClockOffset = LiveLinkTimedDataInput::CVarLiveLinkUpdateContinuousClockOffset.GetValueOnGameThread();
+	const bool bUpdateContinuousClockOffset = LiveLinkTimedDataInput::CVarLiveLinkUpdateContinuousClockOffset.GetValueOnAnyThread();
 
 	//Update both clock offsets for each frame received for our subjects
 	//We mark the last source/frame time that we used to update our offset to only update once per source frame
@@ -257,7 +257,7 @@ void FLiveLinkTimedDataInput::UpdateSmoothEngineTimeOffset(const FLiveLinkBaseFr
 	if (ULiveLinkSourceSettings* Settings = LiveLinkClient->GetSourceSettings(Source))
 	{
 		//Early-out if no frame offset is needed
-		const float NumFramesForSmoothOffset = LiveLinkTimedDataInput::CVarLiveLinkNumFramesForSmoothOffset.GetValueOnGameThread();
+		const float NumFramesForSmoothOffset = LiveLinkTimedDataInput::CVarLiveLinkNumFramesForSmoothOffset.GetValueOnAnyThread();
 
 		if (NumFramesForSmoothOffset <= 0)
 		{
diff --git a/Engine/Plugins/Animation/LiveLink/Source/LiveLink/Public/LiveLinkClient.h b/Engine/Plugins/Animation/LiveLink/Source/LiveLink/Public/LiveLinkClient.h
--- a/Engine/Plugins/Animation/LiveLink/Source/LiveLink/Public/LiveLinkClient.h
+++ b/Engine/Plugins/Animation/LiveLink/Source/LiveLink/Public/LiveLinkClient.h
@@ -192,6 +192,11 @@ protected:
 	virtual FLiveLinkSkeletonStaticData* GetSubjectAnimationStaticData_Deprecation(const FLiveLinkSubjectKey& SubjectKey) override;
 	//~ End FLiveLinkClient_Base_DEPRECATED implementation
 
+	/** Add delegates that will be triggered for all subjects. */
+	bool RegisterGlobalSubjectFramesDelegate(const FOnLiveLinkSubjectStaticDataAdded::FDelegate& InOnStaticDataAdded, const FOnLiveLinkSubjectFrameDataAdded::FDelegate& InOnFrameDataAdded, FDelegateHandle& OutStaticDataAddedHandle, FDelegateHandle& OutFrameDataAddedHandle);
+	/** Remove the delegates that were triggered for all subjects. */
+	void UnregisterGlobalSubjectFramesDelegate(FDelegateHandle& InStaticDataAddedHandle, FDelegateHandle& InFrameDataAddedHandle);
+
 private:
 	/** Struct that hold the pending static data that will be pushed next tick. */
 	struct FPendingSubjectStatic
@@ -253,8 +258,8 @@ protected:
 	/** LiveLink Provider for rebroadcasting */
 	TSharedPtr<ILiveLinkProvider> RebroadcastLiveLinkProvider;
 
-	/** Lock to stop multiple threads accessing the Subjects from the collection at the same time */
-    mutable FCriticalSection CollectionAccessCriticalSection;
+	/** Lock to protect access on SubjectFrameToPush and SubjectStaticToPush. */
+	mutable FCriticalSection PendingFramesCriticalSection;
 
 private:
 	/** Pending static info to add to a subject. */
@@ -293,6 +298,9 @@ private:
 	FString RebroadcastLiveLinkProviderName;
 	TSet<FLiveLinkSubjectKey> RebroadcastedSubjects;
 
+	/** Name token  used to register to all subject updates. */
+	const FName ALL_SUBJECTS_DELEGATE_TOKEN = "__Internal_AllSubjects_Update";
+
 #if WITH_EDITOR
 	/** Delegate when a subject is evaluated. */
 	FOnLiveLinkSubjectEvaluated OnLiveLinkSubjectEvaluatedDelegate;
@@ -302,12 +310,3 @@ private:
 	TOptional<FQualifiedFrameTime> CachedEngineFrameTime;
 #endif
 };
-
-#if UE_ENABLE_INCLUDE_ORDER_DEPRECATED_IN_5_2
-#include "ILiveLinkSource.h"
-#include "LiveLinkProvider.h"
-#include "LiveLinkSourceSettings.h"
-#include "LiveLinkVirtualSubject.h"
-#include "Tickable.h"
-#include "UObject/GCObject.h"
-#endif
diff --git a/Engine/Plugins/Animation/LiveLink/Source/LiveLink/Public/LiveLinkSourceCollection.h b/Engine/Plugins/Animation/LiveLink/Source/LiveLink/Public/LiveLinkSourceCollection.h
--- a/Engine/Plugins/Animation/LiveLink/Source/LiveLink/Public/LiveLinkSourceCollection.h
+++ b/Engine/Plugins/Animation/LiveLink/Source/LiveLink/Public/LiveLinkSourceCollection.h
@@ -2,12 +2,14 @@
 
 #pragma once
 
-
+#include "Async/Async.h"
+#include "Async/RecursiveMutex.h"
 #include "ILiveLinkClient.h"
 #include "LiveLinkSubject.h"
 #include "LiveLinkSubjectSettings.h"
 #include "LiveLinkVirtualSubject.h"
 #include "UObject/GCObject.h"
+#include "UObject/StrongObjectPtr.h"
 
 class FLiveLinkSourceCollection;
 
@@ -29,7 +31,7 @@ struct LIVELINK_API FLiveLinkCollectionSourceItem
 	FLiveLinkCollectionSourceItem& operator=(const FLiveLinkCollectionSourceItem&) = delete;
 
 	FGuid Guid;
-	TObjectPtr<ULiveLinkSourceSettings> Setting; // GC by FLiveLinkSourceCollection::AddReferencedObjects
+	TStrongObjectPtr<ULiveLinkSourceSettings> Setting;
 	TSharedPtr<ILiveLinkSource> Source;
 	TSharedPtr<FLiveLinkTimedDataInput> TimedData;
 	bool bPendingKill = false;
@@ -39,7 +41,6 @@ public:
 	bool IsVirtualSource() const;
 };
 
-
 struct LIVELINK_API FLiveLinkCollectionSubjectItem
 {
 	FLiveLinkCollectionSubjectItem(FLiveLinkSubjectKey InKey, TUniquePtr<FLiveLinkSubject> InLiveSubject, ULiveLinkSubjectSettings* InSettings, bool bInEnabled);
@@ -50,19 +51,21 @@ public:
 	bool bEnabled;
 	bool bPendingKill;
 
-	ILiveLinkSubject* GetSubject() { return VirtualSubject ? static_cast<ILiveLinkSubject*>(VirtualSubject) : static_cast<ILiveLinkSubject*>(LiveSubject.Get()); }
-	FLiveLinkSubject* GetLiveSubject() { return LiveSubject.Get(); }
-	ULiveLinkVirtualSubject* GetVirtualSubject() { return VirtualSubject; }
-	ILiveLinkSubject* GetSubject() const { return VirtualSubject ? VirtualSubject : static_cast<ILiveLinkSubject*>(LiveSubject.Get()); }
+	// Todo: These methods should be revisited because they may not be safe to access when LiveLinkHub is ticked outside of  the game thread.
+	// ie. Calling methods on a subject that is about to be removed will not keep the underlying livelink subject alive.
+	ILiveLinkSubject* GetSubject() { return VirtualSubject ? static_cast<ILiveLinkSubject*>(VirtualSubject.Get()) : static_cast<ILiveLinkSubject*>(LiveSubject.Get()); }
+	ULiveLinkVirtualSubject* GetVirtualSubject() { return VirtualSubject.Get(); }
+	ILiveLinkSubject* GetSubject() const { return VirtualSubject ? VirtualSubject.Get() : static_cast<ILiveLinkSubject*>(LiveSubject.Get()); }
+	ULiveLinkVirtualSubject* GetVirtualSubject() const { return VirtualSubject.Get(); }
+	UObject* GetSettings() const { return VirtualSubject ? static_cast<UObject*>(VirtualSubject.Get()) : static_cast<UObject*>(Setting.Get()); }
+	ULiveLinkSubjectSettings* GetLinkSettings() const { return Setting.Get(); }
 	FLiveLinkSubject* GetLiveSubject() const { return LiveSubject.Get(); }
-	ULiveLinkVirtualSubject* GetVirtualSubject() const { return VirtualSubject; }
-	UObject* GetSettings() const { return VirtualSubject ? static_cast<UObject*>(VirtualSubject) : static_cast<UObject*>(Setting); }
-	ULiveLinkSubjectSettings* GetLinkSettings() const { return Setting; }
+
 
 private:
-	TObjectPtr<ULiveLinkSubjectSettings> Setting; // GC by FLiveLinkSourceCollection::AddReferencedObjects
+	TStrongObjectPtr<ULiveLinkSubjectSettings> Setting;
 	TUniquePtr<FLiveLinkSubject> LiveSubject;
-	TObjectPtr<ULiveLinkVirtualSubject> VirtualSubject; // GC by FLiveLinkSourceCollection::AddReferencedObjects
+	TStrongObjectPtr<ULiveLinkVirtualSubject> VirtualSubject;
 
 public:
 	FLiveLinkCollectionSubjectItem(const FLiveLinkCollectionSubjectItem&) = delete;
@@ -74,7 +77,7 @@ public:
 };
 
 
-class LIVELINK_API FLiveLinkSourceCollection : public FGCObject
+class LIVELINK_API FLiveLinkSourceCollection
 {
 public:
 	// "source guid" for virtual subjects
@@ -84,19 +87,14 @@ public:
 	UE_NONCOPYABLE(FLiveLinkSourceCollection)
 
 public:
-	//~ Begin FGCObject implementation
-	virtual void AddReferencedObjects(FReferenceCollector & Collector) override;
-	virtual FString GetReferencerName() const override
-	{
-		return TEXT("FLiveLinkSourceCollection");
-	}
-	//~ End FGCObject implementation
 
-public:
-
-	
+	UE_DEPRECATED(5.5, "Use ForEachSource instead.")
 	TArray<FLiveLinkCollectionSourceItem>& GetSources() { return Sources; }
+
+	UE_DEPRECATED(5.5, "Use ForEachSource instead.")
 	const TArray<FLiveLinkCollectionSourceItem>& GetSources() const { return Sources; }
+
+	UE_DEPRECATED(5.5, "Use ForEachSubject instead.")
 	const TArray<FLiveLinkCollectionSubjectItem>& GetSubjects() const { return Subjects; }
 
 	void AddSource(FLiveLinkCollectionSourceItem Source);
@@ -108,12 +106,16 @@ public:
 	const FLiveLinkCollectionSourceItem* FindSource(FGuid SourceGuid) const;
 	FLiveLinkCollectionSourceItem* FindVirtualSource(FName VirtualSourceName);
 	const FLiveLinkCollectionSourceItem* FindVirtualSource(FName VirtualSourceName) const;
+	/** Get the number of sources in the collection. */
+	int32 NumSources() const;
 
 	void AddSubject(FLiveLinkCollectionSubjectItem Subject);
 	void RemoveSubject(FLiveLinkSubjectKey SubjectKey);
 	FLiveLinkCollectionSubjectItem* FindSubject(FLiveLinkSubjectKey SubjectKey);
 	const FLiveLinkCollectionSubjectItem* FindSubject(FLiveLinkSubjectKey SubjectKey) const;
 	const FLiveLinkCollectionSubjectItem* FindEnabledSubject(FLiveLinkSubjectName SubjectName) const;
+	/** Get the number of subjects in the collection. */
+	int32 NumSubjects() const;
 
 	bool IsSubjectEnabled(FLiveLinkSubjectKey SubjectKey) const;
 	void SetSubjectEnabled(FLiveLinkSubjectKey SubjectKey, bool bEnabled);
@@ -121,6 +123,26 @@ public:
 	void RemovePendingKill();
 	bool RequestShutdown();
 
+	/**
+	 * Thread safe way to apply a method over every subject.
+	 */
+	void ForEachSubject(TFunctionRef<void(FLiveLinkCollectionSourceItem& /*SourceItem*/, FLiveLinkCollectionSubjectItem& /*SubjectItem*/)> VisitorFunc);
+
+	/**
+	 * Thread safe way to apply a method over every subject.
+	 */
+	void ForEachSubject(TFunctionRef<void(const FLiveLinkCollectionSourceItem&, const FLiveLinkCollectionSubjectItem&)> VisitorFunc) const;
+
+	/**
+	 *  Thread safe way to apply a method over every source.
+	 */
+	void ForEachSource(TFunctionRef<void(FLiveLinkCollectionSourceItem& /*SourceItem*/)> VisitorFunc);
+
+	/**
+	 * Thread safe way to apply a method over every source.
+	 */
+	void ForEachSource(TFunctionRef<void(const FLiveLinkCollectionSourceItem&)> VisitorFunc) const;
+
 	FSimpleMulticastDelegate& OnLiveLinkSourcesChanged() { return OnLiveLinkSourcesChangedDelegate; }
 	FSimpleMulticastDelegate& OnLiveLinkSubjectsChanged() { return OnLiveLinkSubjectsChangedDelegate; }
 
@@ -129,8 +151,22 @@ public:
 	FOnLiveLinkSubjectChangedDelegate& OnLiveLinkSubjectAdded() { return OnLiveLinkSubjectAddedDelegate; }
 	FOnLiveLinkSubjectChangedDelegate& OnLiveLinkSubjectRemoved() { return OnLiveLinkSubjectRemovedDelegate; }
 
-private:
-	void RemoveSource(int32 Index);
+	/** Utility method used to broadcast delegates on the game thread if this function is called on a different thread. */
+	template <typename DelegateType, typename... ArgTypes>
+	void BroadcastOnGameThread(DelegateType& InDelegate, ArgTypes&&... InArgs)
+	{
+		if (IsInGameThread())
+		{
+			InDelegate.Broadcast(Forward<ArgTypes>(InArgs)...);
+		}
+		else
+		{
+			AsyncTask(ENamedThreads::GameThread, [&InDelegate, ...Args = InArgs] () mutable
+			{
+				InDelegate.Broadcast(MoveTemp(Args)...);
+			});
+		}
+	}
 
 private:
 	TArray<FLiveLinkCollectionSourceItem> Sources;
@@ -154,4 +190,10 @@ private:
 
 	/** Notify when a client subject's is removed */
 	FOnLiveLinkSubjectChangedDelegate OnLiveLinkSubjectRemovedDelegate;
+
+	/** Lock to stop multiple threads accessing the Subjects from the collection at the same time */
+	mutable UE::FRecursiveMutex SubjectsLock;
+
+	/** Lock to stop multiple threads accessing the Sources from the collection at the same time */
+	mutable UE::FRecursiveMutex SourcesLock;
 };
diff --git a/Engine/Plugins/Animation/LiveLink/Source/LiveLink/Public/LiveLinkSubject.h b/Engine/Plugins/Animation/LiveLink/Source/LiveLink/Public/LiveLinkSubject.h
--- a/Engine/Plugins/Animation/LiveLink/Source/LiveLink/Public/LiveLinkSubject.h
+++ b/Engine/Plugins/Animation/LiveLink/Source/LiveLink/Public/LiveLinkSubject.h
@@ -63,7 +63,11 @@ public:
 	virtual FLiveLinkStaticDataStruct& GetStaticData() override { return StaticData; }
 	virtual const FLiveLinkStaticDataStruct& GetStaticData() const override { return StaticData; }
 	virtual TArray<FLiveLinkTime> GetFrameTimes() const override;
-	virtual const TArray<ULiveLinkFrameTranslator::FWorkerSharedPtr> GetFrameTranslators() const override { return FrameTranslators; }
+	virtual const TArray<ULiveLinkFrameTranslator::FWorkerSharedPtr> GetFrameTranslators() const override
+	{
+		FScopeLock Lock(&SettingsCriticalSection);
+		return FrameTranslators;
+	}
 	virtual bool IsRebroadcasted() const override { return bRebroadcastSubject; }
 	virtual bool HasStaticDataBeenRebroadcasted() const override { return bRebroadcastStaticDataSent; }
 	virtual void SetStaticDataAsRebroadcasted(const bool bInSent) override { bRebroadcastStaticDataSent = bInSent; }
@@ -232,5 +236,9 @@ private:
 	 * Evaluation can be done on any thread so we need to protect statistic logging 
 	 * Some stats requires more than atomic sized vars so a critical section is used to protect when necessary
 	 */
+
 	mutable FCriticalSection StatisticCriticalSection;
+
+	/** Used to protect access to translators, preprocessors and interpolation processors since they can be set and accessed in different threads. */
+	mutable FCriticalSection SettingsCriticalSection;
 };
diff --git a/Engine/Plugins/Animation/LiveLinkHub/Source/LiveLinkHub/Private/LiveLinkHubApplication.h b/Engine/Plugins/Animation/LiveLinkHub/Source/LiveLinkHub/Private/LiveLinkHubApplication.h
--- a/Engine/Plugins/Animation/LiveLinkHub/Source/LiveLinkHub/Private/LiveLinkHubApplication.h
+++ b/Engine/Plugins/Animation/LiveLinkHub/Source/LiveLinkHub/Private/LiveLinkHubApplication.h
@@ -9,6 +9,8 @@
 #include "Framework/Application/SlateApplication.h"
 #include "LiveLinkHub.h"
 #include "LiveLinkHubInputProcessor.h"
+#include "LiveLinkHubSettings.h"
+#include "LiveLinkHubTicker.h"
 #include "Misc/App.h"
 #include "Misc/CoreDelegates.h"
 #include "Runtime/Launch/Resources/Version.h"
@@ -25,15 +27,23 @@ void LiveLinkHubLoop(const TSharedPtr<FLiveLinkHub>& LiveLinkHub)
 	// Disable throttling for the hub
 	GetMutableDefault<UEditorPerformanceSettings>()->bThrottleCPUWhenNotForeground = false;
 
+	const bool bTickOnGameThread = GetDefault<ULiveLinkHubSettings>()->bTickOnGameThread;
+
 	check(FSlateApplication::IsInitialized());
 	FSlateApplication::Get().RegisterInputPreProcessor(MakeShared<FLiveLinkHubInputProcessor>());
 	{
 		UE_LOG(LogLiveLinkHubApplication, Display, TEXT("LiveLinkHub Initialized (Version: %d.%d)"), ENGINE_MAJOR_VERSION, ENGINE_MINOR_VERSION);
 
+		FLiveLinkHubTicker Ticker{LiveLinkHub.ToSharedRef()};
+
+		if (!bTickOnGameThread)
+		{
+			Ticker.StartTick();
+		}
+
 		double LastTime = FPlatformTime::Seconds();
 
-		constexpr double IdealFrameRate = 60.0;
-		constexpr float IdealFrameTime = 1.0f / IdealFrameRate;
+		const double IdealFrameTime = 1 / GetDefault<ULiveLinkHubSettings>()->TargetFrameRate;
 
 		while (!IsEngineExitRequested())
 		{
@@ -48,7 +58,10 @@ void LiveLinkHubLoop(const TSharedPtr<FLiveLinkHub>& LiveLinkHub)
 			FSlateApplication::Get().PollGameDeviceState();
 
 			// This is normally ticked by OnSamplingInput.
-			LiveLinkHub->Tick();
+			if (bTickOnGameThread)
+			{
+				LiveLinkHub->Tick();
+			}
 
 			// Run garbage collection for the UObjects for the rest of the frame or at least to 2 ms
 			IncrementalPurgeGarbage(true, FMath::Max<float>(0.002f, IdealFrameTime - (FPlatformTime::Seconds() - LastTime)));
@@ -65,5 +78,8 @@ void LiveLinkHubLoop(const TSharedPtr<FLiveLinkHub>& LiveLinkHub)
 
 			LastTime = CurrentTime;
 		}
+
+		Ticker.Exit();
+		Ticker.Stop();
 	}
 }
diff --git a/Engine/Plugins/Animation/LiveLinkHub/Source/LiveLinkHub/Private/LiveLinkHubClient.cpp b/Engine/Plugins/Animation/LiveLinkHub/Source/LiveLinkHub/Private/LiveLinkHubClient.cpp
--- a/Engine/Plugins/Animation/LiveLinkHub/Source/LiveLinkHub/Private/LiveLinkHubClient.cpp
+++ b/Engine/Plugins/Animation/LiveLinkHub/Source/LiveLinkHub/Private/LiveLinkHubClient.cpp
@@ -34,17 +34,17 @@ DECLARE_CYCLE_STAT(TEXT("LiveLinkHub - Push FrameData"), STAT_LiveLinkHub_PushFr
 FLiveLinkHubClient::FLiveLinkHubClient(TSharedPtr<ILiveLinkHub> InLiveLinkHub)
 	: LiveLinkHub(MoveTemp(InLiveLinkHub))
 {
+	RegisterGlobalSubjectFramesDelegate(FOnLiveLinkSubjectStaticDataAdded::FDelegate::CreateRaw(this, &FLiveLinkHubClient::OnStaticDataAdded), FOnLiveLinkSubjectFrameDataAdded::FDelegate::CreateRaw(this, &FLiveLinkHubClient::OnFrameDataAdded), StaticDataAddedHandle, FrameDataAddedHandle);
 }
 
 FLiveLinkHubClient::~FLiveLinkHubClient()
 {
+	UnregisterGlobalSubjectFramesDelegate(StaticDataAddedHandle, FrameDataAddedHandle);
 }
 
 void FLiveLinkHubClient::CacheSubjectSettings(const FLiveLinkSubjectKey& SubjectKey, ULiveLinkSubjectSettings* Settings) const
 {
-	FScopeLock Lock(&CollectionAccessCriticalSection);
-
-	if (FLiveLinkCollectionSubjectItem* SubjectItem = Collection->FindSubject(SubjectKey))
+	if (const FLiveLinkCollectionSubjectItem* SubjectItem = Collection->FindSubject(SubjectKey))
 	{
 		ULiveLinkSourceSettings* SourceSettings = GetSourceSettings(SubjectKey.Source);
 		ULiveLinkSubjectSettings* SubjectSettings = Settings;
@@ -90,13 +90,10 @@ bool FLiveLinkHubClient::CreatePlaybackSource(const FLiveLinkSourcePreset& InSou
 		return false;
 	}
 
+	if (Collection->FindSource(InSourcePreset.Guid) != nullptr)
 	{
-		FScopeLock Lock(&CollectionAccessCriticalSection);
-		if (Collection->FindSource(InSourcePreset.Guid) != nullptr)
-		{
-			UE_LOG(LogLiveLinkHub, Warning, TEXT("Create Source Failure: The guid already exist."));
-			return false;
-		}
+		UE_LOG(LogLiveLinkHub, Warning, TEXT("Create Source Failure: The guid already exist."));
+		return false;
 	}
 
 	ULiveLinkSourceSettings* Setting = nullptr;
@@ -124,17 +121,16 @@ bool FLiveLinkHubClient::CreatePlaybackSource(const FLiveLinkSourcePreset& InSou
 	{
 		Setting = NewObject<ULiveLinkSourceSettings>(GetTransientPackage(), SourceSettingsClass);
 		UEngine::CopyPropertiesForUnrelatedObjects(InSourcePreset.Settings, Setting);
-		Data.Setting = Setting;
+		Data.Setting = TStrongObjectPtr(Setting);
 	}
 	else
 	{
-		Setting = Data.Setting = DuplicateObject<ULiveLinkSourceSettings>(InSourcePreset.Settings, GetTransientPackage());
+		Data.Setting = TStrongObjectPtr(DuplicateObject<ULiveLinkSourceSettings>(InSourcePreset.Settings, GetTransientPackage()));
+		Setting = Data.Setting.Get();
 	}
 
-	{
-		FScopeLock Lock(&CollectionAccessCriticalSection);
-		Collection->AddSource(MoveTemp(Data));
-	}
+	Collection->AddSource(MoveTemp(Data));
+
 	Source->ReceiveClient(this, InSourcePreset.Guid);
 	Source->InitializeSettings(Setting);
 
@@ -188,7 +184,6 @@ bool FLiveLinkHubClient::CreatePlaybackSubject(const FLiveLinkSubjectPreset& InS
 	{
 		if (SubjectItem->bPendingKill)
 		{
-			FScopeLock Lock(&CollectionAccessCriticalSection);
 			Collection->RemoveSubject(InSubjectPreset.Key);
 		}
 		else
@@ -205,7 +200,6 @@ bool FLiveLinkHubClient::CreatePlaybackSubject(const FLiveLinkSubjectPreset& InS
 		FLiveLinkCollectionSubjectItem VSubjectData(InSubjectPreset.Key, VSubject, bEnabled);
 		VSubject->Initialize(VSubjectData.Key, VSubject->GetRole(), this);
 
-		FScopeLock Lock(&CollectionAccessCriticalSection);
 		Collection->AddSubject(MoveTemp(VSubjectData));
 		Collection->SetSubjectEnabled(InSubjectPreset.Key, InSubjectPreset.bEnabled);
 	}
@@ -225,193 +219,12 @@ bool FLiveLinkHubClient::CreatePlaybackSubject(const FLiveLinkSubjectPreset& InS
 		FLiveLinkCollectionSubjectItem CollectionSubjectItem(InSubjectPreset.Key, MakeUnique<FLiveLinkPlaybackSubject>(SourceItem->TimedData), SubjectSettings, bEnabled);
 		CollectionSubjectItem.GetLiveSubject()->Initialize(InSubjectPreset.Key, InSubjectPreset.Role.Get(), this);
 
-		FScopeLock Lock(&CollectionAccessCriticalSection);
 		Collection->AddSubject(MoveTemp(CollectionSubjectItem));
 		Collection->SetSubjectEnabled(InSubjectPreset.Key, InSubjectPreset.bEnabled);
 	}
 	return true;
 }
 
-void FLiveLinkHubClient::PushSubjectStaticData_AnyThread(const FLiveLinkSubjectKey& SubjectKey, TSubclassOf<ULiveLinkRole> Role, FLiveLinkStaticDataStruct&& InStaticData)
-{
-	// Copied and adapted from FLiveLinkClient::PushSubjectStaticData_AnyThread
-	SCOPE_CYCLE_COUNTER(STAT_LiveLinkHub_PushStaticData);
-
-	check(Collection);
-
-	// Check if subject should receive data
-
-	if (!FLiveLinkRoleTrait::Validate(Role, InStaticData))
-	{
-		if (Role == nullptr)
-		{
-			UE_LOG(LogLiveLinkHub, Error, TEXT("Trying to add unsupported static data type with subject '%s'."), *SubjectKey.SubjectName.ToString());
-		}
-		else
-		{
-			UE_LOG(LogLiveLinkHub, Error, TEXT("Trying to add unsupported static data type to role '%s' with subject '%s'."), *Role->GetName(), *SubjectKey.SubjectName.ToString());
-		}
-		return;
-	}
-
-	bool bShouldLogIfInvalidStaticData = true;
-	if (!Role.GetDefaultObject()->IsStaticDataValid(InStaticData, bShouldLogIfInvalidStaticData))
-	{
-		if (bShouldLogIfInvalidStaticData)
-		{
-			UE_LOG(LogLiveLinkHub, Error, TEXT("Trying to add static data that is not formatted properly to role '%s' with subject '%s'."), *Role->GetName(), *SubjectKey.SubjectName.ToString());
-		}
-		return;
-	}
-
-	const FLiveLinkCollectionSourceItem* SourceItem = nullptr;
-	{
-		FScopeLock Lock(&CollectionAccessCriticalSection);
-		SourceItem = Collection->FindSource(SubjectKey.Source);
-	}
-
-	if (SourceItem == nullptr || SourceItem->bPendingKill)
-	{
-		return;
-	}
-
-	if (SourceItem->IsVirtualSource())
-	{
-		UE_LOG(LogLiveLinkHub, Error, TEXT("Trying to add frame data to the subject '%s'."), *Role->GetName(), *SubjectKey.SubjectName.ToString());
-		return;
-	}
-
-	FLiveLinkSubject* LiveLinkSubject = nullptr;
-	{
-		FScopeLock Lock(&CollectionAccessCriticalSection);
-		if (FLiveLinkCollectionSubjectItem* SubjectItem = Collection->FindSubject(SubjectKey))
-		{
-			if (!SubjectItem->bPendingKill)
-			{
-				LiveLinkSubject = SubjectItem->GetLiveSubject();
-
-				if (LiveLinkSubject->GetRole() != Role)
-				{
-					UE_LOG(LogLiveLinkHub, Warning, TEXT("Subject '%s' of role '%s' is changing its role to '%s'. Current subject will be removed and a new one will be created"), *SubjectKey.SubjectName.ToString(), *LiveLinkSubject->GetRole().GetDefaultObject()->GetDisplayName().ToString(), *Role.GetDefaultObject()->GetDisplayName().ToString());
-
-					Collection->RemoveSubject(SubjectKey);
-					LiveLinkSubject = nullptr;
-				}
-				else
-				{
-					LiveLinkSubject->ClearFrames();
-				}
-			}
-		}
-
-	}
-	if (LiveLinkSubject == nullptr)
-	{
-		const ULiveLinkSettings* LiveLinkSettings = GetDefault<ULiveLinkSettings>();
-		const FLiveLinkRoleProjectSetting DefaultSetting = LiveLinkSettings->GetDefaultSettingForRole(Role.Get());
-
-		// Setting class should always be valid
-		ULiveLinkSubjectSettings* SubjectSettings = nullptr;
-		{
-			UClass* SettingClass = DefaultSetting.SettingClass.Get();
-			if (SettingClass == nullptr)
-			{
-				SettingClass = ULiveLinkHubSubjectSettings::StaticClass();
-			}
-
-			SubjectSettings = NewObject<ULiveLinkSubjectSettings>(GetTransientPackage(), SettingClass);
-
-			if (SubjectSettings->IsA<ULiveLinkHubSubjectSettings>())
-			{
-				Cast<ULiveLinkHubSubjectSettings>(SubjectSettings)->Initialize(SubjectKey, Role.Get(), this);
-			}
-
-			SubjectSettings->Role = Role;
-
-			UClass* FrameInterpolationProcessorClass = DefaultSetting.FrameInterpolationProcessor.Get();
-			if (FrameInterpolationProcessorClass != nullptr)
-			{
-				UClass* InterpolationRole = FrameInterpolationProcessorClass->GetDefaultObject<ULiveLinkFrameInterpolationProcessor>()->GetRole();
-				if (Role->IsChildOf(InterpolationRole))
-				{
-					SubjectSettings->InterpolationProcessor = NewObject<ULiveLinkFrameInterpolationProcessor>(SubjectSettings, FrameInterpolationProcessorClass);
-				}
-				else
-				{
-					UE_LOG(LogLiveLinkHub, Warning, TEXT("The interpolator '%s' is not valid for the Role '%s'"), *FrameInterpolationProcessorClass->GetName(), *Role->GetName());
-				}
-			}
-			else
-			{
-				// If no settings were found for a specific role, check if the default interpolator is compatible with the role
-				UClass* FallbackInterpolationProcessorClass = LiveLinkSettings->FrameInterpolationProcessor.Get();
-				if (FallbackInterpolationProcessorClass != nullptr)
-				{
-					UClass* InterpolationRole = FallbackInterpolationProcessorClass->GetDefaultObject<ULiveLinkFrameInterpolationProcessor>()->GetRole();
-					if (Role->IsChildOf(InterpolationRole))
-					{
-						SubjectSettings->InterpolationProcessor = NewObject<ULiveLinkFrameInterpolationProcessor>(SubjectSettings, FallbackInterpolationProcessorClass);
-					}
-				}
-			}
-
-			for (TSubclassOf<ULiveLinkFramePreProcessor> PreProcessor : DefaultSetting.FramePreProcessors)
-			{
-				if (PreProcessor != nullptr)
-				{
-					UClass* PreProcessorRole = PreProcessor->GetDefaultObject<ULiveLinkFramePreProcessor>()->GetRole();
-					if (Role->IsChildOf(PreProcessorRole))
-					{
-						SubjectSettings->PreProcessors.Add(NewObject<ULiveLinkFramePreProcessor>(SubjectSettings, PreProcessor.Get()));
-					}
-					else
-					{
-						UE_LOG(LogLiveLinkHub, Warning, TEXT("The pre processor '%s' is not valid for the Role '%s'"), *PreProcessor->GetName(), *Role->GetName());
-					}
-				}
-			}
-		}
-
-
-		{
-			FScopeLock Lock(&CollectionAccessCriticalSection);
-
-			bool bEnabled = Collection->FindEnabledSubject(SubjectKey.SubjectName) == nullptr;
-			FLiveLinkCollectionSubjectItem CollectionSubjectItem(SubjectKey, MakeUnique<FLiveLinkSubject>(SourceItem->TimedData), SubjectSettings, bEnabled);
-			
-			CollectionSubjectItem.GetLiveSubject()->Initialize(SubjectKey, Role.Get(), this);
-
-			LiveLinkSubject = CollectionSubjectItem.GetLiveSubject();
-
-			Collection->AddSubject(MoveTemp(CollectionSubjectItem));
-			
-			// Clear async flag on the newly created uobject to allow it to be garbage collected.
-			SubjectSettings->AtomicallyClearInternalFlags(EInternalObjectFlags::Async);
-		}
-	}
-
-	BroadcastStaticDataUpdate(LiveLinkSubject, Role, InStaticData);
-
-	if (LiveLinkSubject)
-	{
-		// Dispatch it on the game thread since FLiveLinkSubject::SetStaticData asserts when called outside the game thread.
-		// Note that SetStaticData may not be needed on the hub but we might need it for interpolation or preprocessing in the future.
-		FFunctionGraphTask::CreateAndDispatchWhenReady([this, LiveLinkSubject, SubjectKey, Role, StaticData = MoveTemp(InStaticData)]() mutable
-		{
-			FScopeLock Lock(&CollectionAccessCriticalSection);
-			// Validate live link subject hasn't changed
-			if (FLiveLinkCollectionSubjectItem* SubjectItem = Collection->FindSubject(SubjectKey))
-			{
-				const FLiveLinkSubject* CurrentLiveLinkSubject = SubjectItem->GetLiveSubject();
-				if (LiveLinkSubject == CurrentLiveLinkSubject)
-				{
-    				LiveLinkSubject->SetStaticData(Role, MoveTemp(StaticData));
-				}
-			}
-		}, TStatId(), nullptr, ENamedThreads::GameThread);
-	}
-}
-
 void FLiveLinkHubClient::PushSubjectFrameData_AnyThread(const FLiveLinkSubjectKey& SubjectKey, FLiveLinkFrameDataStruct&& FrameData)
 {
 	SCOPE_CYCLE_COUNTER(STAT_LiveLinkHub_PushFrameData);
@@ -420,7 +233,6 @@ void FLiveLinkHubClient::PushSubjectFrameData_AnyThread(const FLiveLinkSubjectKe
 
 	bool bFrameValid = true;
 	{
-		FScopeLock Lock(&CollectionAccessCriticalSection);
 		if (const FLiveLinkCollectionSubjectItem* SubjectItem = Collection->FindSubject(SubjectKey))
 		{
 			if (FLiveLinkSubject* LiveSubject = SubjectItem->GetLiveSubject())
@@ -477,7 +289,6 @@ FText FLiveLinkHubClient::GetSourceStatus(FGuid InEntryGuid) const
 
 bool FLiveLinkHubClient::IsSubjectValid(const FLiveLinkSubjectKey& InSubjectKey) const
 {
-	FScopeLock Lock(&CollectionAccessCriticalSection);
 	if (const FLiveLinkCollectionSubjectItem* SubjectItem = Collection->FindSubject(InSubjectKey))
 	{
 		if (FLiveLinkSubject* LiveSubject = SubjectItem->GetLiveSubject())
@@ -512,10 +323,13 @@ void FLiveLinkHubClient::RemoveVirtualSubject(const FLiveLinkSubjectKey& Virtual
 {
 	FLiveLinkClient::RemoveVirtualSubject(VirtualSubjectKey);
 
+	bool bAnyVirtualSubject = false;
+	Collection->ForEachSubject([&bAnyVirtualSubject](const FLiveLinkCollectionSourceItem& SourceItem, const FLiveLinkCollectionSubjectItem& SubjectItem)
 	{
-		FScopeLock Lock(&CollectionAccessCriticalSection);
-		bVirtualSubjectsPresent = Algo::AnyOf(Collection->GetSubjects(), [] (const FLiveLinkCollectionSubjectItem& SubjectItem) { return !!SubjectItem.GetVirtualSubject(); });
-	}
+		bAnyVirtualSubject = bAnyVirtualSubject || !!SubjectItem.GetVirtualSubject();;
+	});
+
+	bVirtualSubjectsPresent = bAnyVirtualSubject;
 }
 
 TSharedPtr<ILiveLinkProvider> FLiveLinkHubClient::GetRebroadcastLiveLinkProvider() const
@@ -551,4 +365,18 @@ void FLiveLinkHubClient::BroadcastStaticDataUpdate(FLiveLinkSubject* InLiveSubje
 	}
 }
 
+void FLiveLinkHubClient::OnStaticDataAdded(FLiveLinkSubjectKey SubjectKey,  TSubclassOf<ULiveLinkRole> SubjectRole, const FLiveLinkStaticDataStruct& InStaticData)
+{
+	// Broadcast the subjects to connected clients.
+	if (FLiveLinkCollectionSubjectItem* SubjectItem = Collection->FindSubject(SubjectKey))
+	{
+		BroadcastStaticDataUpdate(SubjectItem->GetLiveSubject(), SubjectRole, InStaticData);
+	}
+}
+
+void FLiveLinkHubClient::OnFrameDataAdded(FLiveLinkSubjectKey InSubjectKey, TSubclassOf<ULiveLinkRole> SubjectRole, const FLiveLinkFrameDataStruct& InFrameData)
+{
+	// No-op, to be used in subsequent CL.
+}
+
 #undef LOCTEXT_NAMESPACE
diff --git a/Engine/Plugins/Animation/LiveLinkHub/Source/LiveLinkHub/Private/LiveLinkHubClient.h b/Engine/Plugins/Animation/LiveLinkHub/Source/LiveLinkHub/Private/LiveLinkHubClient.h
--- a/Engine/Plugins/Animation/LiveLinkHub/Source/LiveLinkHub/Private/LiveLinkHubClient.h
+++ b/Engine/Plugins/Animation/LiveLinkHub/Source/LiveLinkHub/Private/LiveLinkHubClient.h
@@ -53,7 +53,6 @@ public:
 	//~ Begin ILiveLinkClient interface
 	virtual bool CreateSource(const FLiveLinkSourcePreset& InSourcePreset) override;
 	virtual bool CreateSubject(const FLiveLinkSubjectPreset& InSubjectPreset) override;
-	virtual void PushSubjectStaticData_AnyThread(const FLiveLinkSubjectKey& SubjectKey, TSubclassOf<ULiveLinkRole> Role, FLiveLinkStaticDataStruct&& InStaticData) override;
 	virtual void PushSubjectFrameData_AnyThread(const FLiveLinkSubjectKey& SubjectKey, FLiveLinkFrameDataStruct&& FrameData) override;
 	virtual FText GetSourceStatus(FGuid InEntryGuid) const override;
 	virtual bool IsSubjectValid(const FLiveLinkSubjectKey& InSubjectKey) const override;
@@ -74,6 +73,10 @@ private:
 	/** Broadcast a static data update to this client's listeners. */
 	void BroadcastStaticDataUpdate(FLiveLinkSubject* InLiveSubject, TSubclassOf<ULiveLinkRole> InRole, const FLiveLinkStaticDataStruct& InStaticData) const;
 
+	//~ Delegates called by LiveLinkClient
+	void OnStaticDataAdded(FLiveLinkSubjectKey SubjectKey, TSubclassOf<ULiveLinkRole> SubjectRole, const FLiveLinkStaticDataStruct& InStaticData);
+	void OnFrameDataAdded(FLiveLinkSubjectKey InSubjectKey, TSubclassOf<ULiveLinkRole> SubjectRole, const FLiveLinkFrameDataStruct& InFrameData);
+
 private:
 	/** Weak pointer to the live link hub. */
 	TWeakPtr<ILiveLinkHub> LiveLinkHub;
@@ -85,4 +88,8 @@ private:
 	FOnSubjectMarkedPendingKill_AnyThread OnSubjectMarkedPendingKillDelegate_AnyThread;
 	/** Whether there are virtual subjects at the moment. Used to determine if we should cache frame data for their usage. */
 	std::atomic<bool> bVirtualSubjectsPresent = false;
+
+	//~ Delegate handles given by FLiveLinkClient.
+	FDelegateHandle StaticDataAddedHandle;
+	FDelegateHandle FrameDataAddedHandle;
 };
diff --git a/Engine/Plugins/Animation/LiveLinkHub/Source/LiveLinkHub/Private/LiveLinkHubSettings.h b/Engine/Plugins/Animation/LiveLinkHub/Source/LiveLinkHub/Private/LiveLinkHubSettings.h
--- a/Engine/Plugins/Animation/LiveLinkHub/Source/LiveLinkHub/Private/LiveLinkHubSettings.h
+++ b/Engine/Plugins/Animation/LiveLinkHub/Source/LiveLinkHub/Private/LiveLinkHubSettings.h
@@ -4,6 +4,8 @@
 
 #include "UObject/Object.h"
 
+#include "Misc/FrameRate.h"
+
 #include "LiveLinkHubSettings.generated.h"
 
 /**
@@ -26,4 +28,16 @@ public:
 	/** Which project settings sections to display when opening the settings viewer. */
 	UPROPERTY(config)
 	TArray<FName> ProjectSettingsToDisplay;
+
+	/**
+	 * - Experimental - If this is disabled, LiveLinkHub's LiveLink Client will tick outside of the game thread.
+	 * This allows processing LiveLink frame snapshots without the risk of being blocked by the game / ui thread.
+	 * Note that this should only be relevant for virtual subjects since data is already forwarded to UE outside of the game thread.
+	 */
+	UPROPERTY(config, EditAnywhere, Category = "LiveLinkHub", meta = (ConfigRestartRequired = true))
+	bool bTickOnGameThread = false;
+
+	/** Target framerate for ticking LiveLinkHub. */
+	UPROPERTY(config, EditAnywhere, Category="LiveLinkHub", meta = (ConfigRestartRequired = true, ClampMin="15.0"))
+	float TargetFrameRate = 120.0f;
 };
diff --git a/Engine/Plugins/Animation/LiveLinkHub/Source/LiveLinkHub/Private/LiveLinkHubSubjectSettings.cpp b/Engine/Plugins/Animation/LiveLinkHub/Source/LiveLinkHub/Private/LiveLinkHubSubjectSettings.cpp
--- a/Engine/Plugins/Animation/LiveLinkHub/Source/LiveLinkHub/Private/LiveLinkHubSubjectSettings.cpp
+++ b/Engine/Plugins/Animation/LiveLinkHub/Source/LiveLinkHub/Private/LiveLinkHubSubjectSettings.cpp
@@ -7,14 +7,15 @@
 #include "Features/IModularFeatures.h"
 #include "Clients/LiveLinkHubProvider.h"
 
-void ULiveLinkHubSubjectSettings::Initialize(FLiveLinkSubjectKey InSubjectKey, TSubclassOf<ULiveLinkRole> InRole, FLiveLinkHubClient* InLiveLinkClient)
+void ULiveLinkHubSubjectSettings::Initialize(FLiveLinkSubjectKey InSubjectKey)
 {
+	ILiveLinkClient* LiveLinkClient = &IModularFeatures::Get().GetModularFeature<ILiveLinkClient>(ILiveLinkClient::ModularFeatureName);
 	SubjectName = InSubjectKey.SubjectName.ToString();
 	SubjectKey = InSubjectKey;
 
 	OutboundName = SubjectName;
 
-	Source = InLiveLinkClient->GetSourceType(InSubjectKey.Source).ToString();
+	Source = LiveLinkClient->GetSourceType(InSubjectKey.Source).ToString();
 }
 
 void ULiveLinkHubSubjectSettings::NotifyRename()
diff --git a/Engine/Plugins/Animation/LiveLinkHub/Source/LiveLinkHub/Private/LiveLinkHubSubjectSettings.h b/Engine/Plugins/Animation/LiveLinkHub/Source/LiveLinkHub/Private/LiveLinkHubSubjectSettings.h
--- a/Engine/Plugins/Animation/LiveLinkHub/Source/LiveLinkHub/Private/LiveLinkHubSubjectSettings.h
+++ b/Engine/Plugins/Animation/LiveLinkHub/Source/LiveLinkHub/Private/LiveLinkHubSubjectSettings.h
@@ -15,8 +15,9 @@ class ULiveLinkHubSubjectSettings : public ULiveLinkSubjectSettings
 {
 	GENERATED_BODY()
 public:
-	/** Fetch information from the livelink client to initialize this settings object. */
-	void Initialize(FLiveLinkSubjectKey InSubjectKey, TSubclassOf<ULiveLinkRole> InRole, FLiveLinkHubClient* InLiveLinkClient);
+
+	/** Initialize the settings. */
+	virtual void Initialize(FLiveLinkSubjectKey InSubjectKey) override;
 
 	/** Notify connected clients that this subject's name has changed. */
 	void NotifyRename();
diff --git a/Engine/Plugins/Animation/LiveLinkHub/Source/LiveLinkHub/Private/LiveLinkHubTicker.h b/Engine/Plugins/Animation/LiveLinkHub/Source/LiveLinkHub/Private/LiveLinkHubTicker.h
new file mode 100644
--- /dev/null
+++ b/Engine/Plugins/Animation/LiveLinkHub/Source/LiveLinkHub/Private/LiveLinkHubTicker.h
@@ -0,0 +1,79 @@
+// Copyright Epic Games, Inc. All Rights Reserved.
+
+#pragma once
+
+#include "HAL/Runnable.h"
+#include "HAL/RunnableThread.h"
+#include "LiveLinkHub.h"
+#include "LiveLinkHubSettings.h"
+#include "Misc/Timespan.h"
+
+/** Object used to tick LiveLinkHub outside of the game thread. */
+class FLiveLinkHubTicker : public FRunnable
+{
+public:
+	explicit FLiveLinkHubTicker(TSharedRef<FLiveLinkHub> InLiveLinkHub)
+		: LiveLinkHub(MoveTemp(InLiveLinkHub))
+	{
+	}
+
+	void StartTick()
+	{
+		if (!bIsRunning)
+		{
+			bIsRunning = true;
+			TickEvent = FGenericPlatformProcess::GetSynchEventFromPool();
+			check(TickEvent);
+			Thread.Reset(FRunnableThread::Create(this, TEXT("LiveLinkHubTicker")));
+		}
+	}
+
+	//~ Begin FRunnable Interface
+	virtual uint32 Run() override
+	{
+		float TickFrequency = 1 / GetDefault<ULiveLinkHubSettings>()->TargetFrameRate;
+		FTimespan TickTimeSpan = FTimespan::FromSeconds(TickFrequency);
+
+		while (bIsRunning)
+		{
+			check(TickEvent);
+			TickEvent->Wait(TickTimeSpan);
+
+			if (bIsRunning) // make sure we were not told to exit during the wait
+			{
+				TRACE_CPUPROFILER_EVENT_SCOPE(FLiveLinkHubTicker::Tick);
+				LiveLinkHub->Tick();
+			}
+		}
+
+		return 0;
+	}
+	virtual void Exit() override
+	{
+		if (!bIsRunning)
+		{
+			return;
+		}
+
+		bIsRunning = false;
+
+		if (TickEvent)
+		{
+			TickEvent->Trigger();
+		}
+
+		Thread->WaitForCompletion();
+
+		if (TickEvent)
+		{
+			FGenericPlatformProcess::ReturnSynchEventToPool(TickEvent);
+			TickEvent = nullptr;
+		}
+	}
+	//~ End FRunnable Interface
+	std::atomic<bool> bIsRunning = false;
+
+	FEvent* TickEvent = nullptr;
+	TUniquePtr<FRunnableThread> Thread;
+	TSharedPtr<FLiveLinkHub> LiveLinkHub;
+};
diff --git a/Engine/Programs/LiveLinkHub/Config/DefaultGame.ini b/Engine/Programs/LiveLinkHub/Config/DefaultGame.ini
--- a/Engine/Programs/LiveLinkHub/Config/DefaultGame.ini
+++ b/Engine/Programs/LiveLinkHub/Config/DefaultGame.ini
@@ -1,6 +1,4 @@
 [/Script/LiveLink.LiveLinkSettings]
--DefaultRoleSettings=(Role=Class'/Script/LiveLink.LiveLinkAnimationRole', FrameInterpolationProcessor=Class'/Script/LiveLink.LiveLinkAnimationFrameInterpolationProcessor')
-+DefaultRoleSettings=(Role="/Script/CoreUObject.Class'/Script/LiveLinkInterface.LiveLinkAnimationRole'",SettingClass="/Script/CoreUObject.Class'/Script/LiveLinkHub.LiveLinkHubSubjectSettings'",FrameInterpolationProcessor="/Script/CoreUObject.Class'/Script/LiveLink.LiveLinkAnimationFrameInterpolationProcessor'",FramePreProcessors=)
 DefaultSettingsClass=/Script/LiveLinkHub.LiveLinkHubSubjectSettings
 FrameInterpolationProcessor=/Script/CoreUObject.Class'/Script/LiveLink.LiveLinkBasicFrameInterpolationProcessor'
 DefaultLiveLinkPreset=/Game/NewPreset.NewPreset
diff --git a/Engine/Source/Runtime/LiveLinkInterface/Public/LiveLinkSubjectSettings.h b/Engine/Source/Runtime/LiveLinkInterface/Public/LiveLinkSubjectSettings.h
--- a/Engine/Source/Runtime/LiveLinkInterface/Public/LiveLinkSubjectSettings.h
+++ b/Engine/Source/Runtime/LiveLinkInterface/Public/LiveLinkSubjectSettings.h
@@ -27,6 +27,9 @@ class ULiveLinkSubjectSettings : public UObject
 public:
 	GENERATED_BODY()
 
+	/** Initialize the settings. */
+	virtual void Initialize(FLiveLinkSubjectKey InSubjectKey) {}
+
 	/** List of available preprocessor the subject will use. */
 	UPROPERTY(EditAnywhere, Instanced, Category = "LiveLink", meta = (DisplayName = "Pre Processors"))
 	TArray<TObjectPtr<ULiveLinkFramePreProcessor>> PreProcessors;
