LiveLink - Allow providing custom LiveLinkSubject from the subject settings when creating a subject.

Also add settings to specify whether subjects should be translated / preprocessed before being rebroadcasted.

diff --git a/Engine/Plugins/Animation/LiveLink/Source/LiveLink/Private/LiveLinkClient.cpp b/Engine/Plugins/Animation/LiveLink/Source/LiveLink/Private/LiveLinkClient.cpp
--- a/Engine/Plugins/Animation/LiveLink/Source/LiveLink/Private/LiveLinkClient.cpp
+++ b/Engine/Plugins/Animation/LiveLink/Source/LiveLink/Private/LiveLinkClient.cpp
@@ -74,6 +74,9 @@ FLiveLinkClient::FLiveLinkClient()
 
 	// Setup rebroadcaster name in case we need it later
 	RebroadcastLiveLinkProviderName = TEXT("LiveLink Rebroadcast");
+
+	bPreProcessRebroadcastedFrames = GetDefault<ULiveLinkSettings>()->bPreProcessRebroadcastedFrames;
+	bTranslateRebroadcastedFrames = GetDefault<ULiveLinkSettings>()->bTranslateRebroadcastedFrames;
 }
 
 FLiveLinkClient::~FLiveLinkClient()
@@ -263,7 +266,7 @@ void FLiveLinkClient::HandleSubjectRebroadcast(ILiveLinkSubject* InSubject, cons
 	// Check the rebroadcast flag and act accordingly, creating the LiveLinkProvider and/or sending the static data if needed
 	if (InSubject->IsRebroadcasted())
 	{
-		if(InSubject->GetStaticData().IsValid() && InFrameData.IsValid())
+		if (InSubject->GetStaticData().IsValid() && InFrameData.IsValid())
 		{
 			// Setup rebroadcast provider
 			if (!RebroadcastLiveLinkProvider.IsValid())
@@ -273,18 +276,39 @@ void FLiveLinkClient::HandleSubjectRebroadcast(ILiveLinkSubject* InSubject, cons
 				
 			if (RebroadcastLiveLinkProvider.IsValid())
 			{
+				// Make a copy of the data for use by the rebroadcaster
+				FLiveLinkFrameDataStruct FrameDataCopy;
+				FrameDataCopy.InitializeWith(InFrameData);
+
+				if (bPreProcessRebroadcastedFrames)
+				{
+					FLiveLinkSubject* LiveSubject = static_cast<FLiveLinkSubject*>(InSubject);
+					InSubject->PreprocessFrame(FrameDataCopy);
+				}
+
+				FLiveLinkStaticDataStruct StaticDataCopy;
+	            StaticDataCopy.InitializeWith(InSubject->GetStaticData());
+
+				if (bTranslateRebroadcastedFrames)
+				{
+					TArray<ULiveLinkFrameTranslator::FWorkerSharedPtr> Translators = InSubject->GetFrameTranslators();
+					if (Translators.Num() && Translators[0].IsValid())
+					{
+						FLiveLinkSubjectFrameData TranslatedFrameData;
+						if (Translators[0]->Translate(InSubject->GetStaticData(), FrameDataCopy, TranslatedFrameData))
+						{
+							StaticDataCopy = MoveTemp(TranslatedFrameData.StaticData);
+							FrameDataCopy = MoveTemp(TranslatedFrameData.FrameData);
+						}
+					}
+				}
+
 				if (!InSubject->HasStaticDataBeenRebroadcasted())
 				{
-					FLiveLinkStaticDataStruct StaticDataCopy;
-					StaticDataCopy.InitializeWith(InSubject->GetStaticData());
 					RebroadcastLiveLinkProvider->UpdateSubjectStaticData(InSubject->GetSubjectKey().SubjectName, InSubject->GetRole(), MoveTemp(StaticDataCopy));
 					InSubject->SetStaticDataAsRebroadcasted(true);
 					RebroadcastedSubjects.Add(InSubject->GetSubjectKey());
 				}
-				
-				// Make a copy of the data for use by the rebroadcaster
-				FLiveLinkFrameDataStruct FrameDataCopy;
-				FrameDataCopy.InitializeWith(InFrameData);
 
 				RebroadcastLiveLinkProvider->UpdateSubjectFrameData(InSubject->GetSubjectKey().SubjectName, MoveTemp(FrameDataCopy));
 			}
@@ -990,6 +1014,7 @@ bool FLiveLinkClient::CreateSubject(const FLiveLinkSubjectPreset& InSubjectPrese
 		}
 
 		bool bEnabled = false;
+
 		FLiveLinkCollectionSubjectItem CollectionSubjectItem(InSubjectPreset.Key, MakeUnique<FLiveLinkSubject>(SourceItem->TimedData), SubjectSettings, bEnabled);
 		CollectionSubjectItem.GetLiveSubject()->Initialize(InSubjectPreset.Key, InSubjectPreset.Role.Get(), this);
 
diff --git a/Engine/Plugins/Animation/LiveLink/Source/LiveLink/Public/LiveLinkClient.h b/Engine/Plugins/Animation/LiveLink/Source/LiveLink/Public/LiveLinkClient.h
--- a/Engine/Plugins/Animation/LiveLink/Source/LiveLink/Public/LiveLinkClient.h
+++ b/Engine/Plugins/Animation/LiveLink/Source/LiveLink/Public/LiveLinkClient.h
@@ -301,6 +301,12 @@ private:
 	/** Name token  used to register to all subject updates. */
 	const FName ALL_SUBJECTS_DELEGATE_TOKEN = "__Internal_AllSubjects_Update";
 
+	/** Whether to Preprocess frames before rebroadcasting them. */
+	bool bPreProcessRebroadcastedFrames = false;
+
+	/** Whether to translate frames before rebroadcasting them. */
+	bool bTranslateRebroadcastedFrames = false;
+
 #if WITH_EDITOR
 	/** Delegate when a subject is evaluated. */
 	FOnLiveLinkSubjectEvaluated OnLiveLinkSubjectEvaluatedDelegate;
diff --git a/Engine/Plugins/Animation/LiveLink/Source/LiveLink/Public/LiveLinkSettings.h b/Engine/Plugins/Animation/LiveLink/Source/LiveLink/Public/LiveLinkSettings.h
--- a/Engine/Plugins/Animation/LiveLink/Source/LiveLink/Public/LiveLinkSettings.h
+++ b/Engine/Plugins/Animation/LiveLink/Source/LiveLink/Public/LiveLinkSettings.h
@@ -121,7 +121,15 @@ public:
 	/** Subjects will be removed when their source has been unresponsive for this long. */
 	UPROPERTY(config, EditAnywhere, AdvancedDisplay, Category = "LiveLink", meta=(ForceUnits=s))
 	double MessageBusTimeBeforeRemovingInactiveSource;
-	
+
+	/** Whether to Preprocess frames before rebroadcasting them. */
+	UPROPERTY(config)
+	bool bPreProcessRebroadcastedFrames = false;
+
+	/** Whether to translate frames before rebroadcasting them. */
+	UPROPERTY(config)
+	bool bTranslateRebroadcastedFrames = false;
+
 	/**
 	 * A source may still exist but does not send frames for a subject.
 	 * Time before considering the subject as "invalid".
diff --git a/Engine/Plugins/Animation/LiveLink/Source/LiveLink/Public/LiveLinkSubject.h b/Engine/Plugins/Animation/LiveLink/Source/LiveLink/Public/LiveLinkSubject.h
--- a/Engine/Plugins/Animation/LiveLink/Source/LiveLink/Public/LiveLinkSubject.h
+++ b/Engine/Plugins/Animation/LiveLink/Source/LiveLink/Public/LiveLinkSubject.h
@@ -71,6 +71,7 @@ public:
 	virtual bool IsRebroadcasted() const override { return bRebroadcastSubject; }
 	virtual bool HasStaticDataBeenRebroadcasted() const override { return bRebroadcastStaticDataSent; }
 	virtual void SetStaticDataAsRebroadcasted(const bool bInSent) override { bRebroadcastStaticDataSent = bInSent; }
+	virtual void PreprocessFrame(FLiveLinkFrameDataStruct& InOutFrameData) override;
 protected:
 	virtual const FLiveLinkSubjectFrameData& GetFrameSnapshot() const override { return FrameSnapshot; }
 	//~ End ILiveLinkSubject Interface
@@ -115,9 +116,6 @@ public:
 	/** Set the last time a frame was received. */
 	void SetLastPushTime(double InLastPushTime) { LastPushTime = InLastPushTime; }
 
-	/** Apply this subject's preprocessors to a frame data. */
-	void PreprocessFrame(FLiveLinkFrameDataStruct& InOutFrameData);
-
 	/** Validates if the incoming frame data is compatible with the static data for this subject. */
 	bool ValidateFrameData(const FLiveLinkFrameDataStruct& InFrameData);
 
diff --git a/Engine/Plugins/Animation/LiveLink/Source/LiveLinkEditor/Private/LiveLinkClientPanelViews.cpp b/Engine/Plugins/Animation/LiveLink/Source/LiveLinkEditor/Private/LiveLinkClientPanelViews.cpp
--- a/Engine/Plugins/Animation/LiveLink/Source/LiveLinkEditor/Private/LiveLinkClientPanelViews.cpp
+++ b/Engine/Plugins/Animation/LiveLink/Source/LiveLinkEditor/Private/LiveLinkClientPanelViews.cpp
@@ -299,10 +299,6 @@ private:
 			return EntryPtr->IsSubjectValid() ? GetDefault<ULiveLinkSettings>()->GetValidColor() : GetDefault<ULiveLinkSettings>()->GetInvalidColor();
 		}
 
-		if (!EntryPtr->IsSubjectEnabled() && EntryPtr->IsSubject())
-		{
-			UE_LOG(LogTemp, Warning, TEXT("Entry %s disabled"), *EntryPtr->SubjectKey.SubjectName.ToString());
-		}
 		return FLinearColor(0.f, 0.f, 0.f, 0.f);
 	}
 
diff --git a/Engine/Source/Runtime/LiveLinkInterface/Public/ILiveLinkSubject.h b/Engine/Source/Runtime/LiveLinkInterface/Public/ILiveLinkSubject.h
--- a/Engine/Source/Runtime/LiveLinkInterface/Public/ILiveLinkSubject.h
+++ b/Engine/Source/Runtime/LiveLinkInterface/Public/ILiveLinkSubject.h
@@ -51,6 +51,9 @@ public:
 	/** Sets the static data for this subject as rebroadcasted */
 	virtual void SetStaticDataAsRebroadcasted(const bool bInSent) = 0;
 
+	/** Apply this subject's preprocessors to a frame data. */
+	virtual void PreprocessFrame(FLiveLinkFrameDataStruct& InOutFrameData) {};
+
 protected:
 	virtual const FLiveLinkSubjectFrameData& GetFrameSnapshot() const = 0;
 
diff --git a/Engine/Source/Runtime/LiveLinkInterface/Public/LiveLinkSubjectSettings.h b/Engine/Source/Runtime/LiveLinkInterface/Public/LiveLinkSubjectSettings.h
--- a/Engine/Source/Runtime/LiveLinkInterface/Public/LiveLinkSubjectSettings.h
+++ b/Engine/Source/Runtime/LiveLinkInterface/Public/LiveLinkSubjectSettings.h
@@ -4,6 +4,7 @@
 
 #include "Containers/Array.h"
 #include "CoreMinimal.h"
+#include "ILiveLinkSubject.h"
 #include "LiveLinkRole.h"
 #include "Misc/FrameRate.h"
 #include "Templates/SubclassOf.h"
@@ -12,13 +13,37 @@
 #include "UObject/ObjectPtr.h"
 #include "UObject/UObjectGlobals.h"
 
+
 #include "LiveLinkSubjectSettings.generated.h"
 
+class FLiveLinkTimedDataInput;
 class ULiveLinkFrameInterpolationProcessor;
 class ULiveLinkFramePreProcessor;
 class ULiveLinkFrameTranslator;
 class ULiveLinkRole;
 
+/**
+ * Utility class that allows specifying default values for Subject settings.
+ */
+UCLASS(config=Engine, defaultconfig)
+class ULiveLinkDefaultSubjectSettings : public UObject
+{
+	GENERATED_BODY()
+
+public:
+	/** Whether subjects should be rebroadcasted by default. */
+	UPROPERTY(config)
+	bool bRebroadcastSubjectsByDefault = false;
+
+	/**
+	 * Whether a user should be able to edit the bRebroadcastSubject property.
+	 * Setting this to false in a target config will prevent a user from turning on or off the rebroadcast flag on a subject.
+	 *
+	 */
+	UPROPERTY(config)
+	bool bAllowEditingRebroadcastProperty = true;
+};
+
 
 // Base class for live link subject settings
 UCLASS(MinimalAPI)
@@ -27,6 +52,15 @@ class ULiveLinkSubjectSettings : public UObject
 public:
 	GENERATED_BODY()
 
+	LIVELINKINTERFACE_API ULiveLinkSubjectSettings()
+	{
+		if (!HasAnyFlags(RF_ClassDefaultObject | RF_ArchetypeObject))
+		{
+			bRebroadcastSubject = GetDefault<ULiveLinkDefaultSubjectSettings>()->bRebroadcastSubjectsByDefault;
+			bAllowModifyingRebroadcast = GetDefault<ULiveLinkDefaultSubjectSettings>()->bAllowEditingRebroadcastProperty;
+		}
+	}
+
 	/** Initialize the settings. */
 	virtual void Initialize(FLiveLinkSubjectKey InSubjectKey) {}
 
@@ -48,9 +82,13 @@ public:
 	/** Last FrameRate estimated by the subject. If in Timecode mode, this will come directly from the QualifiedFrameTime. */
 	UPROPERTY(VisibleAnywhere, Category="LiveLink")
 	FFrameRate FrameRate;
+
+	/** Allows settings to dictate whether the rebroadcast flag is editable. */
+	UPROPERTY()
+	bool bAllowModifyingRebroadcast = true;
 	
 	/** If enabled, rebroadcast this subject */
-	UPROPERTY(EditAnywhere, Category = "LiveLink")
+	UPROPERTY(EditAnywhere, Category = "LiveLink", meta=(EditConditionHides, EditCondition="bAllowModifyingRebroadcast"))
     bool bRebroadcastSubject;
 
 	/** Validate PreProcessors, Translators and Interpolation processors. Usually called after a property change event.
