LiveLinkHub - Fix recordings not working with translators

diff --git a/Engine/Plugins/Animation/LiveLink/Source/LiveLink/Private/LiveLinkClient.cpp b/Engine/Plugins/Animation/LiveLink/Source/LiveLink/Private/LiveLinkClient.cpp
--- a/Engine/Plugins/Animation/LiveLink/Source/LiveLink/Private/LiveLinkClient.cpp
+++ b/Engine/Plugins/Animation/LiveLink/Source/LiveLink/Private/LiveLinkClient.cpp
@@ -289,6 +289,8 @@ void FLiveLinkClient::HandleSubjectRebroadcast(ILiveLinkSubject* InSubject, cons
 				FLiveLinkStaticDataStruct StaticDataCopy;
 	            StaticDataCopy.InitializeWith(InSubject->GetStaticData());
 
+				TSubclassOf<ULiveLinkRole> SubjectRole = InSubject->GetRole();
+
 				if (bTranslateRebroadcastedFrames)
 				{
 					TArray<ULiveLinkFrameTranslator::FWorkerSharedPtr> Translators = InSubject->GetFrameTranslators();
@@ -297,6 +299,7 @@ void FLiveLinkClient::HandleSubjectRebroadcast(ILiveLinkSubject* InSubject, cons
 						FLiveLinkSubjectFrameData TranslatedFrameData;
 						if (Translators[0]->Translate(InSubject->GetStaticData(), FrameDataCopy, TranslatedFrameData))
 						{
+							SubjectRole = Translators[0]->GetToRole();
 							StaticDataCopy = MoveTemp(TranslatedFrameData.StaticData);
 							FrameDataCopy = MoveTemp(TranslatedFrameData.FrameData);
 						}
@@ -305,7 +308,7 @@ void FLiveLinkClient::HandleSubjectRebroadcast(ILiveLinkSubject* InSubject, cons
 
 				if (!InSubject->HasStaticDataBeenRebroadcasted())
 				{
-					RebroadcastLiveLinkProvider->UpdateSubjectStaticData(InSubject->GetSubjectKey().SubjectName, InSubject->GetRole(), MoveTemp(StaticDataCopy));
+					RebroadcastLiveLinkProvider->UpdateSubjectStaticData(InSubject->GetSubjectKey().SubjectName, SubjectRole, MoveTemp(StaticDataCopy));
 					InSubject->SetStaticDataAsRebroadcasted(true);
 					RebroadcastedSubjects.Add(InSubject->GetSubjectKey());
 				}
