LiveLink: Added basic deadband pre-processor.

Future work may include a catchup servo-control.

diff --git a/Engine/Plugins/Animation/LiveLink/Source/LiveLink/Private/PreProcessor/LiveLinkDeadbandPreProcessor.cpp b/Engine/Plugins/Animation/LiveLink/Source/LiveLink/Private/PreProcessor/LiveLinkDeadbandPreProcessor.cpp
new file mode 100644
--- /dev/null
+++ b/Engine/Plugins/Animation/LiveLink/Source/LiveLink/Private/PreProcessor/LiveLinkDeadbandPreProcessor.cpp
@@ -0,0 +1,93 @@
+// Copyright Epic Games, Inc. All Rights Reserved.
+
+#include "PreProcessor/LiveLinkDeadbandPreProcessor.h"
+#include "Roles/LiveLinkAnimationRole.h"
+#include "Roles/LiveLinkAnimationTypes.h"
+#include "Roles/LiveLinkTransformRole.h"
+#include "Roles/LiveLinkTransformTypes.h"
+#include "UObject/ReleaseObjectVersion.h"
+
+// WORKER
+
+TSubclassOf<ULiveLinkRole> ULiveLinkTransformDeadbandPreProcessor::FLiveLinkTransformDeadbandPreProcessorWorker::GetRole() const
+{
+	return ULiveLinkTransformRole::StaticClass();
+}
+
+bool ULiveLinkTransformDeadbandPreProcessor::FLiveLinkTransformDeadbandPreProcessorWorker::PreProcessFrame(FLiveLinkFrameDataStruct& InOutFrame) const
+{
+	if (!bEnableDeadband)
+	{
+		return true;
+	}
+
+	FLiveLinkTransformFrameData& TransformData = *InOutFrame.Cast<FLiveLinkTransformFrameData>();
+
+	// location
+	{
+		const double DeltaTranslation = (StableTransform.GetLocation() - TransformData.Transform.GetLocation()).Length();
+
+		if (DeltaTranslation < TranslationDeadband)
+		{
+			// Keep stable location
+			TransformData.Transform.SetLocation(StableTransform.GetLocation());
+		}
+	}
+
+	// rotation
+	{
+		const double DeltaRotationInDegrees = FMath::RadiansToDegrees(StableTransform.GetRotation().AngularDistance(TransformData.Transform.GetRotation()));
+
+		if (DeltaRotationInDegrees < RotationDeadbandInDegrees)
+		{
+			// Keep stable rotation
+			TransformData.Transform.SetRotation(StableTransform.GetRotation());
+		}
+	}
+
+	StableTransform = TransformData.Transform;
+
+	return true;
+}
+
+// PREPROCESSOR
+
+TSubclassOf<ULiveLinkRole> ULiveLinkTransformDeadbandPreProcessor::GetRole() const
+{
+	return ULiveLinkTransformRole::StaticClass();
+}
+
+ULiveLinkTransformDeadbandPreProcessor::FWorkerSharedPtr ULiveLinkTransformDeadbandPreProcessor::FetchWorker()
+{
+	if (!Instance.IsValid())
+	{
+		Instance = MakeShared<FLiveLinkTransformDeadbandPreProcessorWorker, ESPMode::ThreadSafe>();
+		Instance->TranslationDeadband = TranslationDeadband;
+		Instance->RotationDeadbandInDegrees = RotationDeadbandInDegrees;
+		Instance->bEnableDeadband = bEnableDeadband;
+	}
+
+	return Instance;
+}
+
+#if WITH_EDITOR
+void ULiveLinkTransformDeadbandPreProcessor::PostEditChangeChainProperty(struct FPropertyChangedChainEvent& PropertyChangedEvent)
+{
+	static const FName NAME_TranslationDeadband = GET_MEMBER_NAME_CHECKED(ThisClass, TranslationDeadband);
+	static const FName NAME_RotationDegreesDeadband = GET_MEMBER_NAME_CHECKED(ThisClass, RotationDeadbandInDegrees);
+	static const FName NAME_bEnableDeadband = GET_MEMBER_NAME_CHECKED(ThisClass, bEnableDeadband);
+
+	const FName PropertyName = PropertyChangedEvent.Property->GetFName();
+
+	if ((PropertyName == NAME_TranslationDeadband) 
+		|| (PropertyName == NAME_RotationDegreesDeadband)
+		|| (PropertyName == NAME_bEnableDeadband)
+	)
+	{
+		Instance.Reset();
+	}
+
+	Super::PostEditChangeChainProperty(PropertyChangedEvent);
+}
+#endif //WITH_EDITOR
+
diff --git a/Engine/Plugins/Animation/LiveLink/Source/LiveLink/Public/PreProcessor/LiveLinkDeadbandPreProcessor.h b/Engine/Plugins/Animation/LiveLink/Source/LiveLink/Public/PreProcessor/LiveLinkDeadbandPreProcessor.h
new file mode 100644
--- /dev/null
+++ b/Engine/Plugins/Animation/LiveLink/Source/LiveLink/Public/PreProcessor/LiveLinkDeadbandPreProcessor.h
@@ -0,0 +1,75 @@
+// Copyright Epic Games, Inc. All Rights Reserved.
+
+#pragma once
+
+#include "LiveLinkFramePreProcessor.h"
+#include "LiveLinkDeadbandPreProcessor.generated.h"
+
+
+/**
+ * Implements a deadband filter that gets applied to the transform, with independent thresholds
+ * for rotation and translation.
+ */
+UCLASS(meta = (DisplayName = "Transform Deadband"))
+class LIVELINK_API ULiveLinkTransformDeadbandPreProcessor : public ULiveLinkFramePreProcessor
+{
+	GENERATED_BODY()
+
+public:
+	class FLiveLinkTransformDeadbandPreProcessorWorker : public ILiveLinkFramePreProcessorWorker
+	{
+	public:
+
+		/** If false, transform is left untouched */
+		bool bEnableDeadband = true;
+
+		/** Translation is updated only if the change is larger than this threshold */
+		double TranslationDeadband = 0;
+
+		/** Rotation is updated only if the change is larger than this threshold */
+		double RotationDeadbandInDegrees = 0;
+
+		/** Keeps track of the last accepted location and translation */
+		mutable FTransform StableTransform;
+
+	public:
+
+		//~Begin ILiveLinkFramePreProcessorWorker
+		virtual TSubclassOf<ULiveLinkRole> GetRole() const override;
+		virtual bool PreProcessFrame(FLiveLinkFrameDataStruct& InOutFrame) const override;
+		//~End ILiveLinkFramePreProcessorWorker
+	};
+
+protected:
+
+	/** If false, transform is left untouched */
+	UPROPERTY(EditAnywhere, Category = "Deadband")
+	bool bEnableDeadband = true;
+
+	/** Translation is updated only if the change is larger than this threshold */
+	UPROPERTY(EditAnywhere, Category = "Deadband")
+	float TranslationDeadband = 0;
+
+	/** Rotation is updated only if the change is larger than this threshold */
+	UPROPERTY(EditAnywhere, Category = "Deadband")
+	float RotationDeadbandInDegrees = 0;
+
+public:
+	//~Begin ULiveLinkFramePreProcessor interface
+	virtual TSubclassOf<ULiveLinkRole> GetRole() const override;
+	virtual ULiveLinkFramePreProcessor::FWorkerSharedPtr FetchWorker() override;
+	//~End ULiveLinkFramePreProcessor interface
+
+public:
+	//~ Begin UObject interface
+#if WITH_EDITOR
+	virtual void PostEditChangeChainProperty(struct FPropertyChangedChainEvent& PropertyChangedEvent) override;
+#endif //WITH_EDITOR
+	//~ End UObject interface
+
+protected:
+
+	/** Instance of the worker object */
+	TSharedPtr<FLiveLinkTransformDeadbandPreProcessorWorker, ESPMode::ThreadSafe> Instance;
+};
+
