Fix livelinkhub uninitlialized properties

diff --git a/Engine/Plugins/Animation/LiveLink/Source/LiveLink/Public/LiveLinkVirtualSubjectBoneAttachment.h b/Engine/Plugins/Animation/LiveLink/Source/LiveLink/Public/LiveLinkVirtualSubjectBoneAttachment.h
--- a/Engine/Plugins/Animation/LiveLink/Source/LiveLink/Public/LiveLinkVirtualSubjectBoneAttachment.h
+++ b/Engine/Plugins/Animation/LiveLink/Source/LiveLink/Public/LiveLinkVirtualSubjectBoneAttachment.h
@@ -38,11 +38,11 @@ public:
 
 	/** Location of the component relative to its parent */
 	UPROPERTY(EditAnywhere,  Category = "Live Link", meta=(LinearDeltaSensitivity = "1", Delta = "1.0"))
-	FVector LocationOffset;
+	FVector LocationOffset = FVector::ZeroVector;
 
 	/** Rotation of the component relative to its parent */
 	UPROPERTY(EditAnywhere, Category = "Live Link", meta=(LinearDeltaSensitivity = "1", Delta = "1.0"))
-	FRotator RotationOffset;
+	FRotator RotationOffset = FRotator::ZeroRotator;
 
 	/** Whether to ignore the parent bone location when attaching the child bone. */
 	UPROPERTY(EditAnywhere, Category = "Live Link")
