LiveLinkComponent: Expose Get/Set component to control functions to BP

diff --git a/Engine/Plugins/Animation/LiveLink/Source/LiveLinkComponents/Public/LiveLinkComponentController.h b/Engine/Plugins/Animation/LiveLink/Source/LiveLinkComponents/Public/LiveLinkComponentController.h
--- a/Engine/Plugins/Animation/LiveLink/Source/LiveLinkComponents/Public/LiveLinkComponentController.h
+++ b/Engine/Plugins/Animation/LiveLink/Source/LiveLinkComponents/Public/LiveLinkComponentController.h
@@ -94,9 +94,11 @@ public:
 	void OnSubjectRoleChanged();
 
 	/** Returns the component controlled by the LiveLink controller of the input Role. Returns null if there is no controller for that Role */
+	UFUNCTION(BlueprintPure, Category="LiveLink")
 	UActorComponent* GetControlledComponent(TSubclassOf<ULiveLinkRole> InRoleClass) const;
 
 	/** Set the component to control for the LiveLink controller of the input Role */
+	UFUNCTION(BlueprintCallable, Category = "LiveLink")
 	void SetControlledComponent(TSubclassOf<ULiveLinkRole> InRoleClass, UActorComponent* InComponent);
 
 	/** Multicast delegate that broadcasts after LiveLink controllers have ticked with the latest frame of subject data */
