LiveLinkHub - Add subject remapper support

Subject remappers allow you to remap bone and transform data that comes through a livelink subject.
This should remove the need for anim blueprints that use a remap asset.

diff --git a/Engine/Plugins/Animation/LiveLink/Source/LiveLink/Private/LiveLinkClient.cpp b/Engine/Plugins/Animation/LiveLink/Source/LiveLink/Private/LiveLinkClient.cpp
--- a/Engine/Plugins/Animation/LiveLink/Source/LiveLink/Private/LiveLinkClient.cpp
+++ b/Engine/Plugins/Animation/LiveLink/Source/LiveLink/Private/LiveLinkClient.cpp
@@ -12,6 +12,7 @@
 #include "LiveLinkRoleTrait.h"
 #include "LiveLinkProvider.h"
 #include "LiveLinkSettings.h"
+#include "LiveLinkSubjectRemapper.h"
 #include "LiveLinkSourceCollection.h"
 #include "LiveLinkSourceFactory.h"
 #include "LiveLinkTimedDataInput.h"
@@ -75,8 +76,8 @@ FLiveLinkClient::FLiveLinkClient()
 	// Setup rebroadcaster name in case we need it later
 	RebroadcastLiveLinkProviderName = TEXT("LiveLink Rebroadcast");
 
-	bPreProcessRebroadcastedFrames = GetDefault<ULiveLinkSettings>()->bPreProcessRebroadcastedFrames;
-	bTranslateRebroadcastedFrames = GetDefault<ULiveLinkSettings>()->bTranslateRebroadcastedFrames;
+	bPreProcessRebroadcastFrames = GetDefault<ULiveLinkSettings>()->bPreProcessRebroadcastFrames;
+	bTranslateRebroadcastFrames = GetDefault<ULiveLinkSettings>()->bTranslateRebroadcastFrames;
 }
 
 FLiveLinkClient::~FLiveLinkClient()
@@ -280,7 +281,7 @@ void FLiveLinkClient::HandleSubjectRebroadcast(ILiveLinkSubject* InSubject, cons
 				FLiveLinkFrameDataStruct FrameDataCopy;
 				FrameDataCopy.InitializeWith(InFrameData);
 
-				if (bPreProcessRebroadcastedFrames)
+				if (bPreProcessRebroadcastFrames)
 				{
 					InSubject->PreprocessFrame(FrameDataCopy);
 				}
@@ -290,7 +291,7 @@ void FLiveLinkClient::HandleSubjectRebroadcast(ILiveLinkSubject* InSubject, cons
 
 				TSubclassOf<ULiveLinkRole> SubjectRole = InSubject->GetRole();
 
-				if (bTranslateRebroadcastedFrames)
+				if (bTranslateRebroadcastFrames)
 				{
 					TArray<ULiveLinkFrameTranslator::FWorkerSharedPtr> Translators = InSubject->GetFrameTranslators();
 					if (Translators.Num() && Translators[0].IsValid())
@@ -801,6 +802,20 @@ void FLiveLinkClient::PushSubjectStaticData_Internal(FPendingSubjectStatic&& Sub
 
 	if (LiveLinkSubject)
 	{
+		if (ULiveLinkSubjectRemapper::FWorkerSharedPtr Remapper = LiveLinkSubject->GetFrameRemapper())
+		{
+			// ATM we will assume vsubjects can't have remappers.
+			if (ULiveLinkSubjectSettings* Settings = Cast<ULiveLinkSubjectSettings>(GetSubjectSettings(SubjectStaticData.SubjectKey)))
+			{
+				// Make sure we have a valid settings object to not remap the static data while we're resetting the remapper.
+				if (Settings->Remapper)
+				{
+					// Make sure to rebroadcast the new static data.
+					LiveLinkSubject->SetStaticDataAsRebroadcasted(false);
+				}
+			}
+		}
+
 		if (const FSubjectFramesAddedHandles* Handles = SubjectFrameAddedHandles.Find(SubjectStaticData.SubjectKey.SubjectName))
 		{
 			Handles->OnStaticDataAdded.Broadcast(SubjectStaticData.SubjectKey, SubjectStaticData.Role, SubjectStaticData.StaticData);
@@ -911,6 +926,11 @@ void FLiveLinkClient::PushSubjectFrameData_Internal(FPendingSubjectFrame&& Subje
 		return;
 	}
 
+	if (ULiveLinkSubjectRemapper::FWorkerSharedPtr Remapper = LinkSubject->GetFrameRemapper())
+	{
+		Remapper->RemapFrameData(LinkSubject->GetStaticData(), SubjectFrameData.FrameData);
+	}
+
 	bool bShouldLogWarning = true;
 	if (!Role.GetDefaultObject()->IsFrameDataValid(LinkSubject->GetStaticData(), SubjectFrameData.FrameData, bShouldLogWarning))
 	{
@@ -1366,6 +1386,14 @@ bool FLiveLinkClient::HasPendingSubjectFrames()
 	return !SubjectFrameToPush.IsEmpty();
 }
 
+void FLiveLinkClient::ClearOverrideStaticData_AnyThread(const FLiveLinkSubjectKey& InSubjectKey)
+{
+	if (const FLiveLinkCollectionSubjectItem* SubjectItem = Collection->FindSubject(InSubjectKey))
+	{
+		SubjectItem->GetLiveSubject()->ClearOverrideStaticData_AnyThread();
+	}
+}
+
 bool FLiveLinkClient::EvaluateFrame_AnyThread(FLiveLinkSubjectName InSubjectName, TSubclassOf<ULiveLinkRole> InDesiredRole, FLiveLinkSubjectFrameData& OutFrame)
 {
 	SCOPE_CYCLE_COUNTER(STAT_LiveLink_EvaluateFrame);
diff --git a/Engine/Plugins/Animation/LiveLink/Source/LiveLink/Private/LiveLinkSubject.cpp b/Engine/Plugins/Animation/LiveLink/Source/LiveLink/Private/LiveLinkSubject.cpp
--- a/Engine/Plugins/Animation/LiveLink/Source/LiveLink/Private/LiveLinkSubject.cpp
+++ b/Engine/Plugins/Animation/LiveLink/Source/LiveLink/Private/LiveLinkSubject.cpp
@@ -58,6 +58,13 @@ void FLiveLinkSubject::Update()
 {
 	FScopeLock Lock(&SettingsCriticalSection);
 
+	if (bClearOverrideStaticData)
+	{
+		ClearFrames();
+		OverrideStaticData.Reset();
+		bClearOverrideStaticData = false;
+	}
+
 	// Clear all frames that are too old
 	if (FrameData.Num() > CachedSettings.BufferSettings.MaxNumberOfFrameToBuffered)
 	{
@@ -210,6 +217,26 @@ bool FLiveLinkSubject::HasValidFrameSnapshot() const
 	return FrameSnapshot.StaticData.IsValid() && FrameSnapshot.FrameData.IsValid();
 }
 
+FLiveLinkStaticDataStruct& FLiveLinkSubject::GetStaticData()
+{
+	if (OverrideStaticData)
+	{
+		return *OverrideStaticData;
+	}
+
+	return StaticData;
+}
+
+const FLiveLinkStaticDataStruct& FLiveLinkSubject::GetStaticData() const
+{
+	if (OverrideStaticData)
+	{
+		return *OverrideStaticData;
+	}
+
+	return StaticData;
+}
+
 TArray<FLiveLinkTime> FLiveLinkSubject::GetFrameTimes() const
 {
 	TArray<FLiveLinkTime> Result;
@@ -408,7 +435,7 @@ void FLiveLinkSubject::AddFrameData(FLiveLinkFrameDataStruct&& InFrameData)
 		{
 			PreProcessor->PreProcessFrame(InFrameData);
 		}
-		
+
 		//Assign identifier to incoming frame and insert it where it belongs
 		const FLiveLinkFrameIdentifier ThisFrameIdentifier = NextIdentifier++;
 		ReceivedOrderedFrames.Enqueue(ThisFrameIdentifier);
@@ -473,6 +500,11 @@ bool FLiveLinkSubject::ValidateFrameData(const FLiveLinkFrameDataStruct& InFrame
 	return true;
 }
 
+void FLiveLinkSubject::ClearOverrideStaticData_AnyThread()
+{
+	bClearOverrideStaticData = true;
+}
+
 int32 FLiveLinkSubject::FindNewFrame_WorldTime(const FLiveLinkWorldTime& WorldTime) const
 {
 	if (CachedSettings.BufferSettings.bValidEngineTimeEnabled)
@@ -1096,6 +1128,58 @@ void FLiveLinkSubject::CacheSettings(ULiveLinkSourceSettings* SourceSetting, ULi
 			}
 		}
 
+		SubjectRemapper.Reset();
+		if (SubjectSetting->Remapper)
+		{
+			// If there wasn't a remapper, then we need to initialize the one we will create.
+			bool bRecreateRemapper = SubjectSetting->Remapper->GetWorker() == nullptr || SubjectSetting->Remapper->bDirty;
+
+			ULiveLinkSubjectRemapper::FWorkerSharedPtr NewRemapper = bRecreateRemapper ? SubjectSetting->Remapper->CreateWorker() : SubjectSetting->Remapper->GetWorker();
+
+			if (NewRemapper.IsValid())
+			{
+				// If this is a new remapper, then do an initial remapping.
+				if (bRecreateRemapper)
+				{
+					SubjectSetting->Remapper->bDirty = false;
+
+					if (SubjectSetting->Remapper->IsValidRemapper())
+					{
+						// Since the remapper has changed, we need to update the static data as well. 
+						FLiveLinkStaticDataStruct RemappedStaticData;
+						RemappedStaticData.InitializeWith(StaticData);
+						NewRemapper->RemapStaticData(RemappedStaticData);
+
+						if (HasValidFrameSnapshot())
+						{
+							NewRemapper->RemapStaticData(FrameSnapshot.StaticData);
+							NewRemapper->RemapFrameData(FrameSnapshot.StaticData, FrameSnapshot.FrameData);
+						}
+
+						// Important: Because we changed the static data, we have to update the frame data to match the new static data.
+						for (FLiveLinkFrameDataStruct& FrameDataStruct : FrameData)
+						{
+							NewRemapper->RemapFrameData(RemappedStaticData, FrameDataStruct);
+						}
+
+						OverrideStaticData = MoveTemp(RemappedStaticData);
+
+						SetStaticDataAsRebroadcasted(false);
+					}
+					else
+					{
+						// Remapper isn't valid, reset static data.
+						ClearFrames();
+
+						OverrideStaticData.Reset();
+					}
+
+				}
+			}
+
+			SubjectRemapper = MoveTemp(NewRemapper);
+		}
+
 		// Create a new or fetch the interpolation for this frame
 		FrameInterpolationProcessor.Reset();
 		if (SubjectSetting->InterpolationProcessor)
diff --git a/Engine/Plugins/Animation/LiveLink/Source/LiveLink/Private/Remapper/LiveLinkSkeletonRemapper.h b/Engine/Plugins/Animation/LiveLink/Source/LiveLink/Private/Remapper/LiveLinkSkeletonRemapper.h
new file mode 100644
--- /dev/null
+++ b/Engine/Plugins/Animation/LiveLink/Source/LiveLink/Private/Remapper/LiveLinkSkeletonRemapper.h
@@ -0,0 +1,100 @@
+// Copyright Epic Games, Inc. All Rights Reserved.
+
+#pragma once
+
+
+#include "LiveLinkSubjectRemapper.h"
+
+#include "Engine/SkeletalMesh.h"
+#include "Features/IModularFeatures.h"
+#include "LiveLinkClient.h"
+#include "LiveLinkTypes.h"
+#include "Roles/LiveLinkAnimationRole.h"
+#include "Roles/LiveLinkAnimationTypes.h"
+
+#include "LiveLinkSkeletonRemapper.generated.h"
+
+class FLiveLinkSkeletonRemapperWorker : public ILiveLinkSubjectRemapperWorker
+{
+public:
+	virtual void RemapStaticData(FLiveLinkStaticDataStruct& InOutStaticData)
+	{
+		if (InOutStaticData.GetStruct() && InOutStaticData.GetStruct()->IsChildOf<FLiveLinkSkeletonStaticData>())
+		{
+			RemapSkeletonStaticData(*InOutStaticData.Cast<FLiveLinkSkeletonStaticData>());
+		}
+	}
+
+	virtual void RemapFrameData(const FLiveLinkStaticDataStruct& InOutStaticData, FLiveLinkFrameDataStruct& InOutFrameData)
+	{
+		if (InOutStaticData.GetStruct() && InOutStaticData.GetStruct()->IsChildOf<FLiveLinkSkeletonStaticData>())
+		{
+			RemapSkeletonFrameData(*InOutStaticData.Cast<FLiveLinkSkeletonStaticData>(), *InOutFrameData.Cast<FLiveLinkAnimationFrameData>());
+		}
+	}
+
+	virtual void RemapSkeletonStaticData(FLiveLinkSkeletonStaticData& InOutSkeletonData)
+	{
+	}
+
+	virtual void RemapSkeletonFrameData(const FLiveLinkSkeletonStaticData& InOutSkeletonData, FLiveLinkAnimationFrameData& InOutFrameData)
+	{
+	}
+
+	FName GetRemappedBoneName(FName BoneName) const
+	{
+		if (const FName* RemappedName = BoneNameMap.Find(BoneName))
+		{
+			return *RemappedName;
+		}
+		return BoneName;
+	}
+
+	/** Map used to provide new names for the bones in the static data. */
+	TMap<FName, FName> BoneNameMap;
+};
+
+UCLASS(Abstract)
+class ULiveLinkSkeletonRemapper : public ULiveLinkSubjectRemapper
+{
+public:
+	GENERATED_BODY()
+
+	virtual void Initialize(const FLiveLinkSubjectKey& SubjectKey) override
+	{
+		ILiveLinkClient& LiveLinkClient = IModularFeatures::Get().GetModularFeature<ILiveLinkClient>(ILiveLinkClient::ModularFeatureName);
+		const FLiveLinkStaticDataStruct* StaticData = LiveLinkClient.GetSubjectStaticData_AnyThread(SubjectKey);
+
+		TSubclassOf<ULiveLinkRole> LiveLinkRole = LiveLinkClient.GetSubjectRole_AnyThread(SubjectKey);
+
+		// Note: Should we initialize the bone name map using the reference skeleton?
+		if (StaticData && LiveLinkRole && LiveLinkRole->IsChildOf(ULiveLinkAnimationRole::StaticClass()))
+		{
+			for (FName BoneName : StaticData->Cast<FLiveLinkSkeletonStaticData>()->GetBoneNames())
+			{
+				BoneNameMap.Add(BoneName, BoneName);
+			}
+		}
+	}
+
+	virtual TSubclassOf<ULiveLinkRole> GetSupportedRole() const override
+	{
+		return ULiveLinkAnimationRole::StaticClass();
+	}
+
+#if WITH_EDITOR
+	virtual void PostEditChangeProperty(struct FPropertyChangedEvent& PropertyChangedEvent) override
+	{
+		bDirty = true;
+	}
+#endif
+
+	virtual bool IsValidRemapper() const
+	{
+		return ReferenceSkeleton.IsValid();
+	}
+
+public:
+	UPROPERTY(EditAnywhere, Category="Remapper", meta=(DisplayThumbnail="false"))
+	TSoftObjectPtr<USkeletalMesh> ReferenceSkeleton;
+};
diff --git a/Engine/Plugins/Animation/LiveLink/Source/LiveLink/Public/LiveLinkClient.h b/Engine/Plugins/Animation/LiveLink/Source/LiveLink/Public/LiveLinkClient.h
--- a/Engine/Plugins/Animation/LiveLink/Source/LiveLink/Public/LiveLinkClient.h
+++ b/Engine/Plugins/Animation/LiveLink/Source/LiveLink/Public/LiveLinkClient.h
@@ -125,6 +125,7 @@ public:
 	virtual bool EvaluateFrameAtSceneTime_AnyThread(FLiveLinkSubjectName SubjectName, const FQualifiedFrameTime& FrameTime, TSubclassOf<ULiveLinkRole> DesiredRole, FLiveLinkSubjectFrameData& OutFrame) override;
 	virtual void ForceTick() override;
 	virtual bool HasPendingSubjectFrames() override;
+	virtual void ClearOverrideStaticData_AnyThread(const FLiveLinkSubjectKey& InSubjectKey) override;
 
 	virtual FSimpleMulticastDelegate& OnLiveLinkTicked() override;
 	virtual FSimpleMulticastDelegate& OnLiveLinkSourcesChanged() override;
@@ -305,10 +306,10 @@ private:
 	const FName ALL_SUBJECTS_DELEGATE_TOKEN = "__Internal_AllSubjects_Update";
 
 	/** Whether to Preprocess frames before rebroadcasting them. */
-	bool bPreProcessRebroadcastedFrames = false;
+	bool bPreProcessRebroadcastFrames = false;
 
 	/** Whether to translate frames before rebroadcasting them. */
-	bool bTranslateRebroadcastedFrames = false;
+	bool bTranslateRebroadcastFrames = false;
 
 #if WITH_EDITOR
 	/** Delegate when a subject is evaluated. */
diff --git a/Engine/Plugins/Animation/LiveLink/Source/LiveLink/Public/LiveLinkSettings.h b/Engine/Plugins/Animation/LiveLink/Source/LiveLink/Public/LiveLinkSettings.h
--- a/Engine/Plugins/Animation/LiveLink/Source/LiveLink/Public/LiveLinkSettings.h
+++ b/Engine/Plugins/Animation/LiveLink/Source/LiveLink/Public/LiveLinkSettings.h
@@ -124,11 +124,11 @@ public:
 
 	/** Whether to Preprocess frames before rebroadcasting them. */
 	UPROPERTY(config)
-	bool bPreProcessRebroadcastedFrames = false;
+	bool bPreProcessRebroadcastFrames = false;
 
 	/** Whether to translate frames before rebroadcasting them. */
 	UPROPERTY(config)
-	bool bTranslateRebroadcastedFrames = false;
+	bool bTranslateRebroadcastFrames = false;
 
 	/**
 	 * A source may still exist but does not send frames for a subject.
diff --git a/Engine/Plugins/Animation/LiveLink/Source/LiveLink/Public/LiveLinkSubject.h b/Engine/Plugins/Animation/LiveLink/Source/LiveLink/Public/LiveLinkSubject.h
--- a/Engine/Plugins/Animation/LiveLink/Source/LiveLink/Public/LiveLinkSubject.h
+++ b/Engine/Plugins/Animation/LiveLink/Source/LiveLink/Public/LiveLinkSubject.h
@@ -9,6 +9,7 @@
 #include "LiveLinkFrameInterpolationProcessor.h"
 #include "LiveLinkFramePreProcessor.h"
 #include "LiveLinkSourceSettings.h"
+#include "LiveLinkSubjectRemapper.h"
 
 class FLiveLinkTimedDataInput;
 class ULiveLinkSubjectSettings;
@@ -60,18 +61,24 @@ public:
 	virtual FLiveLinkSubjectKey GetSubjectKey() const override { return SubjectKey; }
 	virtual TSubclassOf<ULiveLinkRole> GetRole() const override { return Role; }
 	virtual bool HasValidFrameSnapshot() const override;
-	virtual FLiveLinkStaticDataStruct& GetStaticData() override { return StaticData; }
-	virtual const FLiveLinkStaticDataStruct& GetStaticData() const override { return StaticData; }
+	virtual FLiveLinkStaticDataStruct& GetStaticData() override;
+	virtual const FLiveLinkStaticDataStruct& GetStaticData() const override;
 	virtual TArray<FLiveLinkTime> GetFrameTimes() const override;
 	virtual const TArray<ULiveLinkFrameTranslator::FWorkerSharedPtr> GetFrameTranslators() const override
 	{
 		FScopeLock Lock(&SettingsCriticalSection);
 		return FrameTranslators;
 	}
+	virtual const ULiveLinkSubjectRemapper::FWorkerSharedPtr GetFrameRemapper() const override
+	{
+		FScopeLock Lock(&SettingsCriticalSection);
+		return SubjectRemapper;
+	}
 	virtual bool IsRebroadcasted() const override { return bRebroadcastSubject; }
 	virtual bool HasStaticDataBeenRebroadcasted() const override { return bRebroadcastStaticDataSent; }
 	virtual void SetStaticDataAsRebroadcasted(const bool bInSent) override { bRebroadcastStaticDataSent = bInSent; }
 	virtual void PreprocessFrame(FLiveLinkFrameDataStruct& InOutFrameData) override;
+
 protected:
 	virtual const FLiveLinkSubjectFrameData& GetFrameSnapshot() const override { return FrameSnapshot; }
 	//~ End ILiveLinkSubject Interface
@@ -119,6 +126,9 @@ public:
 	/** Validates if the incoming frame data is compatible with the static data for this subject. */
 	bool ValidateFrameData(const FLiveLinkFrameDataStruct& InFrameData);
 
+	/** Clear the override static data for this subject. */
+	void ClearOverrideStaticData_AnyThread();
+
 private:
 	int32 FindNewFrame_WorldTime(const FLiveLinkWorldTime& FrameTime) const;
 	int32 FindNewFrame_WorldTimeInternal(const FLiveLinkWorldTime& FrameTime) const;
@@ -166,6 +176,9 @@ protected:
 	/** List of available translator the subject can use. */
 	TArray<ULiveLinkFrameTranslator::FWorkerSharedPtr> FrameTranslators;
 
+	/** Subject remapper used to modify static and frame data for a subject. */
+	ULiveLinkSubjectRemapper::FWorkerSharedPtr SubjectRemapper;
+
 private:
 	struct FLiveLinkCachedSettings
 	{
@@ -188,6 +201,9 @@ private:
 	// Static data of the subject
 	FLiveLinkStaticDataStruct StaticData;
 
+	// Override static data, set by the remapper.
+	TOptional<FLiveLinkStaticDataStruct> OverrideStaticData;
+
 	// Frames added to the subject
 	TArray<FLiveLinkFrameDataStruct> FrameData;
 
@@ -229,6 +245,9 @@ private:
 
 	// If true, static data has been sent for this rebroadcast
 	bool bRebroadcastStaticDataSent = false;
+
+	/** Flag set to clear the override static data for a subject. */
+	std::atomic<bool> bClearOverrideStaticData = false;
 	
 	/** 
 	 * Evaluation can be done on any thread so we need to protect statistic logging 
diff --git a/Engine/Plugins/Animation/LiveLinkHub/Source/LiveLinkHub/Private/LiveLinkHubSubjectSettings.cpp b/Engine/Plugins/Animation/LiveLinkHub/Source/LiveLinkHub/Private/LiveLinkHubSubjectSettings.cpp
--- a/Engine/Plugins/Animation/LiveLinkHub/Source/LiveLinkHub/Private/LiveLinkHubSubjectSettings.cpp
+++ b/Engine/Plugins/Animation/LiveLinkHub/Source/LiveLinkHub/Private/LiveLinkHubSubjectSettings.cpp
@@ -67,7 +67,8 @@ void ULiveLinkHubSubjectSettings::PostEditChangeProperty(FPropertyChangedEvent&
 		}
 	}
 	else if (PropertyChangedEvent.GetPropertyName() == GET_MEMBER_NAME_CHECKED(ULiveLinkHubSubjectSettings, Translators)
-		|| PropertyChangedEvent.GetPropertyName() == GET_MEMBER_NAME_CHECKED(ULiveLinkHubSubjectSettings, PreProcessors))
+		|| PropertyChangedEvent.GetPropertyName() == GET_MEMBER_NAME_CHECKED(ULiveLinkHubSubjectSettings, PreProcessors)
+		|| PropertyChangedEvent.GetPropertyName() == GET_MEMBER_NAME_CHECKED(ULiveLinkHubSubjectSettings, Remapper))
 	{
 		FLiveLinkHubClient* LiveLinkClient = static_cast<FLiveLinkHubClient*>(&IModularFeatures::Get().GetModularFeature<ILiveLinkClient>(ILiveLinkClient::ModularFeatureName));
 		LiveLinkClient->CacheSubjectSettings(Key, this);
diff --git a/Engine/Programs/LiveLinkHub/Config/DefaultEditorPerProjectUserSettings.ini b/Engine/Programs/LiveLinkHub/Config/DefaultEditorPerProjectUserSettings.ini
--- a/Engine/Programs/LiveLinkHub/Config/DefaultEditorPerProjectUserSettings.ini
+++ b/Engine/Programs/LiveLinkHub/Config/DefaultEditorPerProjectUserSettings.ini
@@ -5,4 +5,8 @@ LogNormalColor=(R=0.72,G=0.72,B=0.72,A=1.000000)
 LogCommandColor=(R=0.033105,G=0.723055,B=0.033105,A=1.000000)
 LogWarningColor=(R=0.921875,G=0.691406,B=0.000000,A=1.000000)
 LogErrorColor=(R=1.000000,G=0.052083,B=0.060957,A=1.000000)
-LogFontSize=9
\ No newline at end of file
+LogFontSize=9
+
+[/Script/UnrealEd.EditorExperimentalSettings]
+bEnableAsyncStaticMeshCompilation=False
+bEnableAsyncSkinnedAssetCompilation=False
\ No newline at end of file
diff --git a/Engine/Programs/LiveLinkHub/Config/DefaultGame.ini b/Engine/Programs/LiveLinkHub/Config/DefaultGame.ini
--- a/Engine/Programs/LiveLinkHub/Config/DefaultGame.ini
+++ b/Engine/Programs/LiveLinkHub/Config/DefaultGame.ini
@@ -8,8 +8,8 @@ MessageBusPingRequestFrequency=1.000000
 MessageBusHeartbeatFrequency=1.000000
 MessageBusHeartbeatTimeout=2.000000
 MessageBusTimeBeforeRemovingInactiveSource=30.000000
-bPreProcessRebroadcastedFrames=true
-bTranslateRebroadcastedFrames=true
+bPreProcessRebroadcastFrames=true
+bTranslateRebroadcastFrames=true
 TimeWithoutFrameToBeConsiderAsInvalid=0.500000
 ValidColor=(R=0.000000,G=1.000000,B=0.000000,A=1.000000)
 InvalidColor=(R=1.000000,G=1.000000,B=0.000000,A=1.000000)
diff --git a/Engine/Source/Runtime/LiveLinkInterface/Private/LiveLinkSubjectSettings.cpp b/Engine/Source/Runtime/LiveLinkInterface/Private/LiveLinkSubjectSettings.cpp
--- a/Engine/Source/Runtime/LiveLinkInterface/Private/LiveLinkSubjectSettings.cpp
+++ b/Engine/Source/Runtime/LiveLinkInterface/Private/LiveLinkSubjectSettings.cpp
@@ -2,11 +2,15 @@
 
 #include "LiveLinkSubjectSettings.h"
 
+#include "Features/IModularFeatures.h"
+#include "ILiveLinkClient.h"
 #include "LiveLinkFrameInterpolationProcessor.h"
 #include "LiveLinkFramePreProcessor.h"
 #include "LiveLinkFrameTranslator.h"
+#include "LiveLinkSubjectRemapper.h"
 #include "LiveLinkRole.h"
 
+
 #include UE_INLINE_GENERATED_CPP_BY_NAME(LiveLinkSubjectSettings)
 
 ULiveLinkSubjectSettings::ULiveLinkSubjectSettings()
@@ -74,20 +78,55 @@ bool ULiveLinkSubjectSettings::ValidateProcessors()
 			}
 		}
 
+		if (Remapper)
+		{
+			if (!RoleClass->IsChildOf(Remapper->GetSupportedRole()))
+			{
+				UE_LOG(LogLiveLinkSubjectSettings, Warning, TEXT("Role '%s' is not supported by remapper '%s', only %s is supported"), *RoleClass->GetName(), *Remapper->GetName(), *Remapper->GetSupportedRole()->GetName());
+				Remapper = nullptr;
+				bValidProcessors = false;
+			}
+		}
+
 		return bValidProcessors;
 	}
 }
 
 #if WITH_EDITOR
+
+void ULiveLinkSubjectSettings::PreEditChange(FProperty* Property)
+{
+	if (Property && Property->GetName() == GET_MEMBER_NAME_CHECKED(ULiveLinkSubjectSettings, Remapper))
+	{
+		RemapperBeingReset = TStrongObjectPtr<ULiveLinkSubjectRemapper>(Remapper);
+	}
+}
+
 void ULiveLinkSubjectSettings::PostEditChangeChainProperty(struct FPropertyChangedChainEvent& PropertyChangedEvent)
 {
 	if (PropertyChangedEvent.Property->GetFName() == GET_MEMBER_NAME_CHECKED(ULiveLinkSubjectSettings, PreProcessors)
 	 || PropertyChangedEvent.Property->GetFName() == GET_MEMBER_NAME_CHECKED(ULiveLinkSubjectSettings, InterpolationProcessor)
-	|| PropertyChangedEvent.Property->GetFName() == GET_MEMBER_NAME_CHECKED(ULiveLinkSubjectSettings, Translators))
+	|| PropertyChangedEvent.Property->GetFName() == GET_MEMBER_NAME_CHECKED(ULiveLinkSubjectSettings, Translators)
+	|| PropertyChangedEvent.Property->GetFName() == GET_MEMBER_NAME_CHECKED(ULiveLinkSubjectSettings, Remapper))
 	{
 		ValidateProcessors();
+
+		if (Remapper && !Remapper->BoneNameMap.Num())
+		{
+			Remapper->Initialize(Key);
+		}
+
+		ILiveLinkClient& LiveLinkClient = static_cast<ILiveLinkClient&>(IModularFeatures::Get().GetModularFeature<ILiveLinkClient&>(ILiveLinkClient::ModularFeatureName));
+
+		if (!Remapper && RemapperBeingReset)
+		{
+			// Remapper got reset, restore the original static data.
+			LiveLinkClient.ClearOverrideStaticData_AnyThread(Key);
+		}
 	}
 
+	RemapperBeingReset.Reset();
+
 	Super::PostEditChangeChainProperty(PropertyChangedEvent);
 }
 #endif //WITH_EDITOR
diff --git a/Engine/Source/Runtime/LiveLinkInterface/Public/ILiveLinkClient.h b/Engine/Source/Runtime/LiveLinkInterface/Public/ILiveLinkClient.h
--- a/Engine/Source/Runtime/LiveLinkInterface/Public/ILiveLinkClient.h
+++ b/Engine/Source/Runtime/LiveLinkInterface/Public/ILiveLinkClient.h
@@ -358,6 +358,9 @@ public:
 	/** Returns true if the client has pending subject frames to process during its next tick */
 	virtual bool HasPendingSubjectFrames() = 0;
 
+	/** Claer the override static data used by the subject remapper. */
+	virtual void ClearOverrideStaticData_AnyThread(const FLiveLinkSubjectKey& InSubjectKey) = 0;
+
 	/** Notify when the list of sources has changed. */
 	virtual FSimpleMulticastDelegate& OnLiveLinkSourcesChanged() = 0;
 
diff --git a/Engine/Source/Runtime/LiveLinkInterface/Public/ILiveLinkSubject.h b/Engine/Source/Runtime/LiveLinkInterface/Public/ILiveLinkSubject.h
--- a/Engine/Source/Runtime/LiveLinkInterface/Public/ILiveLinkSubject.h
+++ b/Engine/Source/Runtime/LiveLinkInterface/Public/ILiveLinkSubject.h
@@ -6,9 +6,12 @@
 #include "CoreMinimal.h"
 #include "LiveLinkFrameTranslator.h"
 #include "LiveLinkRole.h"
+#include "LiveLinkSubjectRemapper.h"
 #include "LiveLinkTypes.h"
 #include "Templates/SubclassOf.h"
 
+struct FLiveLinkAnimationFrameData;
+struct FLiveLinkSkeletonStaticData;
 class ILiveLinkClient;
 class ULiveLinkFrameInterpolationProcessor;
 class ULiveLinkFramePreProcessor;
@@ -41,6 +44,8 @@ public:
 	/** List of available translator the subject can use. */
 	virtual const TArray<ULiveLinkFrameTranslator::FWorkerSharedPtr> GetFrameTranslators() const = 0;
 
+	/** Get this subject's frame remapper. */
+	virtual const ULiveLinkSubjectRemapper::FWorkerSharedPtr GetFrameRemapper() const = 0;
 
 	/** Whether this subject is rebroadcasted */
 	virtual bool IsRebroadcasted() const = 0;
@@ -54,6 +59,9 @@ public:
 	/** Apply this subject's preprocessors to a frame data. */
 	virtual void PreprocessFrame(FLiveLinkFrameDataStruct& InOutFrameData) {};
 
+	/** Apply a remapper to a frame data. Called after preprocessing.*/
+	virtual void RemapFrame(FLiveLinkSkeletonStaticData& InOutSkeletonData, FLiveLinkAnimationFrameData& InOutFrameData) {}
+
 protected:
 	virtual const FLiveLinkSubjectFrameData& GetFrameSnapshot() const = 0;
 
diff --git a/Engine/Source/Runtime/LiveLinkInterface/Public/LiveLinkSubjectRemapper.h b/Engine/Source/Runtime/LiveLinkInterface/Public/LiveLinkSubjectRemapper.h
new file mode 100644
--- /dev/null
+++ b/Engine/Source/Runtime/LiveLinkInterface/Public/LiveLinkSubjectRemapper.h
@@ -0,0 +1,67 @@
+// Copyright Epic Games, Inc. All Rights Reserved.
+
+#pragma once
+
+#include "UObject/Object.h"
+
+#include "LiveLinkRole.h"
+#include "LiveLinkTypes.h"
+#include "UObject/ObjectMacros.h"
+
+#include "LiveLinkSubjectRemapper.generated.h"
+
+/**
+ * Basic object to transform incoming static and frame data for a subject.
+ * @note Can be called from any thread.
+ */
+class ILiveLinkSubjectRemapperWorker
+{
+public:
+	virtual ~ILiveLinkSubjectRemapperWorker() = default;
+
+	/** Remap the static data of a subject. Can be used to modify bone names and bone parents. */
+	LIVELINKINTERFACE_API virtual void RemapStaticData(FLiveLinkStaticDataStruct& InOutStaticData) { }
+
+	/** Remap the frame data of a subject. */
+	LIVELINKINTERFACE_API virtual void RemapFrameData(const FLiveLinkStaticDataStruct& InStaticData, FLiveLinkFrameDataStruct& InOutFrameData) { }
+};
+
+/** Class used to remap livelink subjects without having to rely on animation blueprints. */
+UCLASS(Abstract, MinimalAPI)
+class ULiveLinkSubjectRemapper : public UObject
+{
+public:
+	GENERATED_BODY()
+
+	using FWorkerSharedPtr = TSharedPtr<ILiveLinkSubjectRemapperWorker>;
+
+	/**
+	 * Create an instance of a ILiveLinkSubjectRemapperWorker that can be used outside of the game thread.
+	 */
+	LIVELINKINTERFACE_API virtual FWorkerSharedPtr CreateWorker() PURE_VIRTUAL(ULiveLinkSubjectRemapper::CreateWorker, return FWorkerSharedPtr(););
+
+	/**
+	 * Get the instance of ILiveLinkSubjectRemapperWorker that was created by the CreateWorker method. 
+	 * Returns an invalid pointer if no instance was created.
+	 */
+	LIVELINKINTERFACE_API virtual FWorkerSharedPtr GetWorker() const PURE_VIRTUAL(ULiveLinkSubjectRemapper::GetWorker, return FWorkerSharedPtr(););
+
+	/**
+	 * Called to initialize the remapper with information from the subject that owns the remapper.
+	 */
+	LIVELINKINTERFACE_API virtual void Initialize(const FLiveLinkSubjectKey& SubjectKey) {}
+
+	/** Get what role is supported by this remapper. */
+	LIVELINKINTERFACE_API virtual TSubclassOf<ULiveLinkRole> GetSupportedRole() const PURE_VIRTUAL(ULiveLinkSubjectRemapper::GetSupportedRole, return TSubclassOf<ULiveLinkRole>(););
+	
+	/** Retuns whether the remapper can currently be used for remapping. */
+	LIVELINKINTERFACE_API virtual bool IsValidRemapper() const PURE_VIRTUAL(ULiveLinkSubjectRemapper::IsValid, return false;);
+
+public:
+	/** Name mapping between source bone name and transformed bone name */
+	UPROPERTY(EditAnywhere, Category="Remapper")
+	TMap<FName, FName> BoneNameMap;
+
+	/** When this is true, livelink's buffered frames will be remapped and the subject's static data will be updated. */
+	bool bDirty = false;
+};
diff --git a/Engine/Source/Runtime/LiveLinkInterface/Public/LiveLinkSubjectSettings.h b/Engine/Source/Runtime/LiveLinkInterface/Public/LiveLinkSubjectSettings.h
--- a/Engine/Source/Runtime/LiveLinkInterface/Public/LiveLinkSubjectSettings.h
+++ b/Engine/Source/Runtime/LiveLinkInterface/Public/LiveLinkSubjectSettings.h
@@ -20,6 +20,7 @@ class FLiveLinkTimedDataInput;
 class ULiveLinkFrameInterpolationProcessor;
 class ULiveLinkFramePreProcessor;
 class ULiveLinkFrameTranslator;
+class ULiveLinkSubjectRemapper;
 class ULiveLinkRole;
 
 /**
@@ -78,6 +79,10 @@ public:
 	UPROPERTY(EditAnywhere, Instanced, Category = "LiveLink", meta = (DisplayName = "Translators"))
 	TArray<TObjectPtr<ULiveLinkFrameTranslator>> Translators;
 
+	/** Remapper used to modify incoming static and frame data for a subject. */
+	UPROPERTY(EditAnywhere, Instanced, Category = "LiveLink")
+	TObjectPtr<ULiveLinkSubjectRemapper> Remapper;
+
 	UPROPERTY()
 	TSubclassOf<ULiveLinkRole> Role;
 
@@ -90,7 +95,7 @@ public:
 	bool bAllowModifyingRebroadcast = true;
 	
 	/** If enabled, rebroadcast this subject */
-	UPROPERTY(EditAnywhere, Category = "LiveLink", meta=(EditConditionHides, EditCondition="bAllowModifyingRebroadcast"))
+	UPROPERTY(EditAnywhere, Category = "LiveLink", meta=(EditCondition="bAllowModifyingRebroadcast"))
     bool bRebroadcastSubject;
 
 	/** Validate PreProcessors, Translators and Interpolation processors. Usually called after a property change event.
@@ -101,6 +106,7 @@ public:
 public:
 	//~ Begin UObject interface
 #if WITH_EDITOR
+	LIVELINKINTERFACE_API virtual void PreEditChange(FProperty* Property) override;
 	LIVELINKINTERFACE_API virtual void PostEditChangeChainProperty(struct FPropertyChangedChainEvent& PropertyChangedEvent) override;
 #endif
 	//~ End UObject interface
@@ -109,4 +115,8 @@ protected:
 	/** Key of the subject that owns this setting. */
 	UPROPERTY()
 	FLiveLinkSubjectKey Key;
+
+private:
+	/** We need to keep track of the remapper when it's reset in order to restore the static data. */
+	TStrongObjectPtr<ULiveLinkSubjectRemapper> RemapperBeingReset;
 };
diff --git a/Engine/Source/Runtime/LiveLinkInterface/Public/LiveLinkVirtualSubject.h b/Engine/Source/Runtime/LiveLinkInterface/Public/LiveLinkVirtualSubject.h
--- a/Engine/Source/Runtime/LiveLinkInterface/Public/LiveLinkVirtualSubject.h
+++ b/Engine/Source/Runtime/LiveLinkInterface/Public/LiveLinkVirtualSubject.h
@@ -40,6 +40,7 @@ public:
 	virtual FLiveLinkStaticDataStruct& GetStaticData() override { return CurrentFrameSnapshot.StaticData; }
 	virtual const FLiveLinkStaticDataStruct& GetStaticData() const override { return CurrentFrameSnapshot.StaticData; }
 	virtual const TArray<ULiveLinkFrameTranslator::FWorkerSharedPtr> GetFrameTranslators() const override { return CurrentFrameTranslators; }
+	virtual const ULiveLinkSubjectRemapper::FWorkerSharedPtr GetFrameRemapper() const override { return CurrentSubjectRemapper; }
 	LIVELINKINTERFACE_API virtual TArray<FLiveLinkTime> GetFrameTimes() const override;
 	virtual bool IsRebroadcasted() const override { return bRebroadcastSubject; }
 	virtual bool HasStaticDataBeenRebroadcasted() const override { return bHasStaticDataBeenRebroadcast; }
@@ -125,6 +126,8 @@ protected:
 
 private:
 	TArray<ULiveLinkFrameTranslator::FWorkerSharedPtr> CurrentFrameTranslators;
+	/** Current subject remapper. */
+	ULiveLinkSubjectRemapper::FWorkerSharedPtr CurrentSubjectRemapper;
 
 	/** Last evaluated frame for this subject. */
 	FLiveLinkSubjectFrameData CurrentFrameSnapshot;
