LiveLink - Fix DCS emulator not working with LiveLinkHubEditor

diff --git a/Engine/Plugins/Animation/LiveLink/Source/LiveLink/Private/LiveLinkClient.cpp b/Engine/Plugins/Animation/LiveLink/Source/LiveLink/Private/LiveLinkClient.cpp
--- a/Engine/Plugins/Animation/LiveLink/Source/LiveLink/Private/LiveLinkClient.cpp
+++ b/Engine/Plugins/Animation/LiveLink/Source/LiveLink/Private/LiveLinkClient.cpp
@@ -1552,7 +1552,8 @@ FName FLiveLinkClient::GetRebroadcastName(const FLiveLinkSubjectKey& InSubjectKe
 {
 	if (ULiveLinkSubjectSettings* Settings = Cast<ULiveLinkSubjectSettings>(GetSubjectSettings(InSubjectKey)))
 	{
-		return Settings->GetRebroadcastName();
+		FName RebroadcastName = Settings->GetRebroadcastName();
+		return RebroadcastName.IsNone() ? InSubjectKey.SubjectName.Name : RebroadcastName;
 	}
 
 	return InSubjectKey.SubjectName;
