LiveLink/Performance Capture
- Added bEvaluateLiveLink uproperty to livelinktransform controller so users can pause animation from details panel and from BP.

diff --git a/Engine/Plugins/Animation/LiveLink/Source/LiveLinkComponents/Private/Controllers/LiveLinkTransformController.cpp b/Engine/Plugins/Animation/LiveLink/Source/LiveLinkComponents/Private/Controllers/LiveLinkTransformController.cpp
--- a/Engine/Plugins/Animation/LiveLink/Source/LiveLinkComponents/Private/Controllers/LiveLinkTransformController.cpp
+++ b/Engine/Plugins/Animation/LiveLink/Source/LiveLinkComponents/Private/Controllers/LiveLinkTransformController.cpp
@@ -20,7 +20,7 @@
 
 void FLiveLinkTransformControllerData::ApplyTransform(USceneComponent* SceneComponent, const FTransform& Transform, const FLiveLinkTransformStaticData& StaticData) const
 {
-	if (SceneComponent)
+	if (SceneComponent && bEvaluateLiveLink)
 	{
 		if (bUseLocation && StaticData.bIsLocationSupported)
 		{
diff --git a/Engine/Plugins/Animation/LiveLink/Source/LiveLinkComponents/Public/Controllers/LiveLinkTransformController.h b/Engine/Plugins/Animation/LiveLink/Source/LiveLinkComponents/Public/Controllers/LiveLinkTransformController.h
--- a/Engine/Plugins/Animation/LiveLink/Source/LiveLinkComponents/Public/Controllers/LiveLinkTransformController.h
+++ b/Engine/Plugins/Animation/LiveLink/Source/LiveLinkComponents/Public/Controllers/LiveLinkTransformController.h
@@ -15,6 +15,10 @@ struct LIVELINKCOMPONENTS_API FLiveLinkTransformControllerData
 {
 	GENERATED_BODY()
 
+	/** Whether we should evaluate animation from Live Link */
+	UPROPERTY(EditAnywhere, BlueprintReadWrite, Category = "LiveLink")
+	bool bEvaluateLiveLink = true;
+
 	/** Set the transform of the component in world space of in its local reference frame. */
 	UPROPERTY(EditAnywhere, BlueprintReadOnly, Category = "LiveLink")
 	bool bWorldTransform = false;
