Remove call to FPlatformProcess::Sleep from LiveLinkMessageBusDiscoveryManager and use an event instead to speed up the shutdown process.

diff --git a/Engine/Plugins/Animation/LiveLink/Source/LiveLink/Private/LiveLinkMessageBusDiscoveryManager.cpp b/Engine/Plugins/Animation/LiveLink/Source/LiveLink/Private/LiveLinkMessageBusDiscoveryManager.cpp
--- a/Engine/Plugins/Animation/LiveLink/Source/LiveLink/Private/LiveLinkMessageBusDiscoveryManager.cpp
+++ b/Engine/Plugins/Animation/LiveLink/Source/LiveLink/Private/LiveLinkMessageBusDiscoveryManager.cpp
@@ -2,6 +2,7 @@
 
 #include "LiveLinkMessageBusDiscoveryManager.h"
 
+#include "HAL/PlatformProcess.h"
 #include "HAL/RunnableThread.h"
 #include "ILiveLinkClient.h"
 #include "LiveLinkMessages.h"
@@ -10,8 +11,6 @@
 
 LLM_DEFINE_TAG(LiveLink_LiveLinkMessageBusDiscoveryManager);
 
-#define LL_HEARTBEAT_SLEEP_TIME 1.0f
-
 FLiveLinkMessageBusDiscoveryManager::FLiveLinkMessageBusDiscoveryManager()
 	: bRunning(true)
 	, Thread(nullptr)
@@ -19,7 +18,9 @@ FLiveLinkMessageBusDiscoveryManager::FLiveLinkMessageBusDiscoveryManager()
 	LLM_SCOPE_BYTAG(LiveLink_LiveLinkMessageBusDiscoveryManager);
 
 	PingRequestCounter = 0;
-	PingRequestFrequency = GetDefault<ULiveLinkSettings>()->GetMessageBusPingRequestFrequency();
+	PingRequestFrequency = FTimespan::FromSeconds(GetDefault<ULiveLinkSettings>()->GetMessageBusPingRequestFrequency());
+
+	PollEvent = FPlatformProcess::GetSynchEventFromPool();
 
 	MessageEndpoint = FMessageEndpoint::Builder(TEXT("LiveLinkMessageHeartbeatManager"))
 		.Handling<FLiveLinkPongMessage>(this, &FLiveLinkMessageBusDiscoveryManager::HandlePongMessage);
@@ -44,13 +45,14 @@ FLiveLinkMessageBusDiscoveryManager::~FLiveLinkMessageBusDiscoveryManager()
 		}
 	}
 
-	Stop();
-
 	if (Thread)
 	{
 		Thread->Kill(true);
 		Thread = nullptr;
 	}
+
+	FPlatformProcess::ReturnSynchEventToPool(PollEvent);
+	PollEvent = nullptr;
 }
 
 uint32 FLiveLinkMessageBusDiscoveryManager::Run()
@@ -68,8 +70,7 @@ uint32 FLiveLinkMessageBusDiscoveryManager::Run()
 				MessageEndpoint->Publish(FMessageEndpoint::MakeMessage<FLiveLinkPingMessage>(LastPingRequest, Version));
 			}
 		}
-
-		FPlatformProcess::Sleep(PingRequestFrequency);
+		PollEvent->Wait(PingRequestFrequency.GetTotalMilliseconds());
 	}
 	return 0;
 }
@@ -77,6 +78,7 @@ uint32 FLiveLinkMessageBusDiscoveryManager::Run()
 void FLiveLinkMessageBusDiscoveryManager::Stop()
 {
 	bRunning = false;
+	PollEvent->Trigger();
 }
 
 void FLiveLinkMessageBusDiscoveryManager::AddDiscoveryMessageRequest()
diff --git a/Engine/Plugins/Animation/LiveLink/Source/LiveLink/Public/LiveLinkMessageBusDiscoveryManager.h b/Engine/Plugins/Animation/LiveLink/Source/LiveLink/Public/LiveLinkMessageBusDiscoveryManager.h
--- a/Engine/Plugins/Animation/LiveLink/Source/LiveLink/Public/LiveLinkMessageBusDiscoveryManager.h
+++ b/Engine/Plugins/Animation/LiveLink/Source/LiveLink/Public/LiveLinkMessageBusDiscoveryManager.h
@@ -46,7 +46,7 @@ private:
 	double LastPingRequestTime;
 
 	// Ping request timeout
-	double PingRequestFrequency;
+	FTimespan PingRequestFrequency;
 
 	// Result from the last ping request
 	TArray<FProviderPollResultPtr> LastProviderPoolResults;
@@ -59,6 +59,9 @@ private:
 	// Thread the heartbeats are sent on
 	FRunnableThread* Thread;
 
+	// Event used to poll the discovery results.
+	class FEvent* PollEvent = nullptr;
+
 	// Critical section for accessing the Source Set
 	mutable FCriticalSection SourcesCriticalSection;
 };
