Move LiveLink files that were previously moved to the wrong location

diff --git a/Engine/Plugins/Animation/LiveLink/Source/LiveLink/Private/Remapper/LiveLinkSkeletonRemapper.h b/Engine/Plugins/Animation/LiveLink/Source/LiveLink/Private/Remapper/LiveLinkSkeletonRemapper.h
deleted file mode 100644
diff --git a/Engine/Plugins/Animation/LiveLink/Source/LiveLink/Public/Remapper/LiveLinkSkeletonRemapper.h b/Engine/Plugins/Animation/LiveLink/Source/LiveLink/Public/Remapper/LiveLinkSkeletonRemapper.h
new file mode 100644
--- /dev/null
+++ b/Engine/Plugins/Animation/LiveLink/Source/LiveLink/Public/Remapper/LiveLinkSkeletonRemapper.h
@@ -0,0 +1,100 @@
+// Copyright Epic Games, Inc. All Rights Reserved.
+
+#pragma once
+
+
+#include "LiveLinkSubjectRemapper.h"
+
+#include "Engine/SkeletalMesh.h"
+#include "Features/IModularFeatures.h"
+#include "LiveLinkClient.h"
+#include "LiveLinkTypes.h"
+#include "Roles/LiveLinkAnimationRole.h"
+#include "Roles/LiveLinkAnimationTypes.h"
+
+#include "LiveLinkSkeletonRemapper.generated.h"
+
+class LIVELINK_API FLiveLinkSkeletonRemapperWorker : public ILiveLinkSubjectRemapperWorker
+{
+public:
+	virtual void RemapStaticData(FLiveLinkStaticDataStruct& InOutStaticData)
+	{
+		if (InOutStaticData.GetStruct() && InOutStaticData.GetStruct()->IsChildOf<FLiveLinkSkeletonStaticData>())
+		{
+			RemapSkeletonStaticData(*InOutStaticData.Cast<FLiveLinkSkeletonStaticData>());
+		}
+	}
+
+	virtual void RemapFrameData(const FLiveLinkStaticDataStruct& InOutStaticData, FLiveLinkFrameDataStruct& InOutFrameData)
+	{
+		if (InOutStaticData.GetStruct() && InOutStaticData.GetStruct()->IsChildOf<FLiveLinkSkeletonStaticData>())
+		{
+			RemapSkeletonFrameData(*InOutStaticData.Cast<FLiveLinkSkeletonStaticData>(), *InOutFrameData.Cast<FLiveLinkAnimationFrameData>());
+		}
+	}
+
+	virtual void RemapSkeletonStaticData(FLiveLinkSkeletonStaticData& InOutSkeletonData)
+	{
+	}
+
+	virtual void RemapSkeletonFrameData(const FLiveLinkSkeletonStaticData& InOutSkeletonData, FLiveLinkAnimationFrameData& InOutFrameData)
+	{
+	}
+
+	FName GetRemappedBoneName(FName BoneName) const
+	{
+		if (const FName* RemappedName = BoneNameMap.Find(BoneName))
+		{
+			return *RemappedName;
+		}
+		return BoneName;
+	}
+
+	/** Map used to provide new names for the bones in the static data. */
+	TMap<FName, FName> BoneNameMap;
+};
+
+UCLASS(Abstract)
+class LIVELINK_API ULiveLinkSkeletonRemapper : public ULiveLinkSubjectRemapper
+{
+public:
+	GENERATED_BODY()
+
+	virtual void Initialize(const FLiveLinkSubjectKey& SubjectKey) override
+	{
+		ILiveLinkClient& LiveLinkClient = IModularFeatures::Get().GetModularFeature<ILiveLinkClient>(ILiveLinkClient::ModularFeatureName);
+		const FLiveLinkStaticDataStruct* StaticData = LiveLinkClient.GetSubjectStaticData_AnyThread(SubjectKey);
+
+		TSubclassOf<ULiveLinkRole> LiveLinkRole = LiveLinkClient.GetSubjectRole_AnyThread(SubjectKey);
+
+		// Note: Should we initialize the bone name map using the reference skeleton?
+		if (StaticData && LiveLinkRole && LiveLinkRole->IsChildOf(ULiveLinkAnimationRole::StaticClass()))
+		{
+			for (FName BoneName : StaticData->Cast<FLiveLinkSkeletonStaticData>()->GetBoneNames())
+			{
+				BoneNameMap.Add(BoneName, BoneName);
+			}
+		}
+	}
+
+	virtual TSubclassOf<ULiveLinkRole> GetSupportedRole() const override
+	{
+		return ULiveLinkAnimationRole::StaticClass();
+	}
+
+#if WITH_EDITOR
+	virtual void PostEditChangeProperty(struct FPropertyChangedEvent& PropertyChangedEvent) override
+	{
+		bDirty = true;
+	}
+#endif
+
+	virtual bool IsValidRemapper() const
+	{
+		return ReferenceSkeleton.IsValid();
+	}
+
+public:
+	UPROPERTY(EditAnywhere, Category="Remapper", meta=(DisplayThumbnail="false"))
+	TSoftObjectPtr<USkeletalMesh> ReferenceSkeleton;
+};
