LiveLink - Append basic frame data when combining subjects in an animation virtual subject.

This will enable workflows like combining ARkit's face tracking with a body motion capture tracking subject.

diff --git a/Engine/Plugins/Animation/LiveLink/Source/LiveLink/Private/LiveLinkAnimationVirtualSubject.cpp b/Engine/Plugins/Animation/LiveLink/Source/LiveLink/Private/LiveLinkAnimationVirtualSubject.cpp
--- a/Engine/Plugins/Animation/LiveLink/Source/LiveLink/Private/LiveLinkAnimationVirtualSubject.cpp
+++ b/Engine/Plugins/Animation/LiveLink/Source/LiveLink/Private/LiveLinkAnimationVirtualSubject.cpp
@@ -117,7 +117,15 @@ bool ULiveLinkAnimationVirtualSubject::AreSubjectsValid(const TArray<FLiveLinkSu
 				return (SubjectData.SubjectName == SubjectName);
 			});
 
-		bValid = FoundPtr != nullptr && LiveLinkClient->DoesSubjectSupportsRole_AnyThread(*FoundPtr, GetRole());
+
+		bValid = false;
+
+		if (FoundPtr)
+		{
+			TSubclassOf<ULiveLinkRole> SubjectRole = LiveLinkClient->GetSubjectRole_AnyThread(*FoundPtr);
+			bValid = SubjectRole && (SubjectRole->IsChildOf(ULiveLinkAnimationRole::StaticClass()) || SubjectRole == ULiveLinkBasicRole::StaticClass());
+		}
+
 		if (!bValid)
 		{
 			break;
@@ -136,7 +144,11 @@ bool ULiveLinkAnimationVirtualSubject::BuildSubjectSnapshot(TArray<FLiveLinkSubj
 	for (const FName& SubjectName : Subjects)
 	{
 		FLiveLinkSubjectFrameData& NextSnapshot = OutSnapshot.AddDefaulted_GetRef();
-		if (!LiveLinkClient->EvaluateFrame_AnyThread(SubjectName, GetRole(), NextSnapshot))
+		TSubclassOf<ULiveLinkRole> SubjectRole = LiveLinkClient->GetSubjectRole_AnyThread(SubjectName);
+		
+		// Evaluate for basic role if that's the subject's role.
+		TSubclassOf<ULiveLinkRole> DesiredRole = SubjectRole && SubjectRole == ULiveLinkBasicRole::StaticClass() ? SubjectRole : GetRole();
+		if (!LiveLinkClient->EvaluateFrame_AnyThread(SubjectName, DesiredRole, NextSnapshot))
 		{
 			bSnapshotDone = false;
 			break;
@@ -163,23 +175,29 @@ void ULiveLinkAnimationVirtualSubject::BuildSkeleton(const TArray<FLiveLinkSubje
 
 		TMap<int32, FName> NamesToOverride;
 
+		TArray<FName> CombinedPropertyNames;
+
 		for (int32 Index = 0; Index < InSubjectSnapshots.Num(); ++Index)
 		{
 			const FLiveLinkSubjectFrameData& SubjectSnapShotData = InSubjectSnapshots[Index];
 			check(SubjectSnapShotData.StaticData.IsValid());
-			const FLiveLinkSkeletonStaticData* SubjectSkeletonData = SubjectSnapShotData.StaticData.Cast<FLiveLinkSkeletonStaticData>();
 
-			const FName BonePrefix = bAppendSubjectNameToBones ? Subjects[Index] : NAME_None;
+			if (const FLiveLinkSkeletonStaticData* SubjectSkeletonData = SubjectSnapShotData.StaticData.Cast<FLiveLinkSkeletonStaticData>())
+			{
+				const FName BonePrefix = bAppendSubjectNameToBones ? Subjects[Index] : NAME_None;
 
-			LiveLinkAnimationVirtualSubjectUtils::AddToBoneNames(BoneNames, SubjectSkeletonData->GetBoneNames(), BonePrefix, NamesToOverride);
-			LiveLinkAnimationVirtualSubjectUtils::AddToBoneParents(BoneParents, SubjectSkeletonData->GetBoneParents());
+				LiveLinkAnimationVirtualSubjectUtils::AddToBoneNames(BoneNames, SubjectSkeletonData->GetBoneNames(), BonePrefix, NamesToOverride);
+				LiveLinkAnimationVirtualSubjectUtils::AddToBoneParents(BoneParents, SubjectSkeletonData->GetBoneParents());
 
-			// Cache bone names to bone index, we need to use both the subject and bone names to make sure
-			// that we can find a parent bone even if it's in a name conflict.
-			for (int32 BoneIndex = 0; BoneIndex < BoneNames.Num(); BoneIndex++)
-			{
-				BoneNameToIndex.FindOrAdd({ Subjects[Index], BoneNames[BoneIndex] }) = BoneIndex;
+				// Cache bone names to bone index, we need to use both the subject and bone names to make sure
+				// that we can find a parent bone even if it's in a name conflict.
+				for (int32 BoneIndex = 0; BoneIndex < BoneNames.Num(); BoneIndex++)
+				{
+					BoneNameToIndex.FindOrAdd({ Subjects[Index], BoneNames[BoneIndex] }) = BoneIndex;
+				}
 			}
+
+			CombinedPropertyNames.Append(SubjectSnapShotData.StaticData.GetBaseData()->PropertyNames);
 		}
 
 		ProcessAttachmentsForStaticData(BoneParents);
@@ -191,6 +209,7 @@ void ULiveLinkAnimationVirtualSubject::BuildSkeleton(const TArray<FLiveLinkSubje
 
 		SkeletonData->SetBoneNames(BoneNames);
 		SkeletonData->SetBoneParents(BoneParents);
+		SkeletonData->PropertyNames = MoveTemp(CombinedPropertyNames);
 
 		UpdateStaticDataSnapshot(MoveTemp(StaticData));
 
@@ -213,18 +232,22 @@ void ULiveLinkAnimationVirtualSubject::BuildFrame(const TArray<FLiveLinkSubjectF
 
 	NewSnapshotFrameData->Transforms.Reset(SnapshotSkeletonData->GetBoneNames().Num());
 	NewSnapshotFrameData->MetaData.StringMetaData.Empty();
+	NewSnapshotFrameData->PropertyValues.Empty();
 
 	//Go over each subject snapshot and take transforms and curves
 	check(InSubjectSnapshots.Num() == Subjects.Num());
 	for (int32 Index = 0; Index < InSubjectSnapshots.Num(); ++Index)
 	{
 		const FLiveLinkSubjectFrameData& SubjectSnapShotData = InSubjectSnapshots[Index];
+
 		check(SubjectSnapShotData.FrameData.IsValid());
-		const FLiveLinkAnimationFrameData* SubjectFrameData = SubjectSnapShotData.FrameData.Cast<FLiveLinkAnimationFrameData>();
+		if (const FLiveLinkAnimationFrameData* SubjectAnimationData = SubjectSnapShotData.FrameData.Cast<FLiveLinkAnimationFrameData>())
+		{
+			NewSnapshotFrameData->Transforms.Append(SubjectAnimationData->Transforms);
+		}
 
-		NewSnapshotFrameData->Transforms.Append(SubjectFrameData->Transforms);
-		NewSnapshotFrameData->PropertyValues.Append(SubjectFrameData->PropertyValues);
-		for (const auto& MetaDatum : SubjectFrameData->MetaData.StringMetaData)
+		NewSnapshotFrameData->PropertyValues.Append(SubjectSnapShotData.FrameData.GetBaseData()->PropertyValues);
+		for (const auto& MetaDatum : SubjectSnapShotData.FrameData.GetBaseData()->MetaData.StringMetaData)
 		{
 			NewSnapshotFrameData->MetaData.StringMetaData.Emplace(Subjects[Index], MetaDatum.Value);
 		}
diff --git a/Engine/Plugins/Animation/LiveLink/Source/LiveLink/Public/LiveLinkAnimationVirtualSubject.h b/Engine/Plugins/Animation/LiveLink/Source/LiveLink/Public/LiveLinkAnimationVirtualSubject.h
--- a/Engine/Plugins/Animation/LiveLink/Source/LiveLink/Public/LiveLinkAnimationVirtualSubject.h
+++ b/Engine/Plugins/Animation/LiveLink/Source/LiveLink/Public/LiveLinkAnimationVirtualSubject.h
@@ -23,7 +23,7 @@ enum class EBoneTransformResolution
 	Combine
 };
 
-/** A Skeleton virtual subject is an assembly of different subjects supporting the animation role */
+/** A Skeleton virtual subject is an assembly of different subjects supporting the animation or basic role */
 UCLASS(meta=(DisplayName="Animation Virtual Subject"))
 class LIVELINK_API ULiveLinkAnimationVirtualSubject : public ULiveLinkVirtualSubject
 {
diff --git a/Engine/Plugins/Animation/LiveLink/Source/LiveLinkEditor/Private/LiveLinkVirtualSubjectDetailCustomization.cpp b/Engine/Plugins/Animation/LiveLink/Source/LiveLinkEditor/Private/LiveLinkVirtualSubjectDetailCustomization.cpp
--- a/Engine/Plugins/Animation/LiveLink/Source/LiveLinkEditor/Private/LiveLinkVirtualSubjectDetailCustomization.cpp
+++ b/Engine/Plugins/Animation/LiveLink/Source/LiveLinkEditor/Private/LiveLinkVirtualSubjectDetailCustomization.cpp
@@ -6,7 +6,10 @@
 #include "DetailWidgetRow.h"
 #include "Framework/Views/TableViewMetadata.h"
 #include "ILiveLinkClient.h"
+#include "LiveLinkRole.h"
 #include "LiveLinkVirtualSubject.h"
+#include "Roles/LiveLinkAnimationRole.h"
+#include "Roles/LiveLinkBasicRole.h"
 #include "ScopedTransaction.h"
 #include "Widgets/Input/SCheckBox.h"
 
@@ -50,10 +53,14 @@ void FLiveLinkVirtualSubjectDetailCustomization::CustomizeDetails(IDetailLayoutB
 
 	if (Client)
 	{
-		TArray<FLiveLinkSubjectKey> SubjectKeys = Client->GetSubjectsSupportingRole(Subject->GetRole(), /*bIncludeDisabledSubject*/ false, /*bIncludeVirtualSubject*/ false);
+		TArray<FLiveLinkSubjectKey> SubjectKeys = Client->GetSubjects(/*bIncludeDisabledSubject*/ false, /*bIncludeVirtualSubject*/ false);
 		for (const FLiveLinkSubjectKey& SubjectKey : SubjectKeys)
 		{
-			SubjectsListItems.AddUnique(MakeShared<FName>(SubjectKey.SubjectName.Name));
+			TSubclassOf<ULiveLinkRole> LiveLinkRole = Client->GetSubjectRole_AnyThread(SubjectKey);
+			if (LiveLinkRole && (LiveLinkRole->IsChildOf(ULiveLinkAnimationRole::StaticClass()) || LiveLinkRole == ULiveLinkBasicRole::StaticClass()))
+			{
+				SubjectsListItems.AddUnique(MakeShared<FName>(SubjectKey.SubjectName.Name));
+			}
 		}
 	}
 
@@ -165,8 +172,10 @@ FSlateColor FLiveLinkVirtualSubjectDetailCustomization::HandleSubjectItemColor(F
 	if (ULiveLinkVirtualSubject* Subject = SubjectPtr.Get())
 	{
 		const FName ThisItem = *InItem.Get();
-		TArray<FLiveLinkSubjectKey> SubjectKeys = Client->GetSubjectsSupportingRole(Subject->GetRole(), /*bIncludeDisabledSubject*/ false, /*bIncludeVirtualSubject*/ false);
-		if (false == SubjectKeys.ContainsByPredicate([ThisItem](const FLiveLinkSubjectKey& Other) { return Other.SubjectName.Name == ThisItem; }))
+
+
+		TArray<FLiveLinkSubjectKey> SubjectKeys = Client->GetSubjects(/*bIncludeDisabledSubject*/ false, /*bIncludeVirtualSubject*/ false);
+		if (!SubjectKeys.ContainsByPredicate([ThisItem](const FLiveLinkSubjectKey& Other) { return Other.SubjectName.Name == ThisItem; }))
 		{
 			Result = FLinearColor::Red;
 		}
@@ -182,8 +191,8 @@ FText FLiveLinkVirtualSubjectDetailCustomization::HandleSubjectItemToolTip(FSubj
 	if (ULiveLinkVirtualSubject* Subject = SubjectPtr.Get())
 	{
 		const FName ThisItem = *InItem.Get();
-		TArray<FLiveLinkSubjectKey> SubjectKeys = Client->GetSubjectsSupportingRole(Subject->GetRole(), /*bIncludeDisabledSubject*/ false, /*bIncludeVirtualSubject*/ false);
-		if (false == SubjectKeys.ContainsByPredicate([ThisItem](const FLiveLinkSubjectKey& Other) { return Other.SubjectName.Name == ThisItem; }))
+		TArray<FLiveLinkSubjectKey> SubjectKeys = Client->GetSubjects(/*bIncludeDisabledSubject*/ false, /*bIncludeVirtualSubject*/ false);
+		if (!SubjectKeys.ContainsByPredicate([ThisItem](const FLiveLinkSubjectKey& Other) { return Other.SubjectName.Name == ThisItem; }))
 		{
 			Result = LOCTEXT("LinkedSubjectToolTip", "This subject was not found in the list of available LiveLink subjects. VirtualSubject might not work properly.");
 		}
