LiveLink
- Adds  a user exposed Offset transform to the controller component. The new member applies an offset transform in the local space of the controlled scene component.

[RN] Added an offset transform to LiveLinkTransformControllers

diff --git a/Engine/Plugins/Animation/LiveLink/Source/LiveLinkComponents/Private/Controllers/LiveLinkTransformController.cpp b/Engine/Plugins/Animation/LiveLink/Source/LiveLinkComponents/Private/Controllers/LiveLinkTransformController.cpp
--- a/Engine/Plugins/Animation/LiveLink/Source/LiveLinkComponents/Private/Controllers/LiveLinkTransformController.cpp
+++ b/Engine/Plugins/Animation/LiveLink/Source/LiveLinkComponents/Private/Controllers/LiveLinkTransformController.cpp
@@ -20,6 +20,7 @@
 
 void FLiveLinkTransformControllerData::ApplyTransform(USceneComponent* SceneComponent, const FTransform& Transform, const FLiveLinkTransformStaticData& StaticData) const
 {
+	
 	if (SceneComponent)
 	{
 		if (bUseLocation && StaticData.bIsLocationSupported)
@@ -98,7 +99,8 @@ void ULiveLinkTransformController::Tick(float DeltaTime, const FLiveLinkSubjectF
 	{
 		if (USceneComponent* SceneComponent = Cast<USceneComponent>(GetAttachedComponent()))
 		{
-			TransformData.ApplyTransform(SceneComponent, FrameData->Transform, *StaticData);
+			CombinedTransform = OffsetTransform * FrameData->Transform;
+			TransformData.ApplyTransform(SceneComponent, CombinedTransform, *StaticData);
 		}
 	}
 }
diff --git a/Engine/Plugins/Animation/LiveLink/Source/LiveLinkComponents/Public/Controllers/LiveLinkTransformController.h b/Engine/Plugins/Animation/LiveLink/Source/LiveLinkComponents/Public/Controllers/LiveLinkTransformController.h
--- a/Engine/Plugins/Animation/LiveLink/Source/LiveLinkComponents/Public/Controllers/LiveLinkTransformController.h
+++ b/Engine/Plugins/Animation/LiveLink/Source/LiveLinkComponents/Public/Controllers/LiveLinkTransformController.h
@@ -3,6 +3,7 @@
 #pragma once
 
 #include "LiveLinkControllerBase.h"
+#include "Roles/LiveLinkTransformTypes.h"
 #include "LiveLinkTransformController.generated.h"
 
 class ULiveLinkRole;
@@ -10,32 +11,32 @@ class ULiveLinkRole;
 struct FLiveLinkTransformStaticData;
 class USceneComponent;
 
-USTRUCT(BlueprintType)
+USTRUCT(Blueprintable)
 struct LIVELINKCOMPONENTS_API FLiveLinkTransformControllerData
 {
 	GENERATED_BODY()
-
+	
 	/** Set the transform of the component in world space of in its local reference frame. */
-	UPROPERTY(EditAnywhere, BlueprintReadOnly, Category = "LiveLink")
+	UPROPERTY(EditAnywhere, BlueprintReadWrite, Category = "LiveLink")
 	bool bWorldTransform = false;
 
 	/** Whether we should set the owning actor's location with the value coming from live link. */
-	UPROPERTY(EditAnywhere, Category = "LiveLink", AdvancedDisplay)
+	UPROPERTY(EditAnywhere, BlueprintReadWrite, Category = "LiveLink", AdvancedDisplay)
 	bool bUseLocation = true;
 
 	/** Whether we should set the owning actor's rotation with the value coming from live link. */
-	UPROPERTY(EditAnywhere, Category = "LiveLink", AdvancedDisplay)
+	UPROPERTY(EditAnywhere, BlueprintReadWrite, Category = "LiveLink", AdvancedDisplay)
 	bool bUseRotation = true;
 	
 	/** Whether we should set the owning actor's scale with the value coming from live link. */
-	UPROPERTY(EditAnywhere, Category = "LiveLink", AdvancedDisplay)
+	UPROPERTY(EditAnywhere, BlueprintReadWrite, Category = "LiveLink", AdvancedDisplay)
 	bool bUseScale = true;
 
 	/**
 	 * Whether we sweep to the destination location, triggering overlaps along the way and stopping short of the target if blocked by something.
 	 * Only the root component is swept and checked for blocking collision, child components move without sweeping. If collision is off, this has no effect.
 	*/
-	UPROPERTY(EditAnywhere, BlueprintReadOnly, Category = "LiveLink", AdvancedDisplay)
+	UPROPERTY(EditAnywhere, BlueprintReadWrite, Category = "LiveLink", AdvancedDisplay)
 	bool bSweep = false;
 
 	/**
@@ -44,7 +45,7 @@ struct LIVELINKCOMPONENTS_API FLiveLinkTransformControllerData
 	 * If false, physics velocity is updated based on the change in position (affecting ragdoll parts).
 	 * If CCD is on and not teleporting, this will affect objects along the entire sweep volume.
 	*/
-	UPROPERTY(EditAnywhere, BlueprintReadOnly, Category = "LiveLink", AdvancedDisplay)
+	UPROPERTY(EditAnywhere, BlueprintReadWrite, Category = "LiveLink", AdvancedDisplay)
 	bool bTeleport = true;
 	
 	void ApplyTransform(USceneComponent* SceneComponent, const FTransform& Transform, const FLiveLinkTransformStaticData& StaticData) const;
@@ -53,7 +54,7 @@ struct LIVELINKCOMPONENTS_API FLiveLinkTransformControllerData
 
 /**
  */
-UCLASS()
+UCLASS(BlueprintType)
 class LIVELINKCOMPONENTS_API ULiveLinkTransformController : public ULiveLinkControllerBase
 {
 	GENERATED_BODY()
@@ -64,8 +65,19 @@ public:
 	FComponentReference ComponentToControl_DEPRECATED;
 #endif //WITH_EDITORONLY_DATA
 	
-	UPROPERTY(EditAnywhere, Category="LiveLink", meta=(ShowOnlyInnerProperties))
+	UPROPERTY(EditAnywhere, BlueprintReadWrite, Category="LiveLink", meta=(ShowOnlyInnerProperties))
 	FLiveLinkTransformControllerData TransformData;
+	
+	/**
+	 * Offset transform applied in local space to the controlled scene component
+	 */
+	UPROPERTY(EditAnywhere, BlueprintReadWrite, Category = "LiveLink")
+	FTransform OffsetTransform = FTransform::Identity;
+
+private:
+
+	/** Combined transform resulting from composing the incoming LiveLink transform & the offset transform */
+	FTransform CombinedTransform;
 
 public:
 	//~Begin ULiveLinkControllerBase interface
