LiveLinkHub - Fix AnyThread ticking in LiveLinkHubEditor target

In the regular LLH target, we would manually tick LiveLink in LiveLinkHubApplication::LiveLinkHubLoop unless we used the ticker, because SamplingInput delegate was never broadcast when running in commandlet mode.

Now that we're using the editor target, SamplingInput delegate is getting broadcast, so we have to indicate to LiveLinkClient which delegate to use for ticking. In this case, we provide either the Ticker's tick delegate or FCoreDelegates::SamplingInput if we want to run on the game thread.

#b jason.walter

diff --git a/Engine/Plugins/Animation/LiveLink/Source/LiveLink/Private/LiveLinkClient.cpp b/Engine/Plugins/Animation/LiveLink/Source/LiveLink/Private/LiveLinkClient.cpp
--- a/Engine/Plugins/Animation/LiveLink/Source/LiveLink/Private/LiveLinkClient.cpp
+++ b/Engine/Plugins/Animation/LiveLink/Source/LiveLink/Private/LiveLinkClient.cpp
@@ -54,30 +54,26 @@ static TAutoConsoleVariable<int32> CVarMaxNewFrameDataPerUpdate(
 
 
 FLiveLinkClient::FLiveLinkClient()
-	: Collection(MakeUnique<FLiveLinkSourceCollection>())
+	: FLiveLinkClient(FCoreDelegates::OnSamplingInput)
 {
-	LLM_SCOPE_BYTAG(LiveLink_LiveLinkClient);
-
-#if WITH_EDITOR
-	CachedEngineTime = 0.0;
-#endif
-
-	OnLiveLinkSubjectRemoved().AddRaw(this, &FLiveLinkClient::OnSubjectRemovedCallback);
-	FCoreDelegates::OnPreExit.AddRaw(this, &FLiveLinkClient::Shutdown);
-
-	//IMediaModule& MediaModule = FModuleManager::LoadModuleChecked<IMediaModule>("Media");
-	//MediaModule.GetOnTickPreEngineCompleted().AddRaw(this, &FLiveLinkClient::Tick);
+	// Use OnSamplingInput as the ticking delegate for now since it's as close as the previous PreEngineCompleted callback we were hooked before it was changed
+	// OnBeginFrame is too early since Timecode hasn't been updated for the frame
+	// OnSamplingInput is right before ticking the engine so we can build our snapshots and be consistent throughout the frame
+}
 
-	//Use OnSamplingInput for now since it's as close as the previous PreEngineCompleted callback we were hoooked before it was changed
-	//OnBeginFrame is too early since Timecode hasn't been updated for the frame
-	//OnSamplingInput is right before ticking the engine so we can build our snapshots and be consistent throughout the frame
-	FCoreDelegates::OnSamplingInput.AddRaw(this, &FLiveLinkClient::Tick);
 
-	// Setup rebroadcaster name in case we need it later
-	RebroadcastLiveLinkProviderName = TEXT("LiveLink Rebroadcast");
+FLiveLinkClient::FLiveLinkClient(FSimpleMulticastDelegate& InTickingDelegate)
+	: Collection(MakeUnique<FLiveLinkSourceCollection>())
+{
+	InTickingDelegate.AddRaw(this, &FLiveLinkClient::Tick);
+	Initialize();
+}
 
-	bPreProcessRebroadcastFrames = GetDefault<ULiveLinkSettings>()->bPreProcessRebroadcastFrames;
-	bTranslateRebroadcastFrames = GetDefault<ULiveLinkSettings>()->bTranslateRebroadcastFrames;
+FLiveLinkClient::FLiveLinkClient(FTSSimpleMulticastDelegate& InTickingDelegate)
+	: Collection(MakeUnique<FLiveLinkSourceCollection>())
+{
+	InTickingDelegate.AddRaw(this, &FLiveLinkClient::Tick);
+	Initialize();
 }
 
 FLiveLinkClient::~FLiveLinkClient()
@@ -98,6 +94,24 @@ void FLiveLinkClient::Tick()
 	OnLiveLinkTickedDelegate.Broadcast();
 }
 
+void FLiveLinkClient::Initialize()
+{
+	LLM_SCOPE_BYTAG(LiveLink_LiveLinkClient);
+
+#if WITH_EDITOR
+	CachedEngineTime = 0.0;
+#endif
+
+	OnLiveLinkSubjectRemoved().AddRaw(this, &FLiveLinkClient::OnSubjectRemovedCallback);
+	FCoreDelegates::OnPreExit.AddRaw(this, &FLiveLinkClient::Shutdown);
+
+	// Setup rebroadcaster name in case we need it later
+	RebroadcastLiveLinkProviderName = TEXT("LiveLink Rebroadcast");
+
+	bPreProcessRebroadcastFrames = GetDefault<ULiveLinkSettings>()->bPreProcessRebroadcastFrames;
+	bTranslateRebroadcastFrames = GetDefault<ULiveLinkSettings>()->bTranslateRebroadcastFrames;
+}
+
 void FLiveLinkClient::DoPendingWork()
 {
 	FScopeLock PendingFramesLock(&PendingFramesCriticalSection);
@@ -225,11 +239,6 @@ void FLiveLinkClient::Shutdown()
 {
 	LLM_SCOPE_BYTAG(LiveLink_LiveLinkClient);
 
-	/*if(IMediaModule* MediaModule = FModuleManager::GetModulePtr<IMediaModule>("Media"))
-	{
-		MediaModule->GetOnTickPreEngineCompleted().RemoveAll(this);
-	}*/
-
 	FCoreDelegates::OnSamplingInput.RemoveAll(this);
 
 	// Shut down the rebroadcaster if active
diff --git a/Engine/Plugins/Animation/LiveLink/Source/LiveLink/Public/LiveLinkClient.h b/Engine/Plugins/Animation/LiveLink/Source/LiveLink/Public/LiveLinkClient.h
--- a/Engine/Plugins/Animation/LiveLink/Source/LiveLink/Public/LiveLinkClient.h
+++ b/Engine/Plugins/Animation/LiveLink/Source/LiveLink/Public/LiveLinkClient.h
@@ -4,7 +4,6 @@
 
 #include "ILiveLinkClient.h"
 
-
 class ILiveLinkSubject;
 struct FPropertyChangedEvent;
 struct ILiveLinkProvider;
@@ -64,7 +63,14 @@ PRAGMA_ENABLE_DEPRECATION_WARNINGS
 class LIVELINK_API FLiveLinkClient : public FLiveLinkClient_Base_DEPRECATED
 {
 public:
+	/** Default constructor that setups LiveLink to use the SamplingInput delegate to tick. */
 	FLiveLinkClient();
+
+	//~ Constructor that allow providing a custom delegate for ticking LiveLink.  
+	FLiveLinkClient(FSimpleMulticastDelegate& InTickingDelegate);
+	FLiveLinkClient(FTSSimpleMulticastDelegate& InTickingDelegate);
+	
+
 	virtual ~FLiveLinkClient();
 
 	//~ Begin ILiveLinkClient implementation
@@ -225,6 +231,8 @@ private:
 	};
 
 private:
+	/** Common initialization code for the different constructors. */
+	void Initialize();
 
 	/** Remove old sources & subject,  */
 	void DoPendingWork();
diff --git a/Engine/Plugins/Animation/LiveLinkHub/Source/LiveLinkHub/Private/LiveLinkHub.cpp b/Engine/Plugins/Animation/LiveLinkHub/Source/LiveLinkHub/Private/LiveLinkHub.cpp
--- a/Engine/Plugins/Animation/LiveLinkHub/Source/LiveLinkHub/Private/LiveLinkHub.cpp
+++ b/Engine/Plugins/Animation/LiveLinkHub/Source/LiveLinkHub/Private/LiveLinkHub.cpp
@@ -14,6 +14,7 @@
 #include "LiveLinkHubClient.h"
 #include "LiveLinkHubCommands.h"
 #include "LiveLinkHubSubjectSettings.h"
+#include "LiveLinkHubTicker.h"
 #include "LiveLinkProviderImpl.h"
 #include "LiveLinkSettings.h"
 #include "Misc/App.h"
@@ -29,10 +30,18 @@
 #define LOCTEXT_NAMESPACE "LiveLinkHub"
 
 
-void FLiveLinkHub::Preinitialize()
+void FLiveLinkHub::Preinitialize(FLiveLinkHubTicker& Ticker)
 {
 	// We must register the livelink client first since we might rely on the modular feature to initialize the controllers/managers.
-	LiveLinkHubClient = MakeShared<FLiveLinkHubClient>(AsShared());
+	if (GetDefault<ULiveLinkHubSettings>()->bTickOnGameThread)
+	{
+		LiveLinkHubClient = MakeShared<FLiveLinkHubClient>(AsShared());
+	}
+	else
+	{
+		LiveLinkHubClient = MakeShared<FLiveLinkHubClient>(AsShared(), Ticker.OnTick());
+	}
+	
 	IModularFeatures::Get().RegisterModularFeature(ILiveLinkClient::ModularFeatureName, LiveLinkHubClient.Get());
 }
 
diff --git a/Engine/Plugins/Animation/LiveLinkHub/Source/LiveLinkHub/Private/LiveLinkHub.h b/Engine/Plugins/Animation/LiveLinkHub/Source/LiveLinkHub/Private/LiveLinkHub.h
--- a/Engine/Plugins/Animation/LiveLinkHub/Source/LiveLinkHub/Private/LiveLinkHub.h
+++ b/Engine/Plugins/Animation/LiveLinkHub/Source/LiveLinkHub/Private/LiveLinkHub.h
@@ -50,7 +50,7 @@ public:
 
 public:
 	/** First phase initialization; must precede FEngineLoop::Init(). */
-	void Preinitialize();
+	void Preinitialize(class FLiveLinkHubTicker& Ticker);
 	/** 
 	 * Launch the slate application and initialize its components. 
 	 * @param bLauncherDistribution Whether the hub should be initialized as a launcher distributed LiveLinkHub.
diff --git a/Engine/Plugins/Animation/LiveLinkHub/Source/LiveLinkHub/Private/LiveLinkHubApplication.h b/Engine/Plugins/Animation/LiveLinkHub/Source/LiveLinkHub/Private/LiveLinkHubApplication.h
--- a/Engine/Plugins/Animation/LiveLinkHub/Source/LiveLinkHub/Private/LiveLinkHubApplication.h
+++ b/Engine/Plugins/Animation/LiveLinkHub/Source/LiveLinkHub/Private/LiveLinkHubApplication.h
@@ -23,20 +23,11 @@ DEFINE_LOG_CATEGORY_STATIC(LogLiveLinkHubApplication, Log, All);
 
 void LiveLinkHubLoop(const TSharedPtr<FLiveLinkHub>& LiveLinkHub)
 {
-	const bool bTickOnGameThread = GetDefault<ULiveLinkHubSettings>()->bTickOnGameThread;
-
 	check(FSlateApplication::IsInitialized());
 	FSlateApplication::Get().RegisterInputPreProcessor(MakeShared<FLiveLinkHubInputProcessor>());
 	{
 		UE_LOG(LogLiveLinkHubApplication, Display, TEXT("LiveLinkHub Initialized (Version: %d.%d)"), ENGINE_MAJOR_VERSION, ENGINE_MINOR_VERSION);
 
-		FLiveLinkHubTicker Ticker{LiveLinkHub.ToSharedRef()};
-
-		if (!bTickOnGameThread)
-		{
-			Ticker.StartTick();
-		}
-
 		double LastTime = FPlatformTime::Seconds();
 
 		const double IdealFrameTime = 1 / GetDefault<ULiveLinkHubSettings>()->TargetFrameRate;
@@ -53,12 +44,6 @@ void LiveLinkHubLoop(const TSharedPtr<FLiveLinkHub>& LiveLinkHub)
 
 			FSlateApplication::Get().PollGameDeviceState();
 
-			// This is normally ticked by OnSamplingInput.
-			if (bTickOnGameThread)
-			{
-				LiveLinkHub->Tick();
-			}
-
 			// Run garbage collection for the UObjects for the rest of the frame or at least to 2 ms
 			IncrementalPurgeGarbage(true, FMath::Max<float>(0.002f, IdealFrameTime - (FPlatformTime::Seconds() - LastTime)));
 
@@ -74,8 +59,5 @@ void LiveLinkHubLoop(const TSharedPtr<FLiveLinkHub>& LiveLinkHub)
 
 			LastTime = CurrentTime;
 		}
-
-		Ticker.Exit();
-		Ticker.Stop();
 	}
 }
diff --git a/Engine/Plugins/Animation/LiveLinkHub/Source/LiveLinkHub/Private/LiveLinkHubClient.cpp b/Engine/Plugins/Animation/LiveLinkHub/Source/LiveLinkHub/Private/LiveLinkHubClient.cpp
--- a/Engine/Plugins/Animation/LiveLinkHub/Source/LiveLinkHub/Private/LiveLinkHubClient.cpp
+++ b/Engine/Plugins/Animation/LiveLinkHub/Source/LiveLinkHub/Private/LiveLinkHubClient.cpp
@@ -24,12 +24,6 @@ DECLARE_CYCLE_STAT(TEXT("LiveLinkHub - Push FrameData"), STAT_LiveLinkHub_PushFr
 
 #define LOCTEXT_NAMESPACE "LiveLinkHub.LiveLinkHubClient"
 
-FLiveLinkHubClient::FLiveLinkHubClient(TSharedPtr<ILiveLinkHub> InLiveLinkHub)
-	: LiveLinkHub(MoveTemp(InLiveLinkHub))
-{
-	RegisterGlobalSubjectFramesDelegate(FOnLiveLinkSubjectStaticDataAdded::FDelegate::CreateRaw(this, &FLiveLinkHubClient::OnStaticDataAdded), FOnLiveLinkSubjectFrameDataAdded::FDelegate::CreateRaw(this, &FLiveLinkHubClient::OnFrameDataAdded), StaticDataAddedHandle, FrameDataAddedHandle);
-}
-
 FLiveLinkHubClient::~FLiveLinkHubClient()
 {
 	UnregisterGlobalSubjectFramesDelegate(StaticDataAddedHandle, FrameDataAddedHandle);
diff --git a/Engine/Plugins/Animation/LiveLinkHub/Source/LiveLinkHub/Private/LiveLinkHubClient.h b/Engine/Plugins/Animation/LiveLinkHub/Source/LiveLinkHub/Private/LiveLinkHubClient.h
--- a/Engine/Plugins/Animation/LiveLinkHub/Source/LiveLinkHub/Private/LiveLinkHubClient.h
+++ b/Engine/Plugins/Animation/LiveLinkHub/Source/LiveLinkHub/Private/LiveLinkHubClient.h
@@ -22,7 +22,20 @@ DECLARE_TS_MULTICAST_DELEGATE_OneParam(FOnSubjectMarkedPendingKill_AnyThread, co
 class FLiveLinkHubClient : public FLiveLinkClient
 {
 public:
-	FLiveLinkHubClient(TSharedPtr<ILiveLinkHub> InLiveLinkHub);
+	FLiveLinkHubClient(TSharedPtr<ILiveLinkHub> InLiveLinkHub, FTSSimpleMulticastDelegate& InTickingDelegate)
+		: FLiveLinkClient(InTickingDelegate)
+		, LiveLinkHub(MoveTemp(InLiveLinkHub))
+	{
+		RegisterGlobalSubjectFramesDelegate(FOnLiveLinkSubjectStaticDataAdded::FDelegate::CreateRaw(this, &FLiveLinkHubClient::OnStaticDataAdded), FOnLiveLinkSubjectFrameDataAdded::FDelegate::CreateRaw(this, &FLiveLinkHubClient::OnFrameDataAdded), StaticDataAddedHandle, FrameDataAddedHandle);
+	}
+
+	FLiveLinkHubClient(TSharedPtr<ILiveLinkHub> InLiveLinkHub)
+		: FLiveLinkClient()
+		, LiveLinkHub(MoveTemp(InLiveLinkHub))
+	{
+		RegisterGlobalSubjectFramesDelegate(FOnLiveLinkSubjectStaticDataAdded::FDelegate::CreateRaw(this, &FLiveLinkHubClient::OnStaticDataAdded), FOnLiveLinkSubjectFrameDataAdded::FDelegate::CreateRaw(this, &FLiveLinkHubClient::OnFrameDataAdded), StaticDataAddedHandle, FrameDataAddedHandle);
+	}
+
 	virtual ~FLiveLinkHubClient();
 	
 	/** Get the delegate called when frame data is received. */
diff --git a/Engine/Plugins/Animation/LiveLinkHub/Source/LiveLinkHub/Private/LiveLinkHubModule.cpp b/Engine/Plugins/Animation/LiveLinkHub/Source/LiveLinkHub/Private/LiveLinkHubModule.cpp
--- a/Engine/Plugins/Animation/LiveLinkHub/Source/LiveLinkHub/Private/LiveLinkHubModule.cpp
+++ b/Engine/Plugins/Animation/LiveLinkHub/Source/LiveLinkHub/Private/LiveLinkHubModule.cpp
@@ -22,8 +22,14 @@
 void FLiveLinkHubModule::PreinitializeLiveLinkHub()
 {
 	check(!LiveLinkHub);
+
+	if (!GetDefault<ULiveLinkHubSettings>()->bTickOnGameThread)
+	{
+		Ticker.StartTick();
+	}
+
 	LiveLinkHub = MakeShared<FLiveLinkHub>();
-	LiveLinkHub->Preinitialize();
+	LiveLinkHub->Preinitialize(Ticker);
 }
 
 void FLiveLinkHubModule::StartLiveLinkHub(bool bLauncherDistribution)
@@ -33,9 +39,6 @@ void FLiveLinkHubModule::StartLiveLinkHub(bool bLauncherDistribution)
 
 	LiveLinkHub->Initialize(bLauncherDistribution);
 
-	// Disable throttling for the hub
-	GetMutableDefault<UEditorPerformanceSettings>()->bThrottleCPUWhenNotForeground = false;
-
 #if IS_PROGRAM
 	LiveLinkHubLoop(LiveLinkHub);
 #endif
@@ -44,6 +47,12 @@ void FLiveLinkHubModule::StartLiveLinkHub(bool bLauncherDistribution)
 void FLiveLinkHubModule::ShutdownLiveLinkHub()
 {
 	LiveLinkHub.Reset();
+
+	if (!GetDefault<ULiveLinkHubSettings>()->bTickOnGameThread)
+	{
+		Ticker.Exit();
+		Ticker.Stop();
+	}
 }
 
 void FLiveLinkHubModule::StartupModule()
diff --git a/Engine/Plugins/Animation/LiveLinkHub/Source/LiveLinkHub/Private/LiveLinkHubModule.h b/Engine/Plugins/Animation/LiveLinkHub/Source/LiveLinkHub/Private/LiveLinkHubModule.h
--- a/Engine/Plugins/Animation/LiveLinkHub/Source/LiveLinkHub/Private/LiveLinkHubModule.h
+++ b/Engine/Plugins/Animation/LiveLinkHub/Source/LiveLinkHub/Private/LiveLinkHubModule.h
@@ -3,6 +3,7 @@
 #pragma once
 
 #include "ILiveLinkHubModule.h"
+#include "LiveLinkHubTicker.h"
 #include "Templates/SharedPointer.h"
 
 class FLiveLinkHub;
@@ -44,6 +45,9 @@ private:
 	/** LiveLinkHub object responsible for initializing the different controllers. */
 	TSharedPtr<FLiveLinkHub> LiveLinkHub;
 
+	/** Utility object used to tick LiveLink outside of the game thread. */
+	FLiveLinkHubTicker Ticker;
+
 	/** Config-driven. True in LiveLink Hub standalone application, false in UE. */
 	bool bUseSubjectSettingsDetailsCustomization;
 };
diff --git a/Engine/Plugins/Animation/LiveLinkHub/Source/LiveLinkHub/Private/LiveLinkHubTicker.h b/Engine/Plugins/Animation/LiveLinkHub/Source/LiveLinkHub/Private/LiveLinkHubTicker.h
--- a/Engine/Plugins/Animation/LiveLinkHub/Source/LiveLinkHub/Private/LiveLinkHubTicker.h
+++ b/Engine/Plugins/Animation/LiveLinkHub/Source/LiveLinkHub/Private/LiveLinkHubTicker.h
@@ -12,11 +12,6 @@
 class FLiveLinkHubTicker : public FRunnable
 {
 public:
-	explicit FLiveLinkHubTicker(TSharedRef<FLiveLinkHub> InLiveLinkHub)
-		: LiveLinkHub(MoveTemp(InLiveLinkHub))
-	{
-	}
-
 	void StartTick()
 	{
 		if (!bIsRunning)
@@ -42,7 +37,7 @@ public:
 			if (bIsRunning) // make sure we were not told to exit during the wait
 			{
 				TRACE_CPUPROFILER_EVENT_SCOPE(FLiveLinkHubTicker::Tick);
-				LiveLinkHub->Tick();
+				OnTickDelegate.Broadcast();
 			}
 		}
 
@@ -71,9 +66,18 @@ public:
 		}
 	}
 	//~ End FRunnable Interface
+
+public:
+
+	/** Get the delegate called whenever this object ticks. */
+	FTSSimpleMulticastDelegate& OnTick() { return OnTickDelegate;  }
+private:
+
 	std::atomic<bool> bIsRunning = false;
 
+	/** Delegate called when this ticks. */
+	FTSSimpleMulticastDelegate OnTickDelegate;
+
 	FEvent* TickEvent = nullptr;
 	TUniquePtr<FRunnableThread> Thread;
-	TSharedPtr<FLiveLinkHub> LiveLinkHub;
 };
