LiveLink - Move Parent subject functionality into LiveLink's rebroadcast flow

diff --git a/Engine/Plugins/Animation/LiveLink/Source/LiveLink/Private/LiveLinkClient.cpp b/Engine/Plugins/Animation/LiveLink/Source/LiveLink/Private/LiveLinkClient.cpp
--- a/Engine/Plugins/Animation/LiveLink/Source/LiveLink/Private/LiveLinkClient.cpp
+++ b/Engine/Plugins/Animation/LiveLink/Source/LiveLink/Private/LiveLinkClient.cpp
@@ -110,6 +110,7 @@ void FLiveLinkClient::Initialize()
 
 	bPreProcessRebroadcastFrames = GetDefault<ULiveLinkSettings>()->bPreProcessRebroadcastFrames;
 	bTranslateRebroadcastFrames = GetDefault<ULiveLinkSettings>()->bTranslateRebroadcastFrames;
+	bEnableParentSubjects = GConfig->GetBoolOrDefault(TEXT("LiveLink"), TEXT("bEnableParentSubjects"), false, GEngineIni);
 }
 
 void FLiveLinkClient::DoPendingWork()
@@ -1011,7 +1012,40 @@ void FLiveLinkClient::PushSubjectFrameData_Internal(FPendingSubjectFrame&& Subje
 		AllSubjectsHandler->OnFrameDataAdded.Broadcast(SubjectItem->Key, Role, SubjectFrameData.FrameData);
 	}
 
-	HandleSubjectRebroadcast(LinkSubject, SubjectFrameData.FrameData);
+	const bool bHasParentSubject = SourceItem->Setting->ParentSubject != FLiveLinkSubjectName();
+	if (!bHasParentSubject)
+	{
+		HandleSubjectRebroadcast(LinkSubject, SubjectFrameData.FrameData);
+
+		if (bEnableParentSubjects)
+		{
+			Collection->ForEachSubject([this, &SubjectFrameData](const FLiveLinkCollectionSourceItem& SourceItem, const FLiveLinkCollectionSubjectItem& SubjectItem)
+			{
+				const FLiveLinkSubjectKey SubjectKey = SubjectFrameData.SubjectKey;
+				if (SourceItem.Setting->ParentSubject.Name == SubjectKey.SubjectName)
+				{
+					const FLiveLinkFrameDataStruct& FrameData = SubjectFrameData.FrameData;
+
+					// todo: Time offset evaluation
+					FLiveLinkSubjectFrameData ChildData;
+					if (SubjectItem.GetLiveSubject()->EvaluateFrameAtWorldTime(SubjectFrameData.FrameData.GetBaseData()->WorldTime.GetSourceTime(), SubjectItem.GetLinkSettings()->Role, ChildData))
+					{
+						const FTimecode FrameTC = FTimecode::FromFrameNumber(FrameData.GetBaseData()->MetaData.SceneTime.Time.GetFrame(), FrameData.GetBaseData()->MetaData.SceneTime.Rate);
+						UE_LOG(LogLiveLink, Verbose, TEXT("LiveLinkHub Parent (%s) - Child '%s' adding frame with Timecode:[%s.%0.3f] - SourceTime: %0.4f, Offset: %0.6f, CorrectedTime: %0.4f"), *SubjectKey.SubjectName.ToString(), *SubjectItem.Key.SubjectName.ToString(), *FrameTC.ToString(), FrameData.GetBaseData()->MetaData.SceneTime.Time.GetSubFrame(), FrameData.GetBaseData()->WorldTime.GetSourceTime(), FrameData.GetBaseData()->WorldTime.GetOffset(), FrameData.GetBaseData()->WorldTime.GetOffsettedTime());
+
+						ChildData.FrameData.GetBaseData()->MetaData.SceneTime = FrameData.GetBaseData()->MetaData.SceneTime;
+						ChildData.FrameData.GetBaseData()->MetaData.SceneTime.Rate = FrameData.GetBaseData()->MetaData.SceneTime.Rate;
+
+						HandleSubjectRebroadcast(SubjectItem.GetLiveSubject(), ChildData.FrameData);
+					}
+					else
+					{
+						FLiveLinkLog::Warning(TEXT("Child subjects %s could not be evaluated for data resampling."), *SubjectKey.SubjectName.Name.ToString());
+					}
+				}
+			});
+		}
+	}
 
 	//Finally, add the new frame to the subject. After this point, the frame data is unusable, it has been moved!
 	LinkSubject->AddFrameData(MoveTemp(SubjectFrameData.FrameData));
diff --git a/Engine/Plugins/Animation/LiveLink/Source/LiveLink/Public/LiveLinkClient.h b/Engine/Plugins/Animation/LiveLink/Source/LiveLink/Public/LiveLinkClient.h
--- a/Engine/Plugins/Animation/LiveLink/Source/LiveLink/Public/LiveLinkClient.h
+++ b/Engine/Plugins/Animation/LiveLink/Source/LiveLink/Public/LiveLinkClient.h
@@ -326,6 +326,9 @@ private:
 	/** Whether to translate frames before rebroadcasting them. */
 	bool bTranslateRebroadcastFrames = false;
 
+	/** Whether or not parent subject support is enabled. Parent subjects allow resampling data to a different subject's rate before rebroadcasting it. */
+	bool bEnableParentSubjects = false;
+
 #if WITH_EDITOR
 	/** Delegate when a subject is evaluated. */
 	FOnLiveLinkSubjectEvaluated OnLiveLinkSubjectEvaluatedDelegate;
diff --git a/Engine/Plugins/Animation/LiveLink/Source/LiveLinkEditor/Private/LiveLinkSourceSettingsDetailCustomization.cpp b/Engine/Plugins/Animation/LiveLink/Source/LiveLinkEditor/Private/LiveLinkSourceSettingsDetailCustomization.cpp
--- a/Engine/Plugins/Animation/LiveLink/Source/LiveLinkEditor/Private/LiveLinkSourceSettingsDetailCustomization.cpp
+++ b/Engine/Plugins/Animation/LiveLink/Source/LiveLinkEditor/Private/LiveLinkSourceSettingsDetailCustomization.cpp
@@ -7,6 +7,7 @@
 #include "ILiveLinkClient.h"
 #include "LiveLinkSourceSettings.h"
 #include "LiveLinkSubjectSettings.h"
+#include "Misc/ConfigCacheIni.h"
 
 #include "Widgets/SBoxPanel.h"
 #include "Widgets/Text/STextBlock.h"
@@ -62,6 +63,13 @@ void FLiveLinkSourceSettingsDetailCustomization::CustomizeDetails(IDetailLayoutB
 	InDetailBuilder.AddPropertyToCategory(ModePropertyHandle);
 	ModePropertyHandle->SetOnPropertyValueChanged(FSimpleDelegate::CreateSP(this, &FLiveLinkSourceSettingsDetailCustomization::ForceRefresh));
 
+	bool bEnableParentSubjects = GConfig->GetBoolOrDefault(TEXT("LiveLink"), TEXT("bEnableParentSubjects"), false, GEngineIni);
+
+	if (!bEnableParentSubjects)
+	{
+		InDetailBuilder.HideProperty(InDetailBuilder.GetProperty(GET_MEMBER_NAME_CHECKED(ULiveLinkSourceSettings, ParentSubject)));
+	}
+
 	IDetailCategoryBuilder& CategoryBuilder = InDetailBuilder.EditCategory("Buffer - Settings");
 
 	CategoryBuilder.AddProperty(BufferSettingsPropertyHandle->GetChildHandle(GET_MEMBER_NAME_CHECKED(FLiveLinkSourceBufferManagementSettings, MaxNumberOfFrameToBuffered)));
diff --git a/Engine/Plugins/Animation/LiveLinkHub/Source/LiveLinkHub/Private/LiveLinkHubClient.cpp b/Engine/Plugins/Animation/LiveLinkHub/Source/LiveLinkHub/Private/LiveLinkHubClient.cpp
--- a/Engine/Plugins/Animation/LiveLinkHub/Source/LiveLinkHub/Private/LiveLinkHubClient.cpp
+++ b/Engine/Plugins/Animation/LiveLinkHub/Source/LiveLinkHub/Private/LiveLinkHubClient.cpp
@@ -123,34 +123,7 @@ void FLiveLinkHubClient::OnStaticDataAdded(FLiveLinkSubjectKey SubjectKey,  TSub
 
 void FLiveLinkHubClient::OnFrameDataAdded(FLiveLinkSubjectKey InSubjectKey, TSubclassOf<ULiveLinkRole> SubjectRole, const FLiveLinkFrameDataStruct& InFrameData)
 {
-	Collection->ForEachSubject([this, InSubjectKey, &InFrameData](const FLiveLinkCollectionSourceItem& SourceItem, const FLiveLinkCollectionSubjectItem& SubjectItem)
-	{
-		if (SourceItem.Setting->ParentSubject.Name == InSubjectKey.SubjectName)
-		{
-			// todo: Time offset evaluation
-			FLiveLinkSubjectFrameData ChildData;
-			if (SubjectItem.GetLiveSubject()->EvaluateFrameAtWorldTime(InFrameData.GetBaseData()->WorldTime.GetSourceTime(), SubjectItem.GetLinkSettings()->Role, ChildData))
-			{
-				const FTimecode FrameTC = FTimecode::FromFrameNumber(InFrameData.GetBaseData()->MetaData.SceneTime.Time.GetFrame(), InFrameData.GetBaseData()->MetaData.SceneTime.Rate);
-				UE_LOG(LogLiveLinkHub, Verbose, TEXT("LiveLinkHub Parent (%s) - Child '%s' adding frame with Timecode:[%s.%0.3f] - SourceTime: %0.4f, Offset: %0.6f, CorrectedTime: %0.4f"), *InSubjectKey.SubjectName.ToString(), *SubjectItem.Key.SubjectName.ToString(), *FrameTC.ToString(), InFrameData.GetBaseData()->MetaData.SceneTime.Time.GetSubFrame(), InFrameData.GetBaseData()->WorldTime.GetSourceTime(), InFrameData.GetBaseData()->WorldTime.GetOffset(), InFrameData.GetBaseData()->WorldTime.GetOffsettedTime());
-
-				ChildData.FrameData.GetBaseData()->MetaData.SceneTime = InFrameData.GetBaseData()->MetaData.SceneTime;
-				ChildData.FrameData.GetBaseData()->MetaData.SceneTime.Rate = InFrameData.GetBaseData()->MetaData.SceneTime.Rate;
-				OnFrameDataReceivedDelegate_AnyThread.Broadcast(SubjectItem.Key, ChildData.FrameData);
-			}
-			else
-			{
-				FLiveLinkLog::Warning(TEXT("Child subjects %s could not be evaluated for data resampling."), *InSubjectKey.SubjectName.Name.ToString());
-			}
-		}
-	});
-
-	const bool bIsParented = Collection->FindSource(InSubjectKey.Source)->Setting->ParentSubject != FLiveLinkSubjectName();
-	if (!bIsParented)
-	{
-		// Don't broadcast frames that have a parent, they'll be broadcast in the loop above.
-		OnFrameDataReceivedDelegate_AnyThread.Broadcast(InSubjectKey, InFrameData);
-	}
+	OnFrameDataReceivedDelegate_AnyThread.Broadcast(InSubjectKey, InFrameData);
 }
 
 
diff --git a/Engine/Source/Programs/LiveLinkHubEditor/Config/DefaultEngine.ini b/Engine/Source/Programs/LiveLinkHubEditor/Config/DefaultEngine.ini
--- a/Engine/Source/Programs/LiveLinkHubEditor/Config/DefaultEngine.ini
+++ b/Engine/Source/Programs/LiveLinkHubEditor/Config/DefaultEngine.ini
@@ -2,6 +2,7 @@
 bUseLiveLinkHubSubjectSettingsDetailsCustomization=true
 bPanelControllerSeparateSourcesSubjects=true
 bUseModularClientReference=true
+bEnableParentSubjects=true
 
 [LiveLinkHub]
 bEnableLLHMessageBusSourceFactory=false
diff --git a/Engine/Source/Runtime/LiveLinkInterface/Public/LiveLinkSourceSettings.h b/Engine/Source/Runtime/LiveLinkInterface/Public/LiveLinkSourceSettings.h
--- a/Engine/Source/Runtime/LiveLinkInterface/Public/LiveLinkSourceSettings.h
+++ b/Engine/Source/Runtime/LiveLinkInterface/Public/LiveLinkSourceSettings.h
@@ -186,7 +186,7 @@ public:
 
 	/** 
 	 * Parent subject used for resampling.
-	 * If this is set, this source's subjects will only be updated when the parent subject receives data.
+	 * If this is set, this source's subjects will only be rebroadcast when the parent subject receives data.
 	 * Additionally this source's subjects' timecode will match the parent's subject received timecode.
 	 * This can be useful for synchronizing a higher frequency source to a lower frequency one.
 	 */
