Live Link: Fix missing dependency on MovieScene

#rnx

diff --git a/Engine/Plugins/Animation/LiveLink/Source/LiveLinkEditor/LiveLinkEditor.Build.cs b/Engine/Plugins/Animation/LiveLink/Source/LiveLinkEditor/LiveLinkEditor.Build.cs
--- a/Engine/Plugins/Animation/LiveLink/Source/LiveLinkEditor/LiveLinkEditor.Build.cs
+++ b/Engine/Plugins/Animation/LiveLink/Source/LiveLinkEditor/LiveLinkEditor.Build.cs
@@ -37,6 +37,7 @@ namespace UnrealBuildTool.Rules
 					"LiveLinkMessageBusFramework",
 					"LiveLinkMovieScene",
 					"MessageLog",
+					"MovieScene",
 					"Persona",
 					"Projects",
 					"PropertyEditor",
