LiveLinkHub - Fix remapper being applied twice to recordings. We now save the unmapped data to recordings and apply the remapper on playback.
- Optimize the number of times LLH sends static data, correcting issues with remappers being constantly reapplied.
- Fix an ensure that could be triggered if a remapper was compiling skinned assets on a different thread.

diff --git a/Engine/Plugins/Animation/LiveLink/Source/LiveLink/Private/LiveLinkClient.cpp b/Engine/Plugins/Animation/LiveLink/Source/LiveLink/Private/LiveLinkClient.cpp
--- a/Engine/Plugins/Animation/LiveLink/Source/LiveLink/Private/LiveLinkClient.cpp
+++ b/Engine/Plugins/Animation/LiveLink/Source/LiveLink/Private/LiveLinkClient.cpp
@@ -830,6 +830,8 @@ void FLiveLinkClient::PushSubjectStaticData_Internal(FPendingSubjectStatic&& Sub
 		Collection->AddSubject(MoveTemp(CollectionSubjectItem));
 	}
 
+	const FLiveLinkStaticDataStruct* UnmappedStaticData = &SubjectStaticData.StaticData;
+	
 	if (LiveLinkSubject)
 	{
 		if (ULiveLinkSubjectRemapper::FWorkerSharedPtr Remapper = LiveLinkSubject->GetFrameRemapper())
@@ -844,15 +846,18 @@ void FLiveLinkClient::PushSubjectStaticData_Internal(FPendingSubjectStatic&& Sub
 					LiveLinkSubject->SetStaticDataAsRebroadcasted(false);
 				}
 			}
+			UnmappedStaticData = &LiveLinkSubject->GetStaticData(/*bGetOverrideData*/ false);
 		}
 
 		if (const FSubjectFramesAddedHandles* Handles = SubjectFrameAddedHandles.Find(SubjectStaticData.SubjectKey.SubjectName))
 		{
 			Handles->OnStaticDataAdded.Broadcast(SubjectStaticData.SubjectKey, SubjectStaticData.Role, SubjectStaticData.StaticData);
+			Handles->OnUnmappedStaticDataAdded.Broadcast(SubjectStaticData.SubjectKey, SubjectStaticData.Role, *UnmappedStaticData);
 		}
 		else if (const FSubjectFramesAddedHandles* AllSubjectsHandler = SubjectFrameAddedHandles.Find(ALL_SUBJECTS_DELEGATE_TOKEN))
 		{
 			AllSubjectsHandler->OnStaticDataAdded.Broadcast(SubjectStaticData.SubjectKey, SubjectStaticData.Role, SubjectStaticData.StaticData);
+			AllSubjectsHandler->OnUnmappedStaticDataAdded.Broadcast(SubjectStaticData.SubjectKey, SubjectStaticData.Role, *UnmappedStaticData);
 		}
 
 		LiveLinkSubject->SetStaticData(SubjectStaticData.Role, MoveTemp(SubjectStaticData.StaticData));
@@ -956,8 +961,11 @@ void FLiveLinkClient::PushSubjectFrameData_Internal(FPendingSubjectFrame&& Subje
 		return;
 	}
 
+	FLiveLinkFrameDataStruct UnmappedFrameData;
+	
 	if (ULiveLinkSubjectRemapper::FWorkerSharedPtr Remapper = LinkSubject->GetFrameRemapper())
 	{
+		UnmappedFrameData.InitializeWith(SubjectFrameData.FrameData);
 		Remapper->RemapFrameData(LinkSubject->GetStaticData(), SubjectFrameData.FrameData);
 	}
 
@@ -972,6 +980,16 @@ void FLiveLinkClient::PushSubjectFrameData_Internal(FPendingSubjectFrame&& Subje
 		return;
 	}
 
+	if (UnmappedFrameData.IsValid() && !Role.GetDefaultObject()->IsFrameDataValid(LinkSubject->GetStaticData(/*bGetOverrideData*/ false), UnmappedFrameData, bShouldLogWarning))
+	{
+		if (bShouldLogWarning)
+		{
+			static const FName NAME_InvalidFrameData = "LiveLinkClient_InvalidUnmappedFrameData";
+			FLiveLinkLog::ErrorOnce(NAME_InvalidFrameData, SubjectFrameData.SubjectKey, TEXT("Trying to add unmapped frame data that is not formatted properly to role '%s' with subject '%s'."), *Role->GetName(), *SubjectFrameData.SubjectKey.SubjectName.ToString());
+		}
+		return;
+	}
+
 	//Stamp arrival time of each packet to track clock difference when it is effectively added to the stash.
 	//Doing it in the Add_AnyThread would mean that we stamp it up to 1 frame time behind, causing the offset to always be 1 frame behind
 	//and requiring 2.5 frames or so to have a valid smooth offset
@@ -982,6 +1000,10 @@ void FLiveLinkClient::PushSubjectFrameData_Internal(FPendingSubjectFrame&& Subje
 		if (CurrentTime.IsSet())
 		{
 			SubjectFrameData.FrameData.GetBaseData()->ArrivalTime.SceneTime = *CurrentTime;
+			if (UnmappedFrameData.IsValid() && UnmappedFrameData.GetBaseData())
+			{
+				UnmappedFrameData.GetBaseData()->ArrivalTime.SceneTime = *CurrentTime;
+			}
 		}
 	}
 
@@ -991,11 +1013,13 @@ void FLiveLinkClient::PushSubjectFrameData_Internal(FPendingSubjectFrame&& Subje
 	if (const FSubjectFramesAddedHandles* Handles = SubjectFrameAddedHandles.Find(SubjectFrameData.SubjectKey.SubjectName))
 	{
 		Handles->OnFrameDataAdded.Broadcast(SubjectItem->Key, Role, SubjectFrameData.FrameData);
+		Handles->OnUnmappedFrameDataAdded.Broadcast(SubjectItem->Key, Role, UnmappedFrameData.IsValid() ? UnmappedFrameData : SubjectFrameData.FrameData);
 	}
 	else if (const FSubjectFramesAddedHandles* AllSubjectsHandler = SubjectFrameAddedHandles.Find(ALL_SUBJECTS_DELEGATE_TOKEN))
 	{
 		// NAME_None means we registered for all subjects update.
 		AllSubjectsHandler->OnFrameDataAdded.Broadcast(SubjectItem->Key, Role, SubjectFrameData.FrameData);
+		AllSubjectsHandler->OnUnmappedFrameDataAdded.Broadcast(SubjectItem->Key, Role, UnmappedFrameData.IsValid() ? UnmappedFrameData : SubjectFrameData.FrameData);
 	}
 
 	const bool bHasParentSubject = SourceItem->Setting->ParentSubject != FLiveLinkSubjectName();
@@ -1767,13 +1791,13 @@ UObject* FLiveLinkClient::GetSubjectSettings(const FLiveLinkSubjectKey& InSubjec
 	return nullptr;
 }	
 
-const FLiveLinkStaticDataStruct* FLiveLinkClient::GetSubjectStaticData_AnyThread(const FLiveLinkSubjectKey& InSubjectKey) const
+const FLiveLinkStaticDataStruct* FLiveLinkClient::GetSubjectStaticData_AnyThread(const FLiveLinkSubjectKey& InSubjectKey, bool bGetOverrideData) const
 {
 	if (const FLiveLinkCollectionSubjectItem* SubjectItem = Collection->FindSubject(InSubjectKey))
 	{
 		if (FLiveLinkSubject* LiveLinkSubject = SubjectItem->GetLiveSubject())
 		{
-			return &LiveLinkSubject->GetStaticData();
+			return &LiveLinkSubject->GetStaticData(bGetOverrideData);
 		}
 	}
 
@@ -1926,20 +1950,30 @@ FLiveLinkSkeletonStaticData* FLiveLinkClient::GetSubjectAnimationStaticData_Depr
 	return nullptr;
 }
 
-void FLiveLinkClient::UnregisterGlobalSubjectFramesDelegate(FDelegateHandle& InStaticDataAddedHandle, FDelegateHandle& InFrameDataAddedHandle)
+void FLiveLinkClient::UnregisterGlobalSubjectFramesDelegate(FDelegateHandle& InStaticDataAddedHandle, FDelegateHandle& InFrameDataAddedHandle, bool bUseUnmappedData)
 {
 	if (FSubjectFramesAddedHandles* Handles = SubjectFrameAddedHandles.Find(ALL_SUBJECTS_DELEGATE_TOKEN))
 	{
-		Handles->OnStaticDataAdded.Remove(InStaticDataAddedHandle);
-		Handles->OnFrameDataAdded.Remove(InFrameDataAddedHandle);
+		if (bUseUnmappedData)
+		{
+			Handles->OnUnmappedStaticDataAdded.Remove(InStaticDataAddedHandle);
+			Handles->OnUnmappedFrameDataAdded.Remove(InFrameDataAddedHandle);
+		}
+		else
+		{
+			Handles->OnStaticDataAdded.Remove(InStaticDataAddedHandle);
+			Handles->OnFrameDataAdded.Remove(InFrameDataAddedHandle);
+		}
 	}
 }
 
-bool FLiveLinkClient::RegisterGlobalSubjectFramesDelegate(const FOnLiveLinkSubjectStaticDataAdded::FDelegate& InOnStaticDataAdded, const FOnLiveLinkSubjectFrameDataAdded::FDelegate& InOnFrameDataAdded, FDelegateHandle& OutStaticDataAddedHandle, FDelegateHandle& OutFrameDataAddedHandle)
+bool FLiveLinkClient::RegisterGlobalSubjectFramesDelegate(const FOnLiveLinkSubjectStaticDataAdded::FDelegate& InOnStaticDataAdded,
+	const FOnLiveLinkSubjectFrameDataAdded::FDelegate& InOnFrameDataAdded, FDelegateHandle& OutStaticDataAddedHandle,
+	FDelegateHandle& OutFrameDataAddedHandle, bool bUseUnmappedData)
 {
 	FSubjectFramesAddedHandles& Handles = SubjectFrameAddedHandles.FindOrAdd(ALL_SUBJECTS_DELEGATE_TOKEN);
-	OutStaticDataAddedHandle = Handles.OnStaticDataAdded.Add(InOnStaticDataAdded);
-	OutFrameDataAddedHandle = Handles.OnFrameDataAdded.Add(InOnFrameDataAdded);
+	OutStaticDataAddedHandle = bUseUnmappedData ? Handles.OnUnmappedStaticDataAdded.Add(InOnStaticDataAdded) : Handles.OnStaticDataAdded.Add(InOnStaticDataAdded);
+	OutFrameDataAddedHandle = bUseUnmappedData ? Handles.OnUnmappedFrameDataAdded.Add(InOnFrameDataAdded) : Handles.OnFrameDataAdded.Add(InOnFrameDataAdded);
 
 	return true;
 }
diff --git a/Engine/Plugins/Animation/LiveLink/Source/LiveLink/Private/LiveLinkSubject.cpp b/Engine/Plugins/Animation/LiveLink/Source/LiveLink/Private/LiveLinkSubject.cpp
--- a/Engine/Plugins/Animation/LiveLink/Source/LiveLink/Private/LiveLinkSubject.cpp
+++ b/Engine/Plugins/Animation/LiveLink/Source/LiveLink/Private/LiveLinkSubject.cpp
@@ -13,6 +13,9 @@
 #include "LiveLinkTimedDataInput.h"
 #include "TimedDataInputCollection.h"
 
+#if WITH_EDITOR
+#include "SkinnedAssetCompiler.h"
+#endif
 
 FLiveLinkSubject::FSubjectEvaluationStatistics::FSubjectEvaluationStatistics()
 	: BufferUnderflow(0)
@@ -217,9 +220,9 @@ bool FLiveLinkSubject::HasValidFrameSnapshot() const
 	return FrameSnapshot.StaticData.IsValid() && FrameSnapshot.FrameData.IsValid();
 }
 
-FLiveLinkStaticDataStruct& FLiveLinkSubject::GetStaticData()
+FLiveLinkStaticDataStruct& FLiveLinkSubject::GetStaticData(bool bGetOverrideData)
 {
-	if (OverrideStaticData)
+	if (bGetOverrideData && OverrideStaticData)
 	{
 		return *OverrideStaticData;
 	}
@@ -1129,7 +1132,13 @@ void FLiveLinkSubject::CacheSettings(ULiveLinkSourceSettings* SourceSetting, ULi
 		}
 
 		SubjectRemapper.Reset();
-		if (SubjectSetting->Remapper)
+		
+		if (SubjectSetting->Remapper
+#if WITH_EDITOR
+		/** Remappers may call FinishCompilation for skinned assets, which isn't safe to do if we are running in a different thread. */
+		&& (IsInGameThread() || FSkinnedAssetCompilingManager::Get().GetNumRemainingJobs() == 0)
+#endif
+		)
 		{
 			// If there wasn't a remapper, then we need to initialize the one we will create.
 			bool bRecreateRemapper = SubjectSetting->Remapper->GetWorker() == nullptr || SubjectSetting->Remapper->bDirty;
diff --git a/Engine/Plugins/Animation/LiveLink/Source/LiveLink/Public/LiveLinkClient.h b/Engine/Plugins/Animation/LiveLink/Source/LiveLink/Public/LiveLinkClient.h
--- a/Engine/Plugins/Animation/LiveLink/Source/LiveLink/Public/LiveLinkClient.h
+++ b/Engine/Plugins/Animation/LiveLink/Source/LiveLink/Public/LiveLinkClient.h
@@ -124,7 +124,7 @@ public:
 	virtual FText GetSubjectDisplayName(const FLiveLinkSubjectKey& SubjectKey) const override;
 	virtual ULiveLinkSourceSettings* GetSourceSettings(const FGuid& SourceGuid) const override;
 	virtual UObject* GetSubjectSettings(const FLiveLinkSubjectKey& SubjectKey) const override;
-	virtual const FLiveLinkStaticDataStruct* GetSubjectStaticData_AnyThread(const FLiveLinkSubjectKey& InSubjectKey) const override;
+	virtual const FLiveLinkStaticDataStruct* GetSubjectStaticData_AnyThread(const FLiveLinkSubjectKey& InSubjectKey, bool bGetOverrideData=true) const override;
 
 
 	virtual bool EvaluateFrameFromSource_AnyThread(const FLiveLinkSubjectKey& SubjectKey, TSubclassOf<ULiveLinkRole> Role, FLiveLinkSubjectFrameData& OutFrame) override;
@@ -232,10 +232,24 @@ protected:
 	virtual FLiveLinkSkeletonStaticData* GetSubjectAnimationStaticData_Deprecation(const FLiveLinkSubjectKey& SubjectKey) override;
 	//~ End FLiveLinkClient_Base_DEPRECATED implementation
 
-	/** Add delegates that will be triggered for all subjects. */
-	bool RegisterGlobalSubjectFramesDelegate(const FOnLiveLinkSubjectStaticDataAdded::FDelegate& InOnStaticDataAdded, const FOnLiveLinkSubjectFrameDataAdded::FDelegate& InOnFrameDataAdded, FDelegateHandle& OutStaticDataAddedHandle, FDelegateHandle& OutFrameDataAddedHandle);
-	/** Remove the delegates that were triggered for all subjects. */
-	void UnregisterGlobalSubjectFramesDelegate(FDelegateHandle& InStaticDataAddedHandle, FDelegateHandle& InFrameDataAddedHandle);
+	/**
+	 * Add delegates that will be triggered for all subjects.
+	 * @param InOnStaticDataAdded The delegate for when static data is added.
+	 * @param InOnFrameDataAdded The delegate for when frame data is added.
+	 * @param OutStaticDataAddedHandle [Out] The handle for adding static data.
+	 * @param OutFrameDataAddedHandle [Out] The handle for adding frame data.
+	 * @param bUseUnmappedData Whether to use raw, unmapped data. If false, then the data received may have a remapper applied.
+	 */
+	bool RegisterGlobalSubjectFramesDelegate(const FOnLiveLinkSubjectStaticDataAdded::FDelegate& InOnStaticDataAdded,
+		const FOnLiveLinkSubjectFrameDataAdded::FDelegate& InOnFrameDataAdded, FDelegateHandle& OutStaticDataAddedHandle,
+		FDelegateHandle& OutFrameDataAddedHandle, bool bUseUnmappedData);
+	/**
+	 * Remove the delegates that were triggered for all subjects.
+	 * @param InStaticDataAddedHandle The static data handle to remove.
+	 * @param InFrameDataAddedHandle The frame data handle to remove.
+	 * @param bUseUnmappedData Whether this is for unmapped or remapped data.
+	 */
+	void UnregisterGlobalSubjectFramesDelegate(FDelegateHandle& InStaticDataAddedHandle, FDelegateHandle& InFrameDataAddedHandle, bool bUseUnmappedData);
 
 private:
 	/** Common initialization code for the different constructors. */
@@ -301,6 +315,10 @@ private:
 	{
 		FOnLiveLinkSubjectStaticDataAdded OnStaticDataAdded;
 		FOnLiveLinkSubjectFrameDataAdded OnFrameDataAdded;
+		/** Original data that hasn't been remapped. */
+		FOnLiveLinkSubjectStaticDataAdded OnUnmappedStaticDataAdded;
+		/** Original data that hasn't been remapped. */
+		FOnLiveLinkSubjectFrameDataAdded OnUnmappedFrameDataAdded;
 	};
 
 	/** Map of delegates to notify interested parties when the client receives a static or data frame for each subject */
diff --git a/Engine/Plugins/Animation/LiveLink/Source/LiveLink/Public/LiveLinkSubject.h b/Engine/Plugins/Animation/LiveLink/Source/LiveLink/Public/LiveLinkSubject.h
--- a/Engine/Plugins/Animation/LiveLink/Source/LiveLink/Public/LiveLinkSubject.h
+++ b/Engine/Plugins/Animation/LiveLink/Source/LiveLink/Public/LiveLinkSubject.h
@@ -61,7 +61,7 @@ public:
 	virtual FLiveLinkSubjectKey GetSubjectKey() const override { return SubjectKey; }
 	virtual TSubclassOf<ULiveLinkRole> GetRole() const override { return Role; }
 	virtual bool HasValidFrameSnapshot() const override;
-	virtual FLiveLinkStaticDataStruct& GetStaticData() override;
+	virtual FLiveLinkStaticDataStruct& GetStaticData(bool bGetOverrideData=true) override;
 	virtual const FLiveLinkStaticDataStruct& GetStaticData() const override;
 	virtual TArray<FLiveLinkTime> GetFrameTimes() const override;
 	virtual const TArray<ULiveLinkFrameTranslator::FWorkerSharedPtr> GetFrameTranslators() const override
diff --git a/Engine/Plugins/Animation/LiveLinkHub/Source/LiveLinkHub/Private/LiveLinkHubClient.cpp b/Engine/Plugins/Animation/LiveLinkHub/Source/LiveLinkHub/Private/LiveLinkHubClient.cpp
--- a/Engine/Plugins/Animation/LiveLinkHub/Source/LiveLinkHub/Private/LiveLinkHubClient.cpp
+++ b/Engine/Plugins/Animation/LiveLinkHub/Source/LiveLinkHub/Private/LiveLinkHubClient.cpp
@@ -25,7 +25,8 @@ DECLARE_CYCLE_STAT(TEXT("LiveLinkHub - Push FrameData"), STAT_LiveLinkHub_PushFr
 
 FLiveLinkHubClient::~FLiveLinkHubClient()
 {
-	UnregisterGlobalSubjectFramesDelegate(StaticDataAddedHandle, FrameDataAddedHandle);
+	constexpr bool bUseUnmappedData = true;
+	UnregisterGlobalSubjectFramesDelegate(StaticDataAddedHandle, FrameDataAddedHandle, bUseUnmappedData);
 }
 
 void FLiveLinkHubClient::CacheSubjectSettings(const FLiveLinkSubjectKey& SubjectKey, ULiveLinkSubjectSettings* Settings) const
diff --git a/Engine/Plugins/Animation/LiveLinkHub/Source/LiveLinkHub/Private/LiveLinkHubClient.h b/Engine/Plugins/Animation/LiveLinkHub/Source/LiveLinkHub/Private/LiveLinkHubClient.h
--- a/Engine/Plugins/Animation/LiveLinkHub/Source/LiveLinkHub/Private/LiveLinkHubClient.h
+++ b/Engine/Plugins/Animation/LiveLinkHub/Source/LiveLinkHub/Private/LiveLinkHubClient.h
@@ -26,14 +26,20 @@ public:
 		: FLiveLinkClient(InTickingDelegate)
 		, LiveLinkHub(MoveTemp(InLiveLinkHub))
 	{
-		RegisterGlobalSubjectFramesDelegate(FOnLiveLinkSubjectStaticDataAdded::FDelegate::CreateRaw(this, &FLiveLinkHubClient::OnStaticDataAdded), FOnLiveLinkSubjectFrameDataAdded::FDelegate::CreateRaw(this, &FLiveLinkHubClient::OnFrameDataAdded), StaticDataAddedHandle, FrameDataAddedHandle);
+		constexpr bool bUseUnmappedData = true;
+		RegisterGlobalSubjectFramesDelegate(FOnLiveLinkSubjectStaticDataAdded::FDelegate::CreateRaw(this, &FLiveLinkHubClient::OnStaticDataAdded),
+			FOnLiveLinkSubjectFrameDataAdded::FDelegate::CreateRaw(this, &FLiveLinkHubClient::OnFrameDataAdded),
+			StaticDataAddedHandle, FrameDataAddedHandle, bUseUnmappedData);
 	}
 
 	FLiveLinkHubClient(TSharedPtr<ILiveLinkHub> InLiveLinkHub)
 		: FLiveLinkClient()
 		, LiveLinkHub(MoveTemp(InLiveLinkHub))
 	{
-		RegisterGlobalSubjectFramesDelegate(FOnLiveLinkSubjectStaticDataAdded::FDelegate::CreateRaw(this, &FLiveLinkHubClient::OnStaticDataAdded), FOnLiveLinkSubjectFrameDataAdded::FDelegate::CreateRaw(this, &FLiveLinkHubClient::OnFrameDataAdded), StaticDataAddedHandle, FrameDataAddedHandle);
+		constexpr bool bUseUnmappedData = true;
+		RegisterGlobalSubjectFramesDelegate(FOnLiveLinkSubjectStaticDataAdded::FDelegate::CreateRaw(this, &FLiveLinkHubClient::OnStaticDataAdded),
+			FOnLiveLinkSubjectFrameDataAdded::FDelegate::CreateRaw(this, &FLiveLinkHubClient::OnFrameDataAdded),
+			StaticDataAddedHandle, FrameDataAddedHandle, bUseUnmappedData);
 	}
 
 	virtual ~FLiveLinkHubClient();
diff --git a/Engine/Plugins/Animation/LiveLinkHub/Source/LiveLinkHub/Private/Recording/Implementations/LiveLinkUAssetRecorder.cpp b/Engine/Plugins/Animation/LiveLinkHub/Source/LiveLinkHub/Private/Recording/Implementations/LiveLinkUAssetRecorder.cpp
--- a/Engine/Plugins/Animation/LiveLinkHub/Source/LiveLinkHub/Private/Recording/Implementations/LiveLinkUAssetRecorder.cpp
+++ b/Engine/Plugins/Animation/LiveLinkHub/Source/LiveLinkHub/Private/Recording/Implementations/LiveLinkUAssetRecorder.cpp
@@ -45,7 +45,7 @@ namespace UAssetRecorderUtils
 		FLiveLinkHubClient* LiveLinkClient = static_cast<FLiveLinkHubClient*>(&IModularFeatures::Get().GetModularFeature<ILiveLinkClient>(ILiveLinkClient::ModularFeatureName));
 		TSubclassOf<ULiveLinkRole> LiveLinkRole = LiveLinkClient->GetSubjectRole_AnyThread(SubjectKey);
 		
-		if (const FLiveLinkStaticDataStruct* StaticData = LiveLinkClient->GetSubjectStaticData_AnyThread(SubjectKey))
+		if (const FLiveLinkStaticDataStruct* StaticData = LiveLinkClient->GetSubjectStaticData_AnyThread(SubjectKey, /* bGetOverrideData */ false))
 		{
 			if (StaticData->IsValid())
 			{
diff --git a/Engine/Plugins/Animation/LiveLinkHub/Source/LiveLinkHub/Private/Recording/LiveLinkHubPlaybackController.cpp b/Engine/Plugins/Animation/LiveLinkHub/Source/LiveLinkHub/Private/Recording/LiveLinkHubPlaybackController.cpp
--- a/Engine/Plugins/Animation/LiveLinkHub/Source/LiveLinkHub/Private/Recording/LiveLinkHubPlaybackController.cpp
+++ b/Engine/Plugins/Animation/LiveLinkHub/Source/LiveLinkHub/Private/Recording/LiveLinkHubPlaybackController.cpp
@@ -302,6 +302,8 @@ void FLiveLinkHubPlaybackController::Eject(TFunction<void()> CompletionCallback)
 	bIsReady = false;
 	
 	StopPlayback();
+
+	LastStaticFrameIndex.Empty();
 	
 	bIsPaused = false;
 	RecordingPlayer->RestartPlayback(0);
@@ -550,6 +552,17 @@ void FLiveLinkHubPlaybackController::PushSubjectData(const FLiveLinkRecordedFram
 	// If we're sending static data
 	if (NextFrame.LiveLinkRole)
 	{
+		// Make sure we only push static data if it has changed.
+    	if (const int32* Idx = LastStaticFrameIndex.Find(NextFrame.SubjectKey))
+    	{
+    		if (*Idx == NextFrame.FrameIndex)
+    		{
+    			return;
+    		}
+    	}
+
+    	LastStaticFrameIndex.Add(NextFrame.SubjectKey, NextFrame.FrameIndex);
+	
 		FLiveLinkStaticDataStruct StaticDataStruct;
 		StaticDataStruct.InitializeWith(NextFrame.Data.GetScriptStruct(), (FLiveLinkBaseStaticData*)NextFrame.Data.GetMemory());
 		Client->PushSubjectStaticData_AnyThread(NextFrame.SubjectKey, NextFrame.LiveLinkRole, MoveTemp(StaticDataStruct));
diff --git a/Engine/Plugins/Animation/LiveLinkHub/Source/LiveLinkHub/Private/Recording/LiveLinkHubPlaybackController.h b/Engine/Plugins/Animation/LiveLinkHub/Source/LiveLinkHub/Private/Recording/LiveLinkHubPlaybackController.h
--- a/Engine/Plugins/Animation/LiveLinkHub/Source/LiveLinkHub/Private/Recording/LiveLinkHubPlaybackController.h
+++ b/Engine/Plugins/Animation/LiveLinkHub/Source/LiveLinkHub/Private/Recording/LiveLinkHubPlaybackController.h
@@ -246,4 +246,7 @@ private:
 
 	/** Packages in the process of unloading. */
 	TSet<TWeakObjectPtr<UPackage>> PackagesUnloading;
+
+	/** The last sent static frame index. Used to ensure we only send static data when necessary. */
+	TMap<FLiveLinkSubjectKey, int32> LastStaticFrameIndex;
 };
diff --git a/Engine/Source/Runtime/LiveLinkInterface/Public/ILiveLinkClient.h b/Engine/Source/Runtime/LiveLinkInterface/Public/ILiveLinkClient.h
--- a/Engine/Source/Runtime/LiveLinkInterface/Public/ILiveLinkClient.h
+++ b/Engine/Source/Runtime/LiveLinkInterface/Public/ILiveLinkClient.h
@@ -309,9 +309,11 @@ public:
 	virtual FText GetSubjectDisplayName(const FLiveLinkSubjectKey& SubjectKey) const = 0;
 
 	/**
-	 * Utility method to grab a subject's static data. Used by the RecordingController when static data is missing from the recording. 
+	 * Utility method to grab a subject's static data. Used by the RecordingController when static data is missing from the recording.
+	 * @param InSubjectKey The subject key to use.
+	 * @param bGetOverrideData Whether to get static override data if it exists.
 	 */
-	virtual const FLiveLinkStaticDataStruct* GetSubjectStaticData_AnyThread(const FLiveLinkSubjectKey& InSubjectKey) const = 0;
+	virtual const FLiveLinkStaticDataStruct* GetSubjectStaticData_AnyThread(const FLiveLinkSubjectKey& InSubjectKey, bool bGetOverrideData=true) const = 0;
 
 	/**
 	 * Return the evaluated subject from a specific source snapshot for a specific role.
diff --git a/Engine/Source/Runtime/LiveLinkInterface/Public/ILiveLinkSubject.h b/Engine/Source/Runtime/LiveLinkInterface/Public/ILiveLinkSubject.h
--- a/Engine/Source/Runtime/LiveLinkInterface/Public/ILiveLinkSubject.h
+++ b/Engine/Source/Runtime/LiveLinkInterface/Public/ILiveLinkSubject.h
@@ -37,7 +37,7 @@ public:
 	LIVELINKINTERFACE_API virtual bool SupportsRole(TSubclassOf<ULiveLinkRole> InDesiredRole) const;
 
 	virtual bool HasValidFrameSnapshot() const = 0;
-	virtual FLiveLinkStaticDataStruct& GetStaticData() = 0;
+	virtual FLiveLinkStaticDataStruct& GetStaticData(bool bGetOverrideData=true) = 0;
 	virtual const FLiveLinkStaticDataStruct& GetStaticData() const = 0;
 	virtual TArray<FLiveLinkTime> GetFrameTimes() const = 0;
 
diff --git a/Engine/Source/Runtime/LiveLinkInterface/Public/LiveLinkVirtualSubject.h b/Engine/Source/Runtime/LiveLinkInterface/Public/LiveLinkVirtualSubject.h
--- a/Engine/Source/Runtime/LiveLinkInterface/Public/LiveLinkVirtualSubject.h
+++ b/Engine/Source/Runtime/LiveLinkInterface/Public/LiveLinkVirtualSubject.h
@@ -37,7 +37,7 @@ public:
 	virtual FLiveLinkSubjectKey GetSubjectKey() const override { return SubjectKey; }
 	virtual TSubclassOf<ULiveLinkRole> GetRole() const override { return Role; }
 	LIVELINKINTERFACE_API virtual bool HasValidFrameSnapshot() const override;
-	virtual FLiveLinkStaticDataStruct& GetStaticData() override { return CurrentFrameSnapshot.StaticData; }
+	virtual FLiveLinkStaticDataStruct& GetStaticData(bool bGetOverrideData=true) override { return CurrentFrameSnapshot.StaticData; }
 	virtual const FLiveLinkStaticDataStruct& GetStaticData() const override { return CurrentFrameSnapshot.StaticData; }
 	virtual const TArray<ULiveLinkFrameTranslator::FWorkerSharedPtr> GetFrameTranslators() const override { return CurrentFrameTranslators; }
 	virtual const ULiveLinkSubjectRemapper::FWorkerSharedPtr GetFrameRemapper() const override { return CurrentSubjectRemapper; }
