Fix FLiveLinkClient::CreateSubject creating subjects with an empty subject key

[RN] fix Virtual Production/Tools

diff --git a/Engine/Plugins/Animation/LiveLink/Source/LiveLink/Private/LiveLinkClient.cpp b/Engine/Plugins/Animation/LiveLink/Source/LiveLink/Private/LiveLinkClient.cpp
--- a/Engine/Plugins/Animation/LiveLink/Source/LiveLink/Private/LiveLinkClient.cpp
+++ b/Engine/Plugins/Animation/LiveLink/Source/LiveLink/Private/LiveLinkClient.cpp
@@ -1126,6 +1126,8 @@ bool FLiveLinkClient::CreateSubject(const FLiveLinkSubjectPreset& InSubjectPrese
 			SubjectSettings = NewObject<ULiveLinkSubjectSettings>();
 		}
 
+		SubjectSettings->Initialize(InSubjectPreset.Key);
+
 		bool bEnabled = false;
 
 		FLiveLinkCollectionSubjectItem CollectionSubjectItem(InSubjectPreset.Key, MakeUnique<FLiveLinkSubject>(SourceItem->TimedData), SubjectSettings, bEnabled);
