LiveLinkHub - Fix live subjects with a remapper applied not resuming when exiting playback.

diff --git a/Engine/Plugins/Animation/LiveLink/Source/LiveLink/Private/LiveLinkSubject.cpp b/Engine/Plugins/Animation/LiveLink/Source/LiveLink/Private/LiveLinkSubject.cpp
--- a/Engine/Plugins/Animation/LiveLink/Source/LiveLink/Private/LiveLinkSubject.cpp
+++ b/Engine/Plugins/Animation/LiveLink/Source/LiveLink/Private/LiveLinkSubject.cpp
@@ -1065,6 +1065,9 @@ void FLiveLinkSubject::SetStaticData(TSubclassOf<ULiveLinkRole> InRole, FLiveLin
 		FrameData.Reset();
 		ReceivedOrderedFrames.Empty();
 		StaticData = MoveTemp(InStaticData);
+		// SetStaticData can be called after caching settings and checking for a dirty remapper. But the remapper needs to be applied
+		// to this new static data, so remap again when caching settings.
+		bNeedsStaticRemap = OverrideStaticData.IsSet();
 	}
 	else
 	{
@@ -1141,7 +1144,8 @@ void FLiveLinkSubject::CacheSettings(ULiveLinkSourceSettings* SourceSetting, ULi
 		)
 		{
 			// If there wasn't a remapper, then we need to initialize the one we will create.
-			bool bRecreateRemapper = SubjectSetting->Remapper->GetWorker() == nullptr || SubjectSetting->Remapper->bDirty;
+			bool bRecreateRemapper = SubjectSetting->Remapper->GetWorker() == nullptr || SubjectSetting->Remapper->bDirty || bNeedsStaticRemap;
+			bNeedsStaticRemap = false;
 
 			ULiveLinkSubjectRemapper::FWorkerSharedPtr NewRemapper = bRecreateRemapper ? SubjectSetting->Remapper->CreateWorker() : SubjectSetting->Remapper->GetWorker();
 
diff --git a/Engine/Plugins/Animation/LiveLink/Source/LiveLink/Public/LiveLinkSubject.h b/Engine/Plugins/Animation/LiveLink/Source/LiveLink/Public/LiveLinkSubject.h
--- a/Engine/Plugins/Animation/LiveLink/Source/LiveLink/Public/LiveLinkSubject.h
+++ b/Engine/Plugins/Animation/LiveLink/Source/LiveLink/Public/LiveLinkSubject.h
@@ -246,6 +246,9 @@ private:
 	// If true, static data has been sent for this rebroadcast
 	bool bRebroadcastStaticDataSent = false;
 
+	// If true, override static data may remap when caching settings.
+	bool bNeedsStaticRemap = false;
+
 	/** Flag set to clear the override static data for a subject. */
 	std::atomic<bool> bClearOverrideStaticData = false;
 	
