Removed lots of UE_ENABLE_INCLUDE_ORDER_DEPRECATED_IN_5_2 scopes

diff --git a/Engine/Plugins/2D/Paper2D/Source/Paper2D/Classes/IntMargin.h b/Engine/Plugins/2D/Paper2D/Source/Paper2D/Classes/IntMargin.h
--- a/Engine/Plugins/2D/Paper2D/Source/Paper2D/Classes/IntMargin.h
+++ b/Engine/Plugins/2D/Paper2D/Source/Paper2D/Classes/IntMargin.h
@@ -124,8 +124,3 @@ public:
 		return FIntPoint(Left + Right, Top + Bottom);
 	}
 };
-
-#if UE_ENABLE_INCLUDE_ORDER_DEPRECATED_IN_5_2
-#include "CoreMinimal.h"
-#include "UObject/ObjectMacros.h"
-#endif
diff --git a/Engine/Plugins/2D/Paper2D/Source/Paper2D/Classes/PaperCharacter.h b/Engine/Plugins/2D/Paper2D/Source/Paper2D/Classes/PaperCharacter.h
--- a/Engine/Plugins/2D/Paper2D/Source/Paper2D/Classes/PaperCharacter.h
+++ b/Engine/Plugins/2D/Paper2D/Source/Paper2D/Classes/PaperCharacter.h
@@ -28,8 +28,3 @@ public:
 	/** Returns Sprite subobject **/
 	FORCEINLINE class UPaperFlipbookComponent* GetSprite() const { return Sprite; }
 };
-
-#if UE_ENABLE_INCLUDE_ORDER_DEPRECATED_IN_5_2
-#include "CoreMinimal.h"
-#include "UObject/ObjectMacros.h"
-#endif
diff --git a/Engine/Plugins/2D/Paper2D/Source/Paper2D/Classes/PaperFlipbook.h b/Engine/Plugins/2D/Paper2D/Source/Paper2D/Classes/PaperFlipbook.h
--- a/Engine/Plugins/2D/Paper2D/Source/Paper2D/Classes/PaperFlipbook.h
+++ b/Engine/Plugins/2D/Paper2D/Source/Paper2D/Classes/PaperFlipbook.h
@@ -189,11 +189,3 @@ public:
 		return SourceFlipbook;
 	}
 };
-
-#if UE_ENABLE_INCLUDE_ORDER_DEPRECATED_IN_5_2
-#include "CoreMinimal.h"
-#include "Engine/EngineTypes.h"
-#include "UObject/Object.h"
-#include "UObject/ObjectMacros.h"
-#include "UObject/ScriptMacros.h"
-#endif
diff --git a/Engine/Plugins/2D/Paper2D/Source/Paper2D/Classes/PaperFlipbookActor.h b/Engine/Plugins/2D/Paper2D/Source/Paper2D/Classes/PaperFlipbookActor.h
--- a/Engine/Plugins/2D/Paper2D/Source/Paper2D/Classes/PaperFlipbookActor.h
+++ b/Engine/Plugins/2D/Paper2D/Source/Paper2D/Classes/PaperFlipbookActor.h
@@ -33,9 +33,3 @@ public:
 	/** Returns RenderComponent subobject **/
 	FORCEINLINE class UPaperFlipbookComponent* GetRenderComponent() const { return RenderComponent; }
 };
-
-
-#if UE_ENABLE_INCLUDE_ORDER_DEPRECATED_IN_5_2
-#include "CoreMinimal.h"
-#include "UObject/ObjectMacros.h"
-#endif
diff --git a/Engine/Plugins/2D/Paper2D/Source/Paper2D/Classes/PaperFlipbookComponent.h b/Engine/Plugins/2D/Paper2D/Source/Paper2D/Classes/PaperFlipbookComponent.h
--- a/Engine/Plugins/2D/Paper2D/Source/Paper2D/Classes/PaperFlipbookComponent.h
+++ b/Engine/Plugins/2D/Paper2D/Source/Paper2D/Classes/PaperFlipbookComponent.h
@@ -212,10 +212,3 @@ public:
 	virtual UBodySetup* GetBodySetup() override;
 	// End of UPrimitiveComponent interface
 };
-
-#if UE_ENABLE_INCLUDE_ORDER_DEPRECATED_IN_5_2
-#include "CoreMinimal.h"
-#include "Engine/EngineTypes.h"
-#include "Engine/TextureStreamingTypes.h"
-#include "UObject/ObjectMacros.h"
-#endif
diff --git a/Engine/Plugins/2D/Paper2D/Source/Paper2D/Classes/PaperGroupedSpriteComponent.h b/Engine/Plugins/2D/Paper2D/Source/Paper2D/Classes/PaperGroupedSpriteComponent.h
--- a/Engine/Plugins/2D/Paper2D/Source/Paper2D/Classes/PaperGroupedSpriteComponent.h
+++ b/Engine/Plugins/2D/Paper2D/Source/Paper2D/Classes/PaperGroupedSpriteComponent.h
@@ -182,9 +182,3 @@ protected:
 
 	friend FGroupedSpriteSceneProxy;
 };
-
-#if UE_ENABLE_INCLUDE_ORDER_DEPRECATED_IN_5_2
-#include "CoreMinimal.h"
-#include "Physics/PhysicsInterfaceDeclares.h"
-#include "UObject/ObjectMacros.h"
-#endif
diff --git a/Engine/Plugins/2D/Paper2D/Source/Paper2D/Classes/PaperRuntimeSettings.h b/Engine/Plugins/2D/Paper2D/Source/Paper2D/Classes/PaperRuntimeSettings.h
--- a/Engine/Plugins/2D/Paper2D/Source/Paper2D/Classes/PaperRuntimeSettings.h
+++ b/Engine/Plugins/2D/Paper2D/Source/Paper2D/Classes/PaperRuntimeSettings.h
@@ -24,9 +24,3 @@ class PAPER2D_API UPaperRuntimeSettings : public UObject
 	UPROPERTY(EditAnywhere, config, Category=Settings)
 	bool bResizeSpriteDataToMatchTextures;
 };
-
-#if UE_ENABLE_INCLUDE_ORDER_DEPRECATED_IN_5_2
-#include "CoreMinimal.h"
-#include "UObject/Object.h"
-#include "UObject/ObjectMacros.h"
-#endif
diff --git a/Engine/Plugins/2D/Paper2D/Source/Paper2D/Classes/PaperSprite.h b/Engine/Plugins/2D/Paper2D/Source/Paper2D/Classes/PaperSprite.h
--- a/Engine/Plugins/2D/Paper2D/Source/Paper2D/Classes/PaperSprite.h
+++ b/Engine/Plugins/2D/Paper2D/Source/Paper2D/Classes/PaperSprite.h
@@ -345,11 +345,3 @@ protected:
 	friend struct FPaperAtlasGenerator;
 	friend class FSingleTileEditorViewportClient;
 };
-
-
-#if UE_ENABLE_INCLUDE_ORDER_DEPRECATED_IN_5_2
-#include "CoreMinimal.h"
-#include "Engine/EngineTypes.h"
-#include "UObject/Object.h"
-#include "UObject/ObjectMacros.h"
-#endif
diff --git a/Engine/Plugins/2D/Paper2D/Source/Paper2D/Classes/PaperSpriteAtlas.h b/Engine/Plugins/2D/Paper2D/Source/Paper2D/Classes/PaperSpriteAtlas.h
--- a/Engine/Plugins/2D/Paper2D/Source/Paper2D/Classes/PaperSpriteAtlas.h
+++ b/Engine/Plugins/2D/Paper2D/Source/Paper2D/Classes/PaperSpriteAtlas.h
@@ -129,11 +129,3 @@ public:
 	virtual void GetAssetRegistryTags(TArray<FAssetRegistryTag>& OutTags) const override;
 	// End of UObject interface
 };
-
-#if UE_ENABLE_INCLUDE_ORDER_DEPRECATED_IN_5_2
-#include "CoreMinimal.h"
-#include "Misc/Guid.h"
-#include "PaperSprite.h"
-#include "UObject/Object.h"
-#include "UObject/ObjectMacros.h"
-#endif
diff --git a/Engine/Plugins/2D/Paper2D/Source/Paper2D/Classes/PaperSpriteBlueprintLibrary.h b/Engine/Plugins/2D/Paper2D/Source/Paper2D/Classes/PaperSpriteBlueprintLibrary.h
--- a/Engine/Plugins/2D/Paper2D/Source/Paper2D/Classes/PaperSpriteBlueprintLibrary.h
+++ b/Engine/Plugins/2D/Paper2D/Source/Paper2D/Classes/PaperSpriteBlueprintLibrary.h
@@ -21,9 +21,3 @@ public:
 	UFUNCTION(BlueprintPure, Category="Widget|Brush")
 	static FSlateBrush MakeBrushFromSprite(UPaperSprite* Sprite, int32 Width, int32 Height);
 };
-
-#if UE_ENABLE_INCLUDE_ORDER_DEPRECATED_IN_5_2
-#include "CoreMinimal.h"
-#include "Styling/SlateBrush.h"
-#include "UObject/ObjectMacros.h"
-#endif
diff --git a/Engine/Plugins/2D/Paper2D/Source/Paper2D/Classes/PaperSpriteComponent.h b/Engine/Plugins/2D/Paper2D/Source/Paper2D/Classes/PaperSpriteComponent.h
--- a/Engine/Plugins/2D/Paper2D/Source/Paper2D/Classes/PaperSpriteComponent.h
+++ b/Engine/Plugins/2D/Paper2D/Source/Paper2D/Classes/PaperSpriteComponent.h
@@ -102,10 +102,3 @@ public:
 protected:
 	friend class FPaperSpriteSceneProxy;
 };
-
-#if UE_ENABLE_INCLUDE_ORDER_DEPRECATED_IN_5_2
-#include "CoreMinimal.h"
-#include "Engine/EngineTypes.h"
-#include "Engine/TextureStreamingTypes.h"
-#include "UObject/ObjectMacros.h"
-#endif
diff --git a/Engine/Plugins/2D/Paper2D/Source/Paper2D/Classes/PaperTerrainComponent.h b/Engine/Plugins/2D/Paper2D/Source/Paper2D/Classes/PaperTerrainComponent.h
--- a/Engine/Plugins/2D/Paper2D/Source/Paper2D/Classes/PaperTerrainComponent.h
+++ b/Engine/Plugins/2D/Paper2D/Source/Paper2D/Classes/PaperTerrainComponent.h
@@ -135,10 +135,3 @@ protected:
 	
 	FTransform GetTransformAtDistance(float InDistance) const;
 };
-
-
-#if UE_ENABLE_INCLUDE_ORDER_DEPRECATED_IN_5_2
-#include "CoreMinimal.h"
-#include "SpriteDrawCall.h"
-#include "UObject/ObjectMacros.h"
-#endif
diff --git a/Engine/Plugins/2D/Paper2D/Source/Paper2D/Classes/PaperTerrainSplineComponent.h b/Engine/Plugins/2D/Paper2D/Source/Paper2D/Classes/PaperTerrainSplineComponent.h
--- a/Engine/Plugins/2D/Paper2D/Source/Paper2D/Classes/PaperTerrainSplineComponent.h
+++ b/Engine/Plugins/2D/Paper2D/Source/Paper2D/Classes/PaperTerrainSplineComponent.h
@@ -22,9 +22,3 @@ public:
 	// Triggered when the spline is edited
 	FSimpleDelegate OnSplineEdited;
 };
-
-
-#if UE_ENABLE_INCLUDE_ORDER_DEPRECATED_IN_5_2
-#include "CoreMinimal.h"
-#include "UObject/ObjectMacros.h"
-#endif
diff --git a/Engine/Plugins/2D/Paper2D/Source/Paper2D/Classes/PaperTileLayer.h b/Engine/Plugins/2D/Paper2D/Source/Paper2D/Classes/PaperTileLayer.h
--- a/Engine/Plugins/2D/Paper2D/Source/Paper2D/Classes/PaperTileLayer.h
+++ b/Engine/Plugins/2D/Paper2D/Source/Paper2D/Classes/PaperTileLayer.h
@@ -272,9 +272,3 @@ protected:
 	void ReallocateAndCopyMap();
 
 };
-
-#if UE_ENABLE_INCLUDE_ORDER_DEPRECATED_IN_5_2
-#include "CoreMinimal.h"
-#include "UObject/Object.h"
-#include "UObject/ObjectMacros.h"
-#endif
diff --git a/Engine/Plugins/2D/Paper2D/Source/Paper2D/Classes/PaperTileMap.h b/Engine/Plugins/2D/Paper2D/Source/Paper2D/Classes/PaperTileMap.h
--- a/Engine/Plugins/2D/Paper2D/Source/Paper2D/Classes/PaperTileMap.h
+++ b/Engine/Plugins/2D/Paper2D/Source/Paper2D/Classes/PaperTileMap.h
@@ -243,10 +243,3 @@ public:
 protected:
 	virtual void UpdateBodySetup();
 };
-
-#if UE_ENABLE_INCLUDE_ORDER_DEPRECATED_IN_5_2
-#include "CoreMinimal.h"
-#include "SpriteEditorOnlyTypes.h"
-#include "UObject/Object.h"
-#include "UObject/ObjectMacros.h"
-#endif
diff --git a/Engine/Plugins/2D/Paper2D/Source/Paper2D/Classes/PaperTileMapComponent.h b/Engine/Plugins/2D/Paper2D/Source/Paper2D/Classes/PaperTileMapComponent.h
--- a/Engine/Plugins/2D/Paper2D/Source/Paper2D/Classes/PaperTileMapComponent.h
+++ b/Engine/Plugins/2D/Paper2D/Source/Paper2D/Classes/PaperTileMapComponent.h
@@ -235,10 +235,3 @@ public:
 	void GetRenderingStats(int32& OutNumTriangles, int32& OutNumBatches) const;
 #endif
 };
-
-#if UE_ENABLE_INCLUDE_ORDER_DEPRECATED_IN_5_2
-#include "CoreMinimal.h"
-#include "DynamicMeshBuilder.h"
-#include "PaperTileLayer.h"
-#include "UObject/ObjectMacros.h"
-#endif
diff --git a/Engine/Plugins/2D/Paper2D/Source/Paper2D/Classes/TileMapBlueprintLibrary.h b/Engine/Plugins/2D/Paper2D/Source/Paper2D/Classes/TileMapBlueprintLibrary.h
--- a/Engine/Plugins/2D/Paper2D/Source/Paper2D/Classes/TileMapBlueprintLibrary.h
+++ b/Engine/Plugins/2D/Paper2D/Source/Paper2D/Classes/TileMapBlueprintLibrary.h
@@ -36,9 +36,3 @@ public:
 	UFUNCTION(BlueprintPure, Category=Sprite, meta=(AdvancedDisplay=2))
 	static FPaperTileInfo MakeTile(int32 TileIndex, UPaperTileSet* TileSet, bool bFlipH, bool bFlipV, bool bFlipD);
 };
-
-#if UE_ENABLE_INCLUDE_ORDER_DEPRECATED_IN_5_2
-#include "CoreMinimal.h"
-#include "PaperTileLayer.h"
-#include "UObject/ObjectMacros.h"
-#endif
diff --git a/Engine/Plugins/2D/Paper2D/Source/Paper2D/Public/AlphaBitmap.h b/Engine/Plugins/2D/Paper2D/Source/Paper2D/Public/AlphaBitmap.h
--- a/Engine/Plugins/2D/Paper2D/Source/Paper2D/Public/AlphaBitmap.h
+++ b/Engine/Plugins/2D/Paper2D/Source/Paper2D/Public/AlphaBitmap.h
@@ -361,7 +361,3 @@ public:
 	int32 Height;
 	uint8 DefaultValue;
 };
-
-#if UE_ENABLE_INCLUDE_ORDER_DEPRECATED_IN_5_2
-#include "CoreMinimal.h"
-#endif
diff --git a/Engine/Plugins/2D/Paper2D/Source/Paper2D/Public/Paper2DModule.h b/Engine/Plugins/2D/Paper2D/Source/Paper2D/Public/Paper2DModule.h
--- a/Engine/Plugins/2D/Paper2D/Source/Paper2D/Public/Paper2DModule.h
+++ b/Engine/Plugins/2D/Paper2D/Source/Paper2D/Public/Paper2DModule.h
@@ -37,8 +37,3 @@ extern PAPER2D_API FVector PaperAxisZ;
 
 DECLARE_STATS_GROUP(TEXT("Paper2D"), STATGROUP_Paper2D, STATCAT_Advanced);
 DECLARE_CYCLE_STAT_EXTERN(TEXT("SetSprite (RT)"), STAT_PaperRender_SetSpriteRT, STATGROUP_Paper2D, PAPER2D_API);
-
-#if UE_ENABLE_INCLUDE_ORDER_DEPRECATED_IN_5_2
-#include "CoreMinimal.h"
-#include "Stats/Stats.h"
-#endif
diff --git a/Engine/Plugins/2D/Paper2D/Source/Paper2DEditor/Classes/PaperImporterSettings.h b/Engine/Plugins/2D/Paper2D/Source/Paper2DEditor/Classes/PaperImporterSettings.h
--- a/Engine/Plugins/2D/Paper2D/Source/Paper2DEditor/Classes/PaperImporterSettings.h
+++ b/Engine/Plugins/2D/Paper2D/Source/Paper2DEditor/Classes/PaperImporterSettings.h
@@ -157,11 +157,3 @@ public:
 	// Returns the default pixels/uu setting
 	float GetDefaultPixelsPerUnrealUnit() const { return DefaultPixelsPerUnrealUnit; }
 };
-
-#if UE_ENABLE_INCLUDE_ORDER_DEPRECATED_IN_5_2
-#include "CoreMinimal.h"
-#include "Engine/Texture.h"
-#include "UObject/Object.h"
-#include "UObject/ObjectMacros.h"
-#include "UObject/SoftObjectPath.h"
-#endif
diff --git a/Engine/Plugins/2D/Paper2D/Source/Paper2DEditor/Classes/TileMapAssetImportData.h b/Engine/Plugins/2D/Paper2D/Source/Paper2DEditor/Classes/TileMapAssetImportData.h
--- a/Engine/Plugins/2D/Paper2D/Source/Paper2DEditor/Classes/TileMapAssetImportData.h
+++ b/Engine/Plugins/2D/Paper2D/Source/Paper2DEditor/Classes/TileMapAssetImportData.h
@@ -35,9 +35,3 @@ class PAPER2DEDITOR_API UTileMapAssetImportData : public UAssetImportData
 
 	static UTileMapAssetImportData* GetImportDataForTileMap(class UPaperTileMap* TileMap);
 };
-
-#if UE_ENABLE_INCLUDE_ORDER_DEPRECATED_IN_5_2
-#include "CoreMinimal.h"
-#include "Engine/Texture.h"
-#include "UObject/ObjectMacros.h"
-#endif
diff --git a/Engine/Plugins/2D/Paper2D/Source/Paper2DEditor/Public/Paper2DEditorModule.h b/Engine/Plugins/2D/Paper2D/Source/Paper2DEditor/Public/Paper2DEditorModule.h
--- a/Engine/Plugins/2D/Paper2D/Source/Paper2DEditor/Public/Paper2DEditorModule.h
+++ b/Engine/Plugins/2D/Paper2D/Source/Paper2DEditor/Public/Paper2DEditorModule.h
@@ -21,9 +21,3 @@ public:
 
 	virtual uint32 GetPaper2DAssetCategory() const = 0;
 };
-
-
-#if UE_ENABLE_INCLUDE_ORDER_DEPRECATED_IN_5_2
-#include "CoreMinimal.h"
-#include "Toolkits/AssetEditorToolkit.h"
-#endif
diff --git a/Engine/Plugins/2D/Paper2D/Source/Paper2DEditor/Public/PaperFlipbookHelpers.h b/Engine/Plugins/2D/Paper2D/Source/Paper2DEditor/Public/PaperFlipbookHelpers.h
--- a/Engine/Plugins/2D/Paper2D/Source/Paper2DEditor/Public/PaperFlipbookHelpers.h
+++ b/Engine/Plugins/2D/Paper2D/Source/Paper2DEditor/Public/PaperFlipbookHelpers.h
@@ -29,8 +29,3 @@ private:
 
 	FPaperFlipbookHelpers() {}
 };
-
-
-#if UE_ENABLE_INCLUDE_ORDER_DEPRECATED_IN_5_2
-#include "CoreMinimal.h"
-#endif
diff --git a/Engine/Plugins/2D/Paper2D/Source/Paper2DEditor/Public/PaperJSONHelpers.h b/Engine/Plugins/2D/Paper2D/Source/Paper2DEditor/Public/PaperJSONHelpers.h
--- a/Engine/Plugins/2D/Paper2D/Source/Paper2DEditor/Public/PaperJSONHelpers.h
+++ b/Engine/Plugins/2D/Paper2D/Source/Paper2DEditor/Public/PaperJSONHelpers.h
@@ -49,9 +49,3 @@ public:
 private:
 	FPaperJSONHelpers() {}
 };
-
-
-#if UE_ENABLE_INCLUDE_ORDER_DEPRECATED_IN_5_2
-#include "CoreMinimal.h"
-#include "Dom/JsonObject.h"
-#endif
diff --git a/Engine/Plugins/2D/Paper2D/Source/PaperTiledImporter/Classes/PaperTiledImporterFactory.h b/Engine/Plugins/2D/Paper2D/Source/PaperTiledImporter/Classes/PaperTiledImporterFactory.h
--- a/Engine/Plugins/2D/Paper2D/Source/PaperTiledImporter/Classes/PaperTiledImporterFactory.h
+++ b/Engine/Plugins/2D/Paper2D/Source/PaperTiledImporter/Classes/PaperTiledImporterFactory.h
@@ -391,10 +391,3 @@ public:
 	FTileMapFromTiled();
 	bool IsValid() const;
 };
-
-#if UE_ENABLE_INCLUDE_ORDER_DEPRECATED_IN_5_2
-#include "CoreMinimal.h"
-#include "PaperTileLayer.h"
-#include "PaperTileMap.h"
-#include "UObject/ObjectMacros.h"
-#endif
diff --git a/Engine/Plugins/AI/AISupport/Source/AISupportModule/Public/AISupportModule.h b/Engine/Plugins/AI/AISupport/Source/AISupportModule/Public/AISupportModule.h
--- a/Engine/Plugins/AI/AISupport/Source/AISupportModule/Public/AISupportModule.h
+++ b/Engine/Plugins/AI/AISupport/Source/AISupportModule/Public/AISupportModule.h
@@ -33,8 +33,3 @@ public:
 		return FModuleManager::Get().IsModuleLoaded("AISupportModule");
 	}
 };
-
-
-#if UE_ENABLE_INCLUDE_ORDER_DEPRECATED_IN_5_2
-#include "CoreMinimal.h"
-#endif
diff --git a/Engine/Plugins/AI/EnvironmentQueryEditor/Source/EnvironmentQueryEditor/Public/EdGraphSchema_EnvironmentQuery.h b/Engine/Plugins/AI/EnvironmentQueryEditor/Source/EnvironmentQueryEditor/Public/EdGraphSchema_EnvironmentQuery.h
--- a/Engine/Plugins/AI/EnvironmentQueryEditor/Source/EnvironmentQueryEditor/Public/EdGraphSchema_EnvironmentQuery.h
+++ b/Engine/Plugins/AI/EnvironmentQueryEditor/Source/EnvironmentQueryEditor/Public/EdGraphSchema_EnvironmentQuery.h
@@ -26,8 +26,3 @@ class UEdGraphSchema_EnvironmentQuery : public UAIGraphSchema
 	virtual void GetSubNodeClasses(int32 SubNodeFlags, TArray<FGraphNodeClassData>& ClassData, UClass*& GraphNodeClass) const override;
 	//~ End UAIGraphSchema Interface
 };
-
-#if UE_ENABLE_INCLUDE_ORDER_DEPRECATED_IN_5_2
-#include "AIGraphTypes.h"
-#include "CoreMinimal.h"
-#endif
diff --git a/Engine/Plugins/AI/EnvironmentQueryEditor/Source/EnvironmentQueryEditor/Public/EnvironmentQueryEditorModule.h b/Engine/Plugins/AI/EnvironmentQueryEditor/Source/EnvironmentQueryEditor/Public/EnvironmentQueryEditorModule.h
--- a/Engine/Plugins/AI/EnvironmentQueryEditor/Source/EnvironmentQueryEditor/Public/EnvironmentQueryEditorModule.h
+++ b/Engine/Plugins/AI/EnvironmentQueryEditor/Source/EnvironmentQueryEditor/Public/EnvironmentQueryEditorModule.h
@@ -36,9 +36,3 @@ private:
 	TSharedPtr<FExtensibilityManager> MenuExtensibilityManager;
 	TSharedPtr<FExtensibilityManager> ToolBarExtensibilityManager;
 };
-
-#if UE_ENABLE_INCLUDE_ORDER_DEPRECATED_IN_5_2
-#include "AIGraphTypes.h"
-#include "CoreMinimal.h"
-#include "Toolkits/IToolkitHost.h"
-#endif
diff --git a/Engine/Plugins/AI/EnvironmentQueryEditor/Source/EnvironmentQueryEditor/Public/EnvironmentQueryFactory.h b/Engine/Plugins/AI/EnvironmentQueryEditor/Source/EnvironmentQueryEditor/Public/EnvironmentQueryFactory.h
--- a/Engine/Plugins/AI/EnvironmentQueryEditor/Source/EnvironmentQueryEditor/Public/EnvironmentQueryFactory.h
+++ b/Engine/Plugins/AI/EnvironmentQueryEditor/Source/EnvironmentQueryEditor/Public/EnvironmentQueryFactory.h
@@ -15,10 +15,3 @@ class UEnvironmentQueryFactory : public UFactory
 	virtual bool CanCreateNew() const override;
 	// End of UFactory interface	
 };
-
-
-
-
-#if UE_ENABLE_INCLUDE_ORDER_DEPRECATED_IN_5_2
-#include "CoreMinimal.h"
-#endif
diff --git a/Engine/Plugins/AI/EnvironmentQueryEditor/Source/EnvironmentQueryEditor/Public/EnvironmentQueryGraph.h b/Engine/Plugins/AI/EnvironmentQueryEditor/Source/EnvironmentQueryEditor/Public/EnvironmentQueryGraph.h
--- a/Engine/Plugins/AI/EnvironmentQueryEditor/Source/EnvironmentQueryEditor/Public/EnvironmentQueryGraph.h
+++ b/Engine/Plugins/AI/EnvironmentQueryEditor/Source/EnvironmentQueryEditor/Public/EnvironmentQueryGraph.h
@@ -42,7 +42,3 @@ protected:
 
 	void SpawnMissingSubNodes(UEnvQueryOption* Option, TSet<UEnvQueryTest*> ExistingTests, UEnvironmentQueryGraphNode* OptionNode);
 };
-
-#if UE_ENABLE_INCLUDE_ORDER_DEPRECATED_IN_5_2
-#include "CoreMinimal.h"
-#endif
diff --git a/Engine/Plugins/AI/EnvironmentQueryEditor/Source/EnvironmentQueryEditor/Public/IEnvironmentQueryEditor.h b/Engine/Plugins/AI/EnvironmentQueryEditor/Source/EnvironmentQueryEditor/Public/IEnvironmentQueryEditor.h
--- a/Engine/Plugins/AI/EnvironmentQueryEditor/Source/EnvironmentQueryEditor/Public/IEnvironmentQueryEditor.h
+++ b/Engine/Plugins/AI/EnvironmentQueryEditor/Source/EnvironmentQueryEditor/Public/IEnvironmentQueryEditor.h
@@ -12,9 +12,3 @@ class IEnvironmentQueryEditor : public FAssetEditorToolkit
 public:
 
 };
-
-
-
-#if UE_ENABLE_INCLUDE_ORDER_DEPRECATED_IN_5_2
-#include "CoreMinimal.h"
-#endif
diff --git a/Engine/Plugins/Animation/AnimationModifierLibrary/Source/AnimationModifierLibrary/Public/AnimationModifierLibrary.h b/Engine/Plugins/Animation/AnimationModifierLibrary/Source/AnimationModifierLibrary/Public/AnimationModifierLibrary.h
--- a/Engine/Plugins/Animation/AnimationModifierLibrary/Source/AnimationModifierLibrary/Public/AnimationModifierLibrary.h
+++ b/Engine/Plugins/Animation/AnimationModifierLibrary/Source/AnimationModifierLibrary/Public/AnimationModifierLibrary.h
@@ -12,8 +12,3 @@ public:
 	virtual void StartupModule() override;
 	virtual void ShutdownModule() override;
 };
-
-#if UE_ENABLE_INCLUDE_ORDER_DEPRECATED_IN_5_2
-#include "CoreMinimal.h"
-#include "Modules/ModuleManager.h"
-#endif
diff --git a/Engine/Plugins/Animation/AnimationModifierLibrary/Source/AnimationModifierLibrary/Public/MotionExtractorModifier.h b/Engine/Plugins/Animation/AnimationModifierLibrary/Source/AnimationModifierLibrary/Public/MotionExtractorModifier.h
--- a/Engine/Plugins/Animation/AnimationModifierLibrary/Source/AnimationModifierLibrary/Public/MotionExtractorModifier.h
+++ b/Engine/Plugins/Animation/AnimationModifierLibrary/Source/AnimationModifierLibrary/Public/MotionExtractorModifier.h
@@ -93,8 +93,3 @@ public:
 
 	virtual void PostLoad() override;
 };
-
-#if UE_ENABLE_INCLUDE_ORDER_DEPRECATED_IN_5_2
-#include "CoreMinimal.h"
-#include "MotionExtractorTypes.h"
-#endif
diff --git a/Engine/Plugins/Animation/AnimationModifierLibrary/Source/AnimationModifierLibrary/Public/MotionExtractorTypes.h b/Engine/Plugins/Animation/AnimationModifierLibrary/Source/AnimationModifierLibrary/Public/MotionExtractorTypes.h
--- a/Engine/Plugins/Animation/AnimationModifierLibrary/Source/AnimationModifierLibrary/Public/MotionExtractorTypes.h
+++ b/Engine/Plugins/Animation/AnimationModifierLibrary/Source/AnimationModifierLibrary/Public/MotionExtractorTypes.h
@@ -58,8 +58,3 @@ enum class EMotionExtractor_MathOperation : uint8
 	Division,
 	Multiplication
 };
-
-#if UE_ENABLE_INCLUDE_ORDER_DEPRECATED_IN_5_2
-#include "AnimationBlueprintLibrary.h"
-#include "CoreMinimal.h"
-#endif
diff --git a/Engine/Plugins/Animation/AnimationModifierLibrary/Source/AnimationModifierLibrary/Public/MotionExtractorUtilities.h b/Engine/Plugins/Animation/AnimationModifierLibrary/Source/AnimationModifierLibrary/Public/MotionExtractorUtilities.h
--- a/Engine/Plugins/Animation/AnimationModifierLibrary/Source/AnimationModifierLibrary/Public/MotionExtractorUtilities.h
+++ b/Engine/Plugins/Animation/AnimationModifierLibrary/Source/AnimationModifierLibrary/Public/MotionExtractorUtilities.h
@@ -78,7 +78,3 @@ public:
 	/** Helper function to calculate the magnitude of a vector only considering a specific axis or axes */
 	static double CalculateMagnitude(const FVector& Vector, EMotionExtractor_Axis Axis);
 };
-
-#if UE_ENABLE_INCLUDE_ORDER_DEPRECATED_IN_5_2
-#include "MotionExtractorTypes.h"
-#endif
diff --git a/Engine/Plugins/Animation/AnimationWarping/Source/Editor/Public/AnimationWarpingEditorModule.h b/Engine/Plugins/Animation/AnimationWarping/Source/Editor/Public/AnimationWarpingEditorModule.h
--- a/Engine/Plugins/Animation/AnimationWarping/Source/Editor/Public/AnimationWarpingEditorModule.h
+++ b/Engine/Plugins/Animation/AnimationWarping/Source/Editor/Public/AnimationWarpingEditorModule.h
@@ -7,9 +7,3 @@
 class IAnimationWarpingEditorModule : public IModuleInterface
 {
 };
-
-
-#if UE_ENABLE_INCLUDE_ORDER_DEPRECATED_IN_5_2
-#include "CoreMinimal.h"
-#include "Modules/ModuleManager.h"
-#endif
diff --git a/Engine/Plugins/Animation/AnimationWarping/Source/Runtime/Public/AnimationWarpingRuntimeModule.h b/Engine/Plugins/Animation/AnimationWarping/Source/Runtime/Public/AnimationWarpingRuntimeModule.h
--- a/Engine/Plugins/Animation/AnimationWarping/Source/Runtime/Public/AnimationWarpingRuntimeModule.h
+++ b/Engine/Plugins/Animation/AnimationWarping/Source/Runtime/Public/AnimationWarpingRuntimeModule.h
@@ -7,9 +7,3 @@
 class IAnimationWarpingRuntimeModule : public IModuleInterface
 {
 };
-
-
-#if UE_ENABLE_INCLUDE_ORDER_DEPRECATED_IN_5_2
-#include "CoreMinimal.h"
-#include "Modules/ModuleManager.h"
-#endif
diff --git a/Engine/Plugins/Animation/AnimationWarping/Source/Runtime/Public/BoneControllers/AnimNode_OffsetRootBone.h b/Engine/Plugins/Animation/AnimationWarping/Source/Runtime/Public/BoneControllers/AnimNode_OffsetRootBone.h
--- a/Engine/Plugins/Animation/AnimationWarping/Source/Runtime/Public/BoneControllers/AnimNode_OffsetRootBone.h
+++ b/Engine/Plugins/Animation/AnimationWarping/Source/Runtime/Public/BoneControllers/AnimNode_OffsetRootBone.h
@@ -207,7 +207,3 @@ private:
 
 	FGraphTraversalCounter UpdateCounter;
 };
-
-#if UE_ENABLE_INCLUDE_ORDER_DEPRECATED_IN_5_2
-#include "BoneControllers/BoneControllerSolvers.h"
-#endif
diff --git a/Engine/Plugins/Animation/AnimationWarping/Source/Runtime/Public/BoneControllers/AnimNode_SlopeWarping.h b/Engine/Plugins/Animation/AnimationWarping/Source/Runtime/Public/BoneControllers/AnimNode_SlopeWarping.h
--- a/Engine/Plugins/Animation/AnimationWarping/Source/Runtime/Public/BoneControllers/AnimNode_SlopeWarping.h
+++ b/Engine/Plugins/Animation/AnimationWarping/Source/Runtime/Public/BoneControllers/AnimNode_SlopeWarping.h
@@ -164,7 +164,3 @@ private:
 	virtual void InitializeBoneReferences(const FBoneContainer& RequiredBones) override;
 	// End of FAnimNode_SkeletalControlBase interface
 };
-
-#if UE_ENABLE_INCLUDE_ORDER_DEPRECATED_IN_5_2
-#include "BoneControllers/BoneControllerTypes.h"
-#endif
diff --git a/Engine/Plugins/Animation/BlendSpaceMotionAnalysis/Source/BlendSpaceMotionAnalysis/Public/LocomotionAnalysis.h b/Engine/Plugins/Animation/BlendSpaceMotionAnalysis/Source/BlendSpaceMotionAnalysis/Public/LocomotionAnalysis.h
--- a/Engine/Plugins/Animation/BlendSpaceMotionAnalysis/Source/BlendSpaceMotionAnalysis/Public/LocomotionAnalysis.h
+++ b/Engine/Plugins/Animation/BlendSpaceMotionAnalysis/Source/BlendSpaceMotionAnalysis/Public/LocomotionAnalysis.h
@@ -58,8 +58,3 @@ bool CalculateLocomotion(
 	const ULocomotionAnalysisProperties* AnalysisProperties,
 	const UAnimSequence&                 Animation,
 	const float                          RateScale);
-
-#if UE_ENABLE_INCLUDE_ORDER_DEPRECATED_IN_5_2
-#include "CoreMinimal.h"
-#include "Modules/ModuleInterface.h"
-#endif
diff --git a/Engine/Plugins/Animation/BlendSpaceMotionAnalysis/Source/BlendSpaceMotionAnalysis/Public/RootMotionAnalysis.h b/Engine/Plugins/Animation/BlendSpaceMotionAnalysis/Source/BlendSpaceMotionAnalysis/Public/RootMotionAnalysis.h
--- a/Engine/Plugins/Animation/BlendSpaceMotionAnalysis/Source/BlendSpaceMotionAnalysis/Public/RootMotionAnalysis.h
+++ b/Engine/Plugins/Animation/BlendSpaceMotionAnalysis/Source/BlendSpaceMotionAnalysis/Public/RootMotionAnalysis.h
@@ -75,9 +75,3 @@ bool CalculateRootMotion(
 	const URootMotionAnalysisProperties* AnalysisProperties,
 	const UAnimSequence&                 Animation,
 	const float                          RateScale);
-
-
-#if UE_ENABLE_INCLUDE_ORDER_DEPRECATED_IN_5_2
-#include "CoreMinimal.h"
-#include "Modules/ModuleInterface.h"
-#endif
diff --git a/Engine/Plugins/Animation/DeformerGraph/Source/OptimusSettings/Public/OptimusSettings.h b/Engine/Plugins/Animation/DeformerGraph/Source/OptimusSettings/Public/OptimusSettings.h
--- a/Engine/Plugins/Animation/DeformerGraph/Source/OptimusSettings/Public/OptimusSettings.h
+++ b/Engine/Plugins/Animation/DeformerGraph/Source/OptimusSettings/Public/OptimusSettings.h
@@ -4,9 +4,6 @@
 
 #include "CoreMinimal.h"
 #include "Engine/DeveloperSettings.h"
-#if UE_ENABLE_INCLUDE_ORDER_DEPRECATED_IN_5_2
-#include "RHIDefinitions.h"
-#endif
 #include "OptimusSettings.generated.h"
 
 class UMeshDeformer;
diff --git a/Engine/Plugins/Animation/IKRig/Source/IKRigEditor/Public/RigEditor/IKRigHitProxies.h b/Engine/Plugins/Animation/IKRig/Source/IKRigEditor/Public/RigEditor/IKRigHitProxies.h
--- a/Engine/Plugins/Animation/IKRig/Source/IKRigEditor/Public/RigEditor/IKRigHitProxies.h
+++ b/Engine/Plugins/Animation/IKRig/Source/IKRigEditor/Public/RigEditor/IKRigHitProxies.h
@@ -5,10 +5,6 @@
 #include "HitProxies.h"
 #include "UObject/NameTypes.h"
 
-#if UE_ENABLE_INCLUDE_ORDER_DEPRECATED_IN_5_2
-#include "CoreMinimal.h"
-#endif
-
 struct HIKRigEditorBoneProxy : public HHitProxy
 {
 	DECLARE_HIT_PROXY();
diff --git a/Engine/Plugins/Animation/LiveLink/Source/LiveLink/Public/LiveLinkBlueprintLibrary.h b/Engine/Plugins/Animation/LiveLink/Source/LiveLink/Public/LiveLinkBlueprintLibrary.h
--- a/Engine/Plugins/Animation/LiveLink/Source/LiveLink/Public/LiveLinkBlueprintLibrary.h
+++ b/Engine/Plugins/Animation/LiveLink/Source/LiveLink/Public/LiveLinkBlueprintLibrary.h
@@ -212,10 +212,3 @@ private:
 	DECLARE_FUNCTION(execEvaluateLiveLinkFrameAtWorldTimeOffset);
 	DECLARE_FUNCTION(execEvaluateLiveLinkFrameAtSceneTime);
 };
-
-#if UE_ENABLE_INCLUDE_ORDER_DEPRECATED_IN_5_2
-#include "CoreMinimal.h"
-#include "LiveLinkRole.h"
-#include "Roles/LiveLinkAnimationBlueprintStructs.h"
-#include "Roles/LiveLinkBasicTypes.h"
-#endif
diff --git a/Engine/Plugins/Animation/LiveLink/Source/LiveLink/Public/LiveLinkComponent.h b/Engine/Plugins/Animation/LiveLink/Source/LiveLink/Public/LiveLinkComponent.h
--- a/Engine/Plugins/Animation/LiveLink/Source/LiveLink/Public/LiveLinkComponent.h
+++ b/Engine/Plugins/Animation/LiveLink/Source/LiveLink/Public/LiveLinkComponent.h
@@ -68,10 +68,3 @@ private:
 
 	ILiveLinkClient* LiveLinkClient;
 };
-
-#if UE_ENABLE_INCLUDE_ORDER_DEPRECATED_IN_5_2
-#include "CoreMinimal.h"
-#include "ILiveLinkClient.h"
-#include "LiveLinkTypes.h"
-#include "Roles/LiveLinkAnimationBlueprintStructs.h"
-#endif
diff --git a/Engine/Plugins/Animation/LiveLink/Source/LiveLink/Public/LiveLinkMessageBusFinder.h b/Engine/Plugins/Animation/LiveLink/Source/LiveLink/Public/LiveLinkMessageBusFinder.h
--- a/Engine/Plugins/Animation/LiveLink/Source/LiveLink/Public/LiveLinkMessageBusFinder.h
+++ b/Engine/Plugins/Animation/LiveLink/Source/LiveLink/Public/LiveLinkMessageBusFinder.h
@@ -171,11 +171,3 @@ public:
 	}
 #endif
 };
-
-#if UE_ENABLE_INCLUDE_ORDER_DEPRECATED_IN_5_2
-#include "CoreMinimal.h"
-#include "ILiveLinkSource.h"
-#include "LiveLinkTypes.h"
-#include "MessageEndpoint.h"
-#include "Misc/ScopeLock.h"
-#endif
diff --git a/Engine/Plugins/Animation/LiveLink/Source/LiveLink/Public/LiveLinkMessageBusSource.h b/Engine/Plugins/Animation/LiveLink/Source/LiveLink/Public/LiveLinkMessageBusSource.h
--- a/Engine/Plugins/Animation/LiveLink/Source/LiveLink/Public/LiveLinkMessageBusSource.h
+++ b/Engine/Plugins/Animation/LiveLink/Source/LiveLink/Public/LiveLinkMessageBusSource.h
@@ -137,7 +137,3 @@ private:
 	// Offset between sender's machine engine time and receiver's machine engine time
 	double MachineTimeOffset;
 };
-
-#if UE_ENABLE_INCLUDE_ORDER_DEPRECATED_IN_5_2
-#include "LiveLinkRole.h"
-#endif
diff --git a/Engine/Plugins/Animation/LiveLink/Source/LiveLink/Public/LiveLinkPreset.h b/Engine/Plugins/Animation/LiveLink/Source/LiveLink/Public/LiveLinkPreset.h
--- a/Engine/Plugins/Animation/LiveLink/Source/LiveLink/Public/LiveLinkPreset.h
+++ b/Engine/Plugins/Animation/LiveLink/Source/LiveLink/Public/LiveLinkPreset.h
@@ -77,10 +77,3 @@ private:
 	/** Utility variable used to keep track of the number of times this was applied. */
 	uint32 ApplyCount = 0;
 };
-
-#if UE_ENABLE_INCLUDE_ORDER_DEPRECATED_IN_5_2
-#include "CoreMinimal.h"
-#include "Engine/EngineTypes.h"
-#include "Engine/LatentActionManager.h"
-#include "LiveLinkPresetTypes.h"
-#endif
diff --git a/Engine/Plugins/Animation/LiveLink/Source/LiveLink/Public/VirtualSubjects/LiveLinkBlueprintVirtualSubject.h b/Engine/Plugins/Animation/LiveLink/Source/LiveLink/Public/VirtualSubjects/LiveLinkBlueprintVirtualSubject.h
--- a/Engine/Plugins/Animation/LiveLink/Source/LiveLink/Public/VirtualSubjects/LiveLinkBlueprintVirtualSubject.h
+++ b/Engine/Plugins/Animation/LiveLink/Source/LiveLink/Public/VirtualSubjects/LiveLinkBlueprintVirtualSubject.h
@@ -46,8 +46,3 @@ private:
 	DECLARE_FUNCTION(execUpdateVirtualSubjectStaticData_Internal);
 	DECLARE_FUNCTION(execUpdateVirtualSubjectFrameData_Internal);
 };
-
-#if UE_ENABLE_INCLUDE_ORDER_DEPRECATED_IN_5_2
-#include "CoreMinimal.h"
-#include "ILiveLinkClient.h"
-#endif
diff --git a/Engine/Plugins/Animation/LiveLink/Source/LiveLinkComponents/Public/ILiveLinkComponentModule.h b/Engine/Plugins/Animation/LiveLink/Source/LiveLinkComponents/Public/ILiveLinkComponentModule.h
--- a/Engine/Plugins/Animation/LiveLink/Source/LiveLinkComponents/Public/ILiveLinkComponentModule.h
+++ b/Engine/Plugins/Animation/LiveLink/Source/LiveLinkComponents/Public/ILiveLinkComponentModule.h
@@ -22,9 +22,3 @@ public:
 	/** Virtual destructor. */
 	virtual ~ILiveLinkComponentsModule() { }
 };
-
-
-#if UE_ENABLE_INCLUDE_ORDER_DEPRECATED_IN_5_2
-#include "LiveLinkComponentController.h"
-#include "Modules/ModuleManager.h"
-#endif
diff --git a/Engine/Plugins/Animation/LiveLink/Source/LiveLinkComponents/Public/LiveLinkComponentController.h b/Engine/Plugins/Animation/LiveLink/Source/LiveLinkComponents/Public/LiveLinkComponentController.h
--- a/Engine/Plugins/Animation/LiveLink/Source/LiveLinkComponents/Public/LiveLinkComponentController.h
+++ b/Engine/Plugins/Animation/LiveLink/Source/LiveLinkComponents/Public/LiveLinkComponentController.h
@@ -143,7 +143,3 @@ protected:
 private:
 	FLiveLinkControllersTicked LiveLinkControllersTickedDelegate;
 };
-
-#if UE_ENABLE_INCLUDE_ORDER_DEPRECATED_IN_5_2
-#include "ILiveLinkClient.h"
-#endif
diff --git a/Engine/Plugins/Animation/LiveLink/Source/LiveLinkComponents/Public/LiveLinkComponentSettings.h b/Engine/Plugins/Animation/LiveLink/Source/LiveLinkComponents/Public/LiveLinkComponentSettings.h
--- a/Engine/Plugins/Animation/LiveLink/Source/LiveLinkComponents/Public/LiveLinkComponentSettings.h
+++ b/Engine/Plugins/Animation/LiveLink/Source/LiveLinkComponents/Public/LiveLinkComponentSettings.h
@@ -29,8 +29,3 @@ public:
 	UPROPERTY(config, EditAnywhere, Category = "LiveLink")
 	bool bUpdateSubjectsImmediatelyInPIE = true;
 };
-
-#if UE_ENABLE_INCLUDE_ORDER_DEPRECATED_IN_5_2
-#include "LiveLinkControllerBase.h"
-#include "LiveLinkRole.h"
-#endif
diff --git a/Engine/Plugins/Animation/LiveLink/Source/LiveLinkComponents/Public/LiveLinkControllerBase.h b/Engine/Plugins/Animation/LiveLink/Source/LiveLinkComponents/Public/LiveLinkControllerBase.h
--- a/Engine/Plugins/Animation/LiveLink/Source/LiveLinkComponents/Public/LiveLinkControllerBase.h
+++ b/Engine/Plugins/Animation/LiveLink/Source/LiveLinkComponents/Public/LiveLinkControllerBase.h
@@ -104,9 +104,3 @@ protected:
 
 	FLiveLinkSubjectRepresentation SelectedSubject;
 };
-
-
-#if UE_ENABLE_INCLUDE_ORDER_DEPRECATED_IN_5_2
-#include "CoreMinimal.h"
-#include "ILiveLinkClient.h"
-#endif
diff --git a/Engine/Plugins/Animation/LiveLink/Source/LiveLinkEditor/Public/SLiveLinkSubjectRepresentationPicker.h b/Engine/Plugins/Animation/LiveLink/Source/LiveLinkEditor/Public/SLiveLinkSubjectRepresentationPicker.h
--- a/Engine/Plugins/Animation/LiveLink/Source/LiveLinkEditor/Public/SLiveLinkSubjectRepresentationPicker.h
+++ b/Engine/Plugins/Animation/LiveLink/Source/LiveLinkEditor/Public/SLiveLinkSubjectRepresentationPicker.h
@@ -152,8 +152,3 @@ private:
 	/** Delegate used to retrieve a custom list of subjects. */
 	FOnGetSubjects GetSubjectsDelegate;
 };
-
-#if UE_ENABLE_INCLUDE_ORDER_DEPRECATED_IN_5_2
-#include "LiveLinkPreset.h"
-#include "Widgets/Views/SListView.h"
-#endif
diff --git a/Engine/Plugins/Animation/LiveLink/Source/LiveLinkGraphNode/Public/K2Node_EvaluateLiveLinkFrame.h b/Engine/Plugins/Animation/LiveLink/Source/LiveLinkGraphNode/Public/K2Node_EvaluateLiveLinkFrame.h
--- a/Engine/Plugins/Animation/LiveLink/Source/LiveLinkGraphNode/Public/K2Node_EvaluateLiveLinkFrame.h
+++ b/Engine/Plugins/Animation/LiveLink/Source/LiveLinkGraphNode/Public/K2Node_EvaluateLiveLinkFrame.h
@@ -89,10 +89,3 @@ protected:
 	/** Verify if selected role class is valid for evaluation */
 	bool IsRoleValidForEvaluation(TSubclassOf<ULiveLinkRole> InRoleClass) const;
 };
-
-#if UE_ENABLE_INCLUDE_ORDER_DEPRECATED_IN_5_2
-#include "CoreMinimal.h"
-#include "EdGraph/EdGraphNodeUtils.h"
-#include "LiveLinkRole.h"
-#include "Textures/SlateIcon.h"
-#endif
diff --git a/Engine/Plugins/Animation/LiveLink/Source/LiveLinkGraphNode/Public/K2Node_UpdateVirtualSubjectDataBase.h b/Engine/Plugins/Animation/LiveLink/Source/LiveLinkGraphNode/Public/K2Node_UpdateVirtualSubjectDataBase.h
--- a/Engine/Plugins/Animation/LiveLink/Source/LiveLinkGraphNode/Public/K2Node_UpdateVirtualSubjectDataBase.h
+++ b/Engine/Plugins/Animation/LiveLink/Source/LiveLinkGraphNode/Public/K2Node_UpdateVirtualSubjectDataBase.h
@@ -76,8 +76,3 @@ private:
 	/** Struct pin description */
 	static const FText LiveLinkStructPinDescription;
 };
-
-#if UE_ENABLE_INCLUDE_ORDER_DEPRECATED_IN_5_2
-#include "CoreMinimal.h"
-#include "Textures/SlateIcon.h"
-#endif
diff --git a/Engine/Plugins/Animation/LiveLink/Source/LiveLinkMovieScene/Public/MovieScene/IMovieSceneLiveLinkPropertyHandler.h b/Engine/Plugins/Animation/LiveLink/Source/LiveLinkMovieScene/Public/MovieScene/IMovieSceneLiveLinkPropertyHandler.h
--- a/Engine/Plugins/Animation/LiveLink/Source/LiveLinkMovieScene/Public/MovieScene/IMovieSceneLiveLinkPropertyHandler.h
+++ b/Engine/Plugins/Animation/LiveLink/Source/LiveLinkMovieScene/Public/MovieScene/IMovieSceneLiveLinkPropertyHandler.h
@@ -40,9 +40,3 @@ public:
 	virtual void FillFrame(int32 InKeyIndex, const FLiveLinkWorldTime& InWorldTime, const TOptional<FQualifiedFrameTime>& InTimecodeTime, const UScriptStruct& InStruct, FLiveLinkBaseFrameData* OutFrame) = 0;
 	virtual void FillFrameInterpolated(const FFrameTime& InFrameTime, const FLiveLinkWorldTime& InWorldTime, const TOptional<FQualifiedFrameTime>& InTimecodeTime, const UScriptStruct& InStruct, FLiveLinkBaseFrameData* OutFrame) = 0;
 };
-
-
-#if UE_ENABLE_INCLUDE_ORDER_DEPRECATED_IN_5_2
-#include "CoreMinimal.h"
-#include "LiveLinkTypes.h"
-#endif
diff --git a/Engine/Plugins/Animation/LiveLink/Source/LiveLinkMovieScene/Public/MovieScene/MovieSceneLiveLinkBufferData.h b/Engine/Plugins/Animation/LiveLink/Source/LiveLinkMovieScene/Public/MovieScene/MovieSceneLiveLinkBufferData.h
--- a/Engine/Plugins/Animation/LiveLink/Source/LiveLinkMovieScene/Public/MovieScene/MovieSceneLiveLinkBufferData.h
+++ b/Engine/Plugins/Animation/LiveLink/Source/LiveLinkMovieScene/Public/MovieScene/MovieSceneLiveLinkBufferData.h
@@ -136,7 +136,3 @@ struct FLiveLinkTransformKeys
 		}
 	}
 };
-
-#if UE_ENABLE_INCLUDE_ORDER_DEPRECATED_IN_5_2
-#include "Channels/MovieSceneChannelProxy.h"
-#endif
diff --git a/Engine/Plugins/Animation/LiveLink/Source/LiveLinkMovieScene/Public/MovieScene/MovieSceneLiveLinkEnumHandler.h b/Engine/Plugins/Animation/LiveLink/Source/LiveLinkMovieScene/Public/MovieScene/MovieSceneLiveLinkEnumHandler.h
--- a/Engine/Plugins/Animation/LiveLink/Source/LiveLinkMovieScene/Public/MovieScene/MovieSceneLiveLinkEnumHandler.h
+++ b/Engine/Plugins/Animation/LiveLink/Source/LiveLinkMovieScene/Public/MovieScene/MovieSceneLiveLinkEnumHandler.h
@@ -53,12 +53,3 @@ protected:
 	/** Enum property value are interacted as int64. We store keys like this but channel used will be a byte one */
 	TArray<TArray<FLiveLinkPropertyKey<int64>>> Keys;
 };
-
-
-
-
-
-#if UE_ENABLE_INCLUDE_ORDER_DEPRECATED_IN_5_2
-#include "CoreMinimal.h"
-#include "MovieScene/MovieSceneLiveLinkStructProperties.h"
-#endif
diff --git a/Engine/Plugins/Animation/LiveLink/Source/LiveLinkMovieScene/Public/MovieScene/MovieSceneLiveLinkSection.h b/Engine/Plugins/Animation/LiveLink/Source/LiveLinkMovieScene/Public/MovieScene/MovieSceneLiveLinkSection.h
--- a/Engine/Plugins/Animation/LiveLink/Source/LiveLinkMovieScene/Public/MovieScene/MovieSceneLiveLinkSection.h
+++ b/Engine/Plugins/Animation/LiveLink/Source/LiveLinkMovieScene/Public/MovieScene/MovieSceneLiveLinkSection.h
@@ -100,9 +100,3 @@ public:
 	TArray <FMovieSceneFloatChannel> PropertyFloatChannels_DEPRECATED;
 };
 PRAGMA_ENABLE_DEPRECATION_WARNINGS
-
-#if UE_ENABLE_INCLUDE_ORDER_DEPRECATED_IN_5_2
-#include "Channels/MovieSceneFloatChannel.h"
-#include "Evaluation/MovieSceneEvalTemplate.h"
-#include "LiveLinkRole.h"
-#endif
diff --git a/Engine/Plugins/Animation/LiveLink/Source/LiveLinkMovieScene/Public/MovieScene/MovieSceneLiveLinkStructProperties.h b/Engine/Plugins/Animation/LiveLink/Source/LiveLinkMovieScene/Public/MovieScene/MovieSceneLiveLinkStructProperties.h
--- a/Engine/Plugins/Animation/LiveLink/Source/LiveLinkMovieScene/Public/MovieScene/MovieSceneLiveLinkStructProperties.h
+++ b/Engine/Plugins/Animation/LiveLink/Source/LiveLinkMovieScene/Public/MovieScene/MovieSceneLiveLinkStructProperties.h
@@ -60,14 +60,3 @@ public:
 	UPROPERTY()
 	TArray<FLiveLinkPropertyData> Properties;
 };
-
-
-#if UE_ENABLE_INCLUDE_ORDER_DEPRECATED_IN_5_2
-#include "Channels/MovieSceneBoolChannel.h"
-#include "Channels/MovieSceneByteChannel.h"
-#include "Channels/MovieSceneFloatChannel.h"
-#include "Channels/MovieSceneIntegerChannel.h"
-#include "Channels/MovieSceneStringChannel.h"
-#include "CoreMinimal.h"
-#include "LiveLinkTypes.h"
-#endif
diff --git a/Engine/Plugins/Animation/LiveLink/Source/LiveLinkMovieScene/Public/MovieScene/MovieSceneLiveLinkStructPropertyBindings.h b/Engine/Plugins/Animation/LiveLink/Source/LiveLinkMovieScene/Public/MovieScene/MovieSceneLiveLinkStructPropertyBindings.h
--- a/Engine/Plugins/Animation/LiveLink/Source/LiveLinkMovieScene/Public/MovieScene/MovieSceneLiveLinkStructPropertyBindings.h
+++ b/Engine/Plugins/Animation/LiveLink/Source/LiveLinkMovieScene/Public/MovieScene/MovieSceneLiveLinkStructPropertyBindings.h
@@ -287,13 +287,3 @@ LIVELINKMOVIESCENE_API void FLiveLinkStructPropertyBindings::SetCurrentValue<boo
 
 template<> LIVELINKMOVIESCENE_API bool FLiveLinkStructPropertyBindings::GetCurrentValueAt<bool>(int32 InIndex, const UScriptStruct& InStruct, const void* InSourceAddress);
 template<> LIVELINKMOVIESCENE_API void FLiveLinkStructPropertyBindings::SetCurrentValueAt<bool>(int32 InIndex, const UScriptStruct& InStruct, void* InSourceAddress, TCallTraits<bool>::ParamType InValue);
-
-#if UE_ENABLE_INCLUDE_ORDER_DEPRECATED_IN_5_2
-#include "Channels/MovieSceneBoolChannel.h"
-#include "Channels/MovieSceneByteChannel.h"
-#include "Channels/MovieSceneFloatChannel.h"
-#include "Channels/MovieSceneIntegerChannel.h"
-#include "Channels/MovieSceneStringChannel.h"
-#include "CoreMinimal.h"
-#include "LiveLinkTypes.h"
-#endif
diff --git a/Engine/Plugins/Animation/LiveLink/Source/LiveLinkMovieScene/Public/MovieScene/MovieSceneLiveLinkSubSection.h b/Engine/Plugins/Animation/LiveLink/Source/LiveLinkMovieScene/Public/MovieScene/MovieSceneLiveLinkSubSection.h
--- a/Engine/Plugins/Animation/LiveLink/Source/LiveLinkMovieScene/Public/MovieScene/MovieSceneLiveLinkSubSection.h
+++ b/Engine/Plugins/Animation/LiveLink/Source/LiveLinkMovieScene/Public/MovieScene/MovieSceneLiveLinkSubSection.h
@@ -102,11 +102,3 @@ public:
 protected:
 	TSharedPtr<FLiveLinkStaticDataStruct> StaticData;
 };
-
-
-#if UE_ENABLE_INCLUDE_ORDER_DEPRECATED_IN_5_2
-#include "Channels/MovieSceneFloatChannel.h"
-#include "Evaluation/MovieSceneEvalTemplate.h"
-#include "LiveLinkRole.h"
-#include "MovieSceneSection.h"
-#endif
diff --git a/Engine/Plugins/Animation/LiveLink/Source/LiveLinkMovieScene/Public/MovieScene/MovieSceneLiveLinkSubSectionAnimation.h b/Engine/Plugins/Animation/LiveLink/Source/LiveLinkMovieScene/Public/MovieScene/MovieSceneLiveLinkSubSectionAnimation.h
--- a/Engine/Plugins/Animation/LiveLink/Source/LiveLinkMovieScene/Public/MovieScene/MovieSceneLiveLinkSubSectionAnimation.h
+++ b/Engine/Plugins/Animation/LiveLink/Source/LiveLinkMovieScene/Public/MovieScene/MovieSceneLiveLinkSubSectionAnimation.h
@@ -39,9 +39,3 @@ protected:
 	/** Helper struct to manage filling channels from transforms */
 	TSharedPtr<FMovieSceneLiveLinkTransformHandler> TransformHandler;
 };
-
-#if UE_ENABLE_INCLUDE_ORDER_DEPRECATED_IN_5_2
-#include "Channels/MovieSceneFloatChannel.h"
-#include "LiveLinkRole.h"
-#include "MovieScene/MovieSceneLiveLinkTransformHandler.h"
-#endif
diff --git a/Engine/Plugins/Animation/LiveLink/Source/LiveLinkMovieScene/Public/MovieScene/MovieSceneLiveLinkSubSectionBasicRole.h b/Engine/Plugins/Animation/LiveLink/Source/LiveLinkMovieScene/Public/MovieScene/MovieSceneLiveLinkSubSectionBasicRole.h
--- a/Engine/Plugins/Animation/LiveLink/Source/LiveLinkMovieScene/Public/MovieScene/MovieSceneLiveLinkSubSectionBasicRole.h
+++ b/Engine/Plugins/Animation/LiveLink/Source/LiveLinkMovieScene/Public/MovieScene/MovieSceneLiveLinkSubSectionBasicRole.h
@@ -37,9 +37,3 @@ protected:
 	/** Helper struct to manage filling channels from property array */
 	TSharedPtr<IMovieSceneLiveLinkPropertyHandler> PropertyHandler;
 };
-
-#if UE_ENABLE_INCLUDE_ORDER_DEPRECATED_IN_5_2
-#include "Channels/MovieSceneFloatChannel.h"
-#include "LiveLinkRole.h"
-#include "MovieScene/MovieSceneLiveLinkPropertyHandler.h"
-#endif
diff --git a/Engine/Plugins/Animation/LiveLink/Source/LiveLinkMovieScene/Public/MovieScene/MovieSceneLiveLinkSubSectionProperties.h b/Engine/Plugins/Animation/LiveLink/Source/LiveLinkMovieScene/Public/MovieScene/MovieSceneLiveLinkSubSectionProperties.h
--- a/Engine/Plugins/Animation/LiveLink/Source/LiveLinkMovieScene/Public/MovieScene/MovieSceneLiveLinkSubSectionProperties.h
+++ b/Engine/Plugins/Animation/LiveLink/Source/LiveLinkMovieScene/Public/MovieScene/MovieSceneLiveLinkSubSectionProperties.h
@@ -42,9 +42,3 @@ protected:
 	/** Helper struct to manage filling channels for each properties */
 	TArray<TSharedPtr<IMovieSceneLiveLinkPropertyHandler>> PropertyHandlers;
 };
-
-
-#if UE_ENABLE_INCLUDE_ORDER_DEPRECATED_IN_5_2
-#include "LiveLinkRole.h"
-#include "MovieScene/MovieSceneLiveLinkPropertyHandler.h"
-#endif
diff --git a/Engine/Plugins/Animation/LiveLink/Source/LiveLinkMovieScene/Public/MovieScene/MovieSceneLiveLinkTransformHandler.h b/Engine/Plugins/Animation/LiveLink/Source/LiveLinkMovieScene/Public/MovieScene/MovieSceneLiveLinkTransformHandler.h
--- a/Engine/Plugins/Animation/LiveLink/Source/LiveLinkMovieScene/Public/MovieScene/MovieSceneLiveLinkTransformHandler.h
+++ b/Engine/Plugins/Animation/LiveLink/Source/LiveLinkMovieScene/Public/MovieScene/MovieSceneLiveLinkTransformHandler.h
@@ -48,12 +48,3 @@ protected:
 	/** Buffer of transform keys. Keys are inserted into tracks in FinalizeTrack() */
 	TArray<FLiveLinkTransformKeys> BufferedTransformsPerElement;
 };
-
-
-
-
-
-#if UE_ENABLE_INCLUDE_ORDER_DEPRECATED_IN_5_2
-#include "CoreMinimal.h"
-#include "MovieScene/MovieSceneLiveLinkStructProperties.h"
-#endif
diff --git a/Engine/Plugins/Animation/LiveLink/Source/LiveLinkSequencer/Public/LiveLinkSequencerSettings.h b/Engine/Plugins/Animation/LiveLink/Source/LiveLinkSequencer/Public/LiveLinkSequencerSettings.h
--- a/Engine/Plugins/Animation/LiveLink/Source/LiveLinkSequencer/Public/LiveLinkSequencerSettings.h
+++ b/Engine/Plugins/Animation/LiveLink/Source/LiveLinkSequencer/Public/LiveLinkSequencerSettings.h
@@ -35,8 +35,3 @@ public:
 	UPROPERTY(config, EditAnywhere, Category = "LiveLink", meta = (AllowAbstract = "false"))
 	TMap<TSubclassOf<ULiveLinkControllerBase>, TSubclassOf<UMovieSceneLiveLinkControllerTrackRecorder>> DefaultTrackRecordersForController;
 };
-
-#if UE_ENABLE_INCLUDE_ORDER_DEPRECATED_IN_5_2
-#include "LiveLinkControllerBase.h"
-#include "MovieSceneLiveLinkControllerTrackRecorder.h"
-#endif
diff --git a/Engine/Plugins/Animation/LiveLink/Source/LiveLinkSequencer/Public/MovieSceneLiveLinkControllerMapTrackRecorder.h b/Engine/Plugins/Animation/LiveLink/Source/LiveLinkSequencer/Public/MovieSceneLiveLinkControllerMapTrackRecorder.h
--- a/Engine/Plugins/Animation/LiveLink/Source/LiveLinkSequencer/Public/MovieSceneLiveLinkControllerMapTrackRecorder.h
+++ b/Engine/Plugins/Animation/LiveLink/Source/LiveLinkSequencer/Public/MovieSceneLiveLinkControllerMapTrackRecorder.h
@@ -53,7 +53,3 @@ private:
 	UPROPERTY(Transient)
 	TArray<TObjectPtr<UMovieSceneLiveLinkControllerTrackRecorder>> ControllerRecorders;
 };
-
-#if UE_ENABLE_INCLUDE_ORDER_DEPRECATED_IN_5_2
-#include "MovieSceneLiveLinkControllerTrackRecorder.h"
-#endif
diff --git a/Engine/Plugins/Animation/LiveLink/Source/LiveLinkSequencer/Public/MovieSceneLiveLinkTrackRecorder.h b/Engine/Plugins/Animation/LiveLink/Source/LiveLinkSequencer/Public/MovieSceneLiveLinkTrackRecorder.h
--- a/Engine/Plugins/Animation/LiveLink/Source/LiveLinkSequencer/Public/MovieSceneLiveLinkTrackRecorder.h
+++ b/Engine/Plugins/Animation/LiveLink/Source/LiveLinkSequencer/Public/MovieSceneLiveLinkTrackRecorder.h
@@ -111,10 +111,3 @@ private:
 	/** Pairs of recorded times with the corresponding timecode.*/
 	TArray<TPair<FQualifiedFrameTime, FQualifiedFrameTime>> RecordedTimes;
 };
-
-#if UE_ENABLE_INCLUDE_ORDER_DEPRECATED_IN_5_2
-#include "CoreMinimal.h"
-#include "GameFramework/Actor.h"
-#include "ILiveLinkClient.h"
-#include "TrackRecorders/IMovieSceneTrackRecorderFactory.h"
-#endif
diff --git a/Engine/Plugins/Animation/MotionWarping/Source/MotionWarping/Public/MotionWarping.h b/Engine/Plugins/Animation/MotionWarping/Source/MotionWarping/Public/MotionWarping.h
--- a/Engine/Plugins/Animation/MotionWarping/Source/MotionWarping/Public/MotionWarping.h
+++ b/Engine/Plugins/Animation/MotionWarping/Source/MotionWarping/Public/MotionWarping.h
@@ -12,8 +12,3 @@ public:
 	virtual void StartupModule() override;
 	virtual void ShutdownModule() override;
 };
-
-#if UE_ENABLE_INCLUDE_ORDER_DEPRECATED_IN_5_2
-#include "CoreMinimal.h"
-#include "Modules/ModuleManager.h"
-#endif
diff --git a/Engine/Plugins/Animation/MotionWarping/Source/MotionWarping/Public/MotionWarpingComponent.h b/Engine/Plugins/Animation/MotionWarping/Source/MotionWarping/Public/MotionWarpingComponent.h
--- a/Engine/Plugins/Animation/MotionWarping/Source/MotionWarping/Public/MotionWarpingComponent.h
+++ b/Engine/Plugins/Animation/MotionWarping/Source/MotionWarping/Public/MotionWarpingComponent.h
@@ -245,9 +245,3 @@ protected:
 
 	FTransform ProcessRootMotionPreConvertToWorld(const FTransform& InRootMotion, float DeltaSeconds, const FMotionWarpingUpdateContext* InContext=nullptr); // callback with optional context
 };
-
-#if UE_ENABLE_INCLUDE_ORDER_DEPRECATED_IN_5_2
-#include "CoreMinimal.h"
-#include "HAL/IConsoleManager.h"
-#include "Kismet/KismetSystemLibrary.h"
-#endif
diff --git a/Engine/Plugins/Animation/MotionWarping/Source/MotionWarping/Public/RootMotionModifier.h b/Engine/Plugins/Animation/MotionWarping/Source/MotionWarping/Public/RootMotionModifier.h
--- a/Engine/Plugins/Animation/MotionWarping/Source/MotionWarping/Public/RootMotionModifier.h
+++ b/Engine/Plugins/Animation/MotionWarping/Source/MotionWarping/Public/RootMotionModifier.h
@@ -454,8 +454,3 @@ public:
 		UPARAM(DisplayName = "End Time") float InEndTime,
 		UPARAM(DisplayName = "Scale") FVector InScale = FVector(1.f));
 };
-
-#if UE_ENABLE_INCLUDE_ORDER_DEPRECATED_IN_5_2
-#include "Animation/AnimSequence.h"
-#include "CoreMinimal.h"
-#endif
diff --git a/Engine/Plugins/Animation/MotionWarping/Source/MotionWarping/Public/RootMotionModifier_AdjustmentBlendWarp.h b/Engine/Plugins/Animation/MotionWarping/Source/MotionWarping/Public/RootMotionModifier_AdjustmentBlendWarp.h
--- a/Engine/Plugins/Animation/MotionWarping/Source/MotionWarping/Public/RootMotionModifier_AdjustmentBlendWarp.h
+++ b/Engine/Plugins/Animation/MotionWarping/Source/MotionWarping/Public/RootMotionModifier_AdjustmentBlendWarp.h
@@ -103,8 +103,3 @@ protected:
 
 	static void AdjustmentBlendWarp(const FBoneContainer& BoneContainer, const FCSPose<FCompactPose>& AdditivePose, const FMotionDeltaTrackContainer& MotionDeltaTracks, FAnimSequenceTrackContainer& Output);
 };
-
-#if UE_ENABLE_INCLUDE_ORDER_DEPRECATED_IN_5_2
-#include "BonePose.h"
-#include "CoreMinimal.h"
-#endif
diff --git a/Engine/Plugins/AudioGameplay/Source/AudioGameplay/Public/AudioComponentGroup.h b/Engine/Plugins/AudioGameplay/Source/AudioGameplay/Public/AudioComponentGroup.h
--- a/Engine/Plugins/AudioGameplay/Source/AudioGameplay/Public/AudioComponentGroup.h
+++ b/Engine/Plugins/AudioGameplay/Source/AudioGameplay/Public/AudioComponentGroup.h
@@ -197,7 +197,3 @@ protected:
 
 	friend class FAudioComponentGroupDebug;
 };
-
-#if UE_ENABLE_INCLUDE_ORDER_DEPRECATED_IN_5_2
-#include "Components/AudioComponent.h"
-#endif
diff --git a/Engine/Plugins/AudioGameplay/Source/AudioGameplay/Public/AudioComponentGroupDebug.h b/Engine/Plugins/AudioGameplay/Source/AudioGameplay/Public/AudioComponentGroupDebug.h
--- a/Engine/Plugins/AudioGameplay/Source/AudioGameplay/Public/AudioComponentGroupDebug.h
+++ b/Engine/Plugins/AudioGameplay/Source/AudioGameplay/Public/AudioComponentGroupDebug.h
@@ -31,7 +31,3 @@ public:
 	// print the final debug string above the target actor
 	static void PrintDebugString(const FString& InString, AActor* Owner, const FColor DrawColor);
 };
-
-#if UE_ENABLE_INCLUDE_ORDER_DEPRECATED_IN_5_2
-#include "AudioParameter.h"
-#endif
diff --git a/Engine/Plugins/AudioGameplay/Source/AudioGameplay/Public/AudioGameplayModule.h b/Engine/Plugins/AudioGameplay/Source/AudioGameplay/Public/AudioGameplayModule.h
--- a/Engine/Plugins/AudioGameplay/Source/AudioGameplay/Public/AudioGameplayModule.h
+++ b/Engine/Plugins/AudioGameplay/Source/AudioGameplay/Public/AudioGameplayModule.h
@@ -16,7 +16,3 @@ public:
 	virtual void ShutdownModule() override;
 	//~ End IModuleInterface interface
 };
-
-#if UE_ENABLE_INCLUDE_ORDER_DEPRECATED_IN_5_2
-#include "AudioGameplayLogs.h"
-#endif
diff --git a/Engine/Plugins/AudioGameplayVolume/Source/AudioGameplayVolume/Public/AudioGameplayVolumeModule.h b/Engine/Plugins/AudioGameplayVolume/Source/AudioGameplayVolume/Public/AudioGameplayVolumeModule.h
--- a/Engine/Plugins/AudioGameplayVolume/Source/AudioGameplayVolume/Public/AudioGameplayVolumeModule.h
+++ b/Engine/Plugins/AudioGameplayVolume/Source/AudioGameplayVolume/Public/AudioGameplayVolumeModule.h
@@ -16,7 +16,3 @@ public:
 	virtual void ShutdownModule() override;
 	//~ End IModuleInterface interface
 };
-
-#if UE_ENABLE_INCLUDE_ORDER_DEPRECATED_IN_5_2
-#include "AudioGameplayVolumeLogs.h"
-#endif
diff --git a/Engine/Plugins/AudioGameplayVolume/Source/AudioGameplayVolume/Public/AudioGameplayVolumeMutator.h b/Engine/Plugins/AudioGameplayVolume/Source/AudioGameplayVolume/Public/AudioGameplayVolumeMutator.h
--- a/Engine/Plugins/AudioGameplayVolume/Source/AudioGameplayVolume/Public/AudioGameplayVolumeMutator.h
+++ b/Engine/Plugins/AudioGameplayVolume/Source/AudioGameplayVolume/Public/AudioGameplayVolumeMutator.h
@@ -141,8 +141,3 @@ private:
 	/** Called for you during mutator creation. See CopyAudioDataToMutator for adding data to derived classes */
 	void CopyAudioDataToMutatorBase(TSharedPtr<FProxyVolumeMutator>& Mutator) const;
 };
-
-#if UE_ENABLE_INCLUDE_ORDER_DEPRECATED_IN_5_2
-#include "AudioGameplayVolumeComponent.h"
-#include "Sound/AudioVolume.h"
-#endif
diff --git a/Engine/Plugins/Bridge/Source/MegascansPlugin/Internal/BridgeDragDropHelper.h b/Engine/Plugins/Bridge/Source/MegascansPlugin/Internal/BridgeDragDropHelper.h
--- a/Engine/Plugins/Bridge/Source/MegascansPlugin/Internal/BridgeDragDropHelper.h
+++ b/Engine/Plugins/Bridge/Source/MegascansPlugin/Internal/BridgeDragDropHelper.h
@@ -25,8 +25,3 @@ public:
 	static void Initialize();
 	static TSharedPtr<FBridgeDragDropHelperImpl> Instance;
 };
-
-#if UE_ENABLE_INCLUDE_ORDER_DEPRECATED_IN_5_2
-#include "AssetRegistry/AssetData.h"
-#include "CoreMinimal.h"
-#endif
diff --git a/Engine/Plugins/Bridge/Source/MegascansPlugin/Public/AssetsImportController.h b/Engine/Plugins/Bridge/Source/MegascansPlugin/Public/AssetsImportController.h
--- a/Engine/Plugins/Bridge/Source/MegascansPlugin/Public/AssetsImportController.h
+++ b/Engine/Plugins/Bridge/Source/MegascansPlugin/Public/AssetsImportController.h
@@ -19,10 +19,3 @@ public:
 	static TSharedPtr<FAssetsImportController> Get();
 	void DataReceived(const FString DataFromBridge);
 };
-
-#if UE_ENABLE_INCLUDE_ORDER_DEPRECATED_IN_5_2
-#include "CoreMinimal.h"
-#include "Engine/StaticMeshActor.h"
-#include "MSAssetImportData.h"
-#include "Materials/MaterialInstanceConstant.h"
-#endif
diff --git a/Engine/Plugins/Bridge/Source/MegascansPlugin/Public/BridgeDragDropUtils.h b/Engine/Plugins/Bridge/Source/MegascansPlugin/Public/BridgeDragDropUtils.h
--- a/Engine/Plugins/Bridge/Source/MegascansPlugin/Public/BridgeDragDropUtils.h
+++ b/Engine/Plugins/Bridge/Source/MegascansPlugin/Public/BridgeDragDropUtils.h
@@ -23,8 +23,3 @@ public:
     static void Initialize();
     static TSharedPtr<FBridgeDragDropImpl> Instance;
 };
-
-#if UE_ENABLE_INCLUDE_ORDER_DEPRECATED_IN_5_2
-#include "AssetRegistry/AssetData.h"
-#include "CoreMinimal.h"
-#endif
diff --git a/Engine/Plugins/Bridge/Source/MegascansPlugin/Public/MSSettings.h b/Engine/Plugins/Bridge/Source/MegascansPlugin/Public/MSSettings.h
--- a/Engine/Plugins/Bridge/Source/MegascansPlugin/Public/MSSettings.h
+++ b/Engine/Plugins/Bridge/Source/MegascansPlugin/Public/MSSettings.h
@@ -98,10 +98,3 @@ public:
 #endif
 
 };
-
-
-#if UE_ENABLE_INCLUDE_ORDER_DEPRECATED_IN_5_2
-#include "CoreMinimal.h"
-#include "Engine/EngineTypes.h"
-#include "Materials/Material.h"
-#endif
diff --git a/Engine/Plugins/Bridge/Source/MegascansPlugin/Public/SMSWindow.h b/Engine/Plugins/Bridge/Source/MegascansPlugin/Public/SMSWindow.h
--- a/Engine/Plugins/Bridge/Source/MegascansPlugin/Public/SMSWindow.h
+++ b/Engine/Plugins/Bridge/Source/MegascansPlugin/Public/SMSWindow.h
@@ -20,13 +20,3 @@ public:
 	
 	
 };
-
-
-
-
-
-#if UE_ENABLE_INCLUDE_ORDER_DEPRECATED_IN_5_2
-#include "CoreMinimal.h"
-#include "MSSettings.h"
-#include "Widgets/SWindow.h"
-#endif
diff --git a/Engine/Plugins/Bridge/Source/MegascansPlugin/Public/TCPServer.h b/Engine/Plugins/Bridge/Source/MegascansPlugin/Public/TCPServer.h
--- a/Engine/Plugins/Bridge/Source/MegascansPlugin/Public/TCPServer.h
+++ b/Engine/Plugins/Bridge/Source/MegascansPlugin/Public/TCPServer.h
@@ -46,14 +46,3 @@ private:
 	class FTcpListener *Listener = NULL;	
 	TArray<int32> ConnectionTimer;
 };
-
-
-
-#if UE_ENABLE_INCLUDE_ORDER_DEPRECATED_IN_5_2
-#include "Async/AsyncWork.h"
-#include "Common/TcpListener.h"
-#include "Common/UdpSocketReceiver.h"
-#include "CoreMinimal.h"
-#include "Misc/DateTime.h"
-#include "Sockets.h"
-#endif
diff --git a/Engine/Plugins/Cameras/EngineCameras/Source/EngineCameras/Public/Animations/CameraAnimationCameraModifier.h b/Engine/Plugins/Cameras/EngineCameras/Source/EngineCameras/Public/Animations/CameraAnimationCameraModifier.h
--- a/Engine/Plugins/Cameras/EngineCameras/Source/EngineCameras/Public/Animations/CameraAnimationCameraModifier.h
+++ b/Engine/Plugins/Cameras/EngineCameras/Source/EngineCameras/Public/Animations/CameraAnimationCameraModifier.h
@@ -291,10 +291,3 @@ public:
 class ENGINECAMERAS_API UGameplayCamerasFunctionLibrary : public UEngineCameraAnimationFunctionLibrary
 {
 };
-
-
-#if UE_ENABLE_INCLUDE_ORDER_DEPRECATED_IN_5_2
-#include "Camera/CameraTypes.h"
-#include "CoreMinimal.h"
-#include "Templates/SubclassOf.h"
-#endif
diff --git a/Engine/Plugins/Cameras/EngineCameras/Source/EngineCameras/Public/EngineCamerasSubsystem.h b/Engine/Plugins/Cameras/EngineCameras/Source/EngineCameras/Public/EngineCamerasSubsystem.h
--- a/Engine/Plugins/Cameras/EngineCameras/Source/EngineCameras/Public/EngineCamerasSubsystem.h
+++ b/Engine/Plugins/Cameras/EngineCameras/Source/EngineCameras/Public/EngineCamerasSubsystem.h
@@ -68,8 +68,3 @@ public:
 	UFUNCTION(BlueprintCallable, Category="Camera Animation")
 	void StopAllCameraAnimations(APlayerController* PlayerController, bool bImmediate = false);
 };
-
-#if UE_ENABLE_INCLUDE_ORDER_DEPRECATED_IN_5_2
-#include "CameraAnimationCameraModifier.h"
-#include "CoreMinimal.h"
-#endif
diff --git a/Engine/Plugins/Cameras/EngineCameras/Source/EngineCameras/Public/Shakes/LegacyCameraShake.h b/Engine/Plugins/Cameras/EngineCameras/Source/EngineCameras/Public/Shakes/LegacyCameraShake.h
--- a/Engine/Plugins/Cameras/EngineCameras/Source/EngineCameras/Public/Shakes/LegacyCameraShake.h
+++ b/Engine/Plugins/Cameras/EngineCameras/Source/EngineCameras/Public/Shakes/LegacyCameraShake.h
@@ -322,8 +322,3 @@ public:
 		return CastChecked<ULegacyCameraShake>(CameraShake, ECastCheckedType::NullAllowed);
 	}
 };
-
-#if UE_ENABLE_INCLUDE_ORDER_DEPRECATED_IN_5_2
-#include "CoreMinimal.h"
-#include "GameFramework/Actor.h"
-#endif
diff --git a/Engine/Plugins/Compositing/OpenColorIO/Source/OpenColorIO/Public/OpenColorIOConfiguration.h b/Engine/Plugins/Compositing/OpenColorIO/Source/OpenColorIO/Public/OpenColorIOConfiguration.h
--- a/Engine/Plugins/Compositing/OpenColorIO/Source/OpenColorIO/Public/OpenColorIOConfiguration.h
+++ b/Engine/Plugins/Compositing/OpenColorIO/Source/OpenColorIO/Public/OpenColorIOConfiguration.h
@@ -7,10 +7,6 @@
 #include "CoreMinimal.h"
 #include "Engine/EngineTypes.h"
 #include "OpenColorIOColorSpace.h"
-#if UE_ENABLE_INCLUDE_ORDER_DEPRECATED_IN_5_2
-#include "RHIDefinitions.h"
-#endif
-
 #include "OpenColorIOConfiguration.generated.h"
 
 
diff --git a/Engine/Plugins/Compositing/OpenColorIO/Source/OpenColorIO/Public/OpenColorIOShared.h b/Engine/Plugins/Compositing/OpenColorIO/Source/OpenColorIO/Public/OpenColorIOShared.h
--- a/Engine/Plugins/Compositing/OpenColorIO/Source/OpenColorIO/Public/OpenColorIOShared.h
+++ b/Engine/Plugins/Compositing/OpenColorIO/Source/OpenColorIO/Public/OpenColorIOShared.h
@@ -10,9 +10,6 @@
 #include "Misc/Guid.h"
 #include "Misc/SecureHash.h"
 #include "RenderResource.h"
-#if UE_ENABLE_INCLUDE_ORDER_DEPRECATED_IN_5_2
-#include "RenderingThread.h"
-#endif
 #include "RenderDeferredCleanup.h"
 #include "RHIFeatureLevel.h"
 #include "SceneTypes.h"
diff --git a/Engine/Plugins/Compression/OodleNetwork/Source/Classes/OodleNetworkTrainerCommandlet.h b/Engine/Plugins/Compression/OodleNetwork/Source/Classes/OodleNetworkTrainerCommandlet.h
--- a/Engine/Plugins/Compression/OodleNetwork/Source/Classes/OodleNetworkTrainerCommandlet.h
+++ b/Engine/Plugins/Compression/OodleNetwork/Source/Classes/OodleNetworkTrainerCommandlet.h
@@ -326,8 +326,3 @@ private:
 	void Cleanup();
 };
 #endif
-
-#if UE_ENABLE_INCLUDE_ORDER_DEPRECATED_IN_5_2
-#include "CoreMinimal.h"
-#include "OodleNetworkHandlerComponent.h"
-#endif
diff --git a/Engine/Plugins/Compression/OodleNetwork/Source/Public/OodleNetworkAnalytics.h b/Engine/Plugins/Compression/OodleNetwork/Source/Public/OodleNetworkAnalytics.h
--- a/Engine/Plugins/Compression/OodleNetwork/Source/Public/OodleNetworkAnalytics.h
+++ b/Engine/Plugins/Compression/OodleNetwork/Source/Public/OodleNetworkAnalytics.h
@@ -130,11 +130,3 @@ struct FClientOodleNetAnalyticsData : public FOodleNetAnalyticsData
 {
 	virtual const TCHAR* GetAnalyticsEventName() const override;
 };
-
-
-
-
-#if UE_ENABLE_INCLUDE_ORDER_DEPRECATED_IN_5_2
-#include "CoreMinimal.h"
-#include "Runtime/Launch/Resources/Version.h"
-#endif
diff --git a/Engine/Plugins/Compression/OodleNetwork/Source/Public/OodleNetworkArchives.h b/Engine/Plugins/Compression/OodleNetwork/Source/Public/OodleNetworkArchives.h
--- a/Engine/Plugins/Compression/OodleNetwork/Source/Public/OodleNetworkArchives.h
+++ b/Engine/Plugins/Compression/OodleNetwork/Source/Public/OodleNetworkArchives.h
@@ -450,10 +450,3 @@ public:
 	void SerializeDictionaryAndState(uint8*& DictionaryData, uint32& DictionaryBytes,
 										uint8*& CompactCompresorState, uint32& CompactCompressorStateBytes);
 };
-
-
-
-
-#if UE_ENABLE_INCLUDE_ORDER_DEPRECATED_IN_5_2
-#include "CoreMinimal.h"
-#endif
diff --git a/Engine/Plugins/Compression/OodleNetwork/Source/Public/OodleNetworkHandlerComponent.h b/Engine/Plugins/Compression/OodleNetwork/Source/Public/OodleNetworkHandlerComponent.h
--- a/Engine/Plugins/Compression/OodleNetwork/Source/Public/OodleNetworkHandlerComponent.h
+++ b/Engine/Plugins/Compression/OodleNetwork/Source/Public/OodleNetworkHandlerComponent.h
@@ -434,9 +434,3 @@ public:
 	virtual void StartupModule() override;
 	virtual void ShutdownModule() override;
 };
-
-#if UE_ENABLE_INCLUDE_ORDER_DEPRECATED_IN_5_2
-#include "CoreMinimal.h"
-#include "OodleNetworkAnalytics.h"
-#include "UObject/CoreNet.h"
-#endif
diff --git a/Engine/Plugins/Developer/AnimationSharing/Source/AnimationSharing/Public/AdditiveAnimationInstance.h b/Engine/Plugins/Developer/AnimationSharing/Source/AnimationSharing/Public/AdditiveAnimationInstance.h
--- a/Engine/Plugins/Developer/AnimationSharing/Source/AnimationSharing/Public/AdditiveAnimationInstance.h
+++ b/Engine/Plugins/Developer/AnimationSharing/Source/AnimationSharing/Public/AdditiveAnimationInstance.h
@@ -30,8 +30,3 @@ protected:
 	USkeletalMeshComponent* BaseComponent;
 	bool bLoopingState;
 };
-
-#if UE_ENABLE_INCLUDE_ORDER_DEPRECATED_IN_5_2
-#include "Animation/SkeletalMeshActor.h"
-#include "Components/SkeletalMeshComponent.h"
-#endif
diff --git a/Engine/Plugins/Developer/AnimationSharing/Source/AnimationSharing/Public/AnimationSharingInstances.h b/Engine/Plugins/Developer/AnimationSharing/Source/AnimationSharing/Public/AnimationSharingInstances.h
--- a/Engine/Plugins/Developer/AnimationSharing/Source/AnimationSharing/Public/AnimationSharingInstances.h
+++ b/Engine/Plugins/Developer/AnimationSharing/Source/AnimationSharing/Public/AnimationSharingInstances.h
@@ -78,9 +78,3 @@ protected:
 	UPROPERTY(EditDefaultsOnly, BlueprintReadOnly, Transient, Category = Additive)
 	bool bStateBool;
 };
-
-
-
-#if UE_ENABLE_INCLUDE_ORDER_DEPRECATED_IN_5_2
-#include "CoreMinimal.h"
-#endif
diff --git a/Engine/Plugins/Developer/AnimationSharing/Source/AnimationSharing/Public/AnimationSharingManager.h b/Engine/Plugins/Developer/AnimationSharing/Source/AnimationSharing/Public/AnimationSharingManager.h
--- a/Engine/Plugins/Developer/AnimationSharing/Source/AnimationSharing/Public/AnimationSharingManager.h
+++ b/Engine/Plugins/Developer/AnimationSharing/Source/AnimationSharing/Public/AnimationSharingManager.h
@@ -576,11 +576,3 @@ public:
 	static FName GetPlatformName();
 #endif
 };
-
-#if UE_ENABLE_INCLUDE_ORDER_DEPRECATED_IN_5_2
-#include "Animation/AnimBlueprint.h"
-#include "AnimationSharingModule.h"
-#include "Engine/DeveloperSettings.h"
-#include "Engine/SkeletalMesh.h"
-#include "TransitionBlendInstance.h"
-#endif
diff --git a/Engine/Plugins/Developer/AnimationSharing/Source/AnimationSharing/Public/AnimationSharingTypes.h b/Engine/Plugins/Developer/AnimationSharing/Source/AnimationSharing/Public/AnimationSharingTypes.h
--- a/Engine/Plugins/Developer/AnimationSharing/Source/AnimationSharing/Public/AnimationSharingTypes.h
+++ b/Engine/Plugins/Developer/AnimationSharing/Source/AnimationSharing/Public/AnimationSharingTypes.h
@@ -179,8 +179,3 @@ public:
 	UPROPERTY(EditAnywhere, Category = AnimationSharing, meta = (ClampMin = "0.0", UIMin = "0.0"))
 	FPerPlatformFloat TickSignificanceValue;
 };
-
-
-#if UE_ENABLE_INCLUDE_ORDER_DEPRECATED_IN_5_2
-#include "AnimationSharingInstances.h"
-#endif
diff --git a/Engine/Plugins/Developer/AnimationSharing/Source/AnimationSharing/Public/TransitionBlendInstance.h b/Engine/Plugins/Developer/AnimationSharing/Source/AnimationSharing/Public/TransitionBlendInstance.h
--- a/Engine/Plugins/Developer/AnimationSharing/Source/AnimationSharing/Public/TransitionBlendInstance.h
+++ b/Engine/Plugins/Developer/AnimationSharing/Source/AnimationSharing/Public/TransitionBlendInstance.h
@@ -28,8 +28,3 @@ protected:
 	float BlendTime;
 	bool bBlendState;
 };
-
-#if UE_ENABLE_INCLUDE_ORDER_DEPRECATED_IN_5_2
-#include "Animation/SkeletalMeshActor.h"
-#include "Components/SkeletalMeshComponent.h"
-#endif
diff --git a/Engine/Plugins/Developer/AnimationSharing/Source/AnimationSharingEd/Public/SetupDetailsViewCustomizations.h b/Engine/Plugins/Developer/AnimationSharing/Source/AnimationSharingEd/Public/SetupDetailsViewCustomizations.h
--- a/Engine/Plugins/Developer/AnimationSharing/Source/AnimationSharingEd/Public/SetupDetailsViewCustomizations.h
+++ b/Engine/Plugins/Developer/AnimationSharing/Source/AnimationSharingEd/Public/SetupDetailsViewCustomizations.h
@@ -75,9 +75,3 @@ protected:
 protected:
 	TSharedPtr<IPropertyHandle> AnimSequencePropertyHandle;
 };
-
-#if UE_ENABLE_INCLUDE_ORDER_DEPRECATED_IN_5_2
-#include "CoreMinimal.h"
-#include "IDetailCustomization.h"
-#include "PropertyCustomizationHelpers.h"
-#endif
diff --git a/Engine/Plugins/Developer/Concert/ConcertSync/ConcertSyncCore/Source/ConcertSyncCore/Public/ConcertDataStoreMessages.h b/Engine/Plugins/Developer/Concert/ConcertSync/ConcertSyncCore/Source/ConcertSyncCore/Public/ConcertDataStoreMessages.h
--- a/Engine/Plugins/Developer/Concert/ConcertSync/ConcertSyncCore/Source/ConcertSyncCore/Public/ConcertDataStoreMessages.h
+++ b/Engine/Plugins/Developer/Concert/ConcertSync/ConcertSyncCore/Source/ConcertSyncCore/Public/ConcertDataStoreMessages.h
@@ -510,8 +510,3 @@ private:
 	/** The weakly-typed operation result. */
 	FConcertDataStoreResult Result;
 };
-
-#if UE_ENABLE_INCLUDE_ORDER_DEPRECATED_IN_5_2
-#include "ConcertMessages.h"
-#include "CoreMinimal.h"
-#endif
diff --git a/Engine/Plugins/Developer/Concert/ConcertSync/ConcertSyncCore/Source/ConcertSyncCore/Public/ConcertSyncSessionDatabase.h b/Engine/Plugins/Developer/Concert/ConcertSync/ConcertSyncCore/Source/ConcertSyncCore/Public/ConcertSyncSessionDatabase.h
--- a/Engine/Plugins/Developer/Concert/ConcertSync/ConcertSyncCore/Source/ConcertSyncCore/Public/ConcertSyncSessionDatabase.h
+++ b/Engine/Plugins/Developer/Concert/ConcertSync/ConcertSyncCore/Source/ConcertSyncCore/Public/ConcertSyncSessionDatabase.h
@@ -1130,9 +1130,3 @@ namespace ConcertSyncSessionDatabaseFilterUtil
 	CONCERTSYNCCORE_API bool PackageEventPassesFilter(const int64 InPackageEventId, const FConcertSessionFilter& InSessionFilter, const FConcertSyncSessionDatabase& InDatabase);
 
 } // namespace ConcertSyncSessionDatabaseFilterUtil
-
-#if UE_ENABLE_INCLUDE_ORDER_DEPRECATED_IN_5_2
-#include "ConcertSyncSessionTypes.h"
-#include "CoreMinimal.h"
-#include "Templates/SharedPointerInternals.h"
-#endif
diff --git a/Engine/Plugins/Developer/Concert/ConcertSync/ConcertSyncCore/Source/ConcertSyncCore/Public/ConcertTransactionEvents.h b/Engine/Plugins/Developer/Concert/ConcertSync/ConcertSyncCore/Source/ConcertSyncCore/Public/ConcertTransactionEvents.h
--- a/Engine/Plugins/Developer/Concert/ConcertSync/ConcertSyncCore/Source/ConcertSyncCore/Public/ConcertTransactionEvents.h
+++ b/Engine/Plugins/Developer/Concert/ConcertSync/ConcertSyncCore/Source/ConcertSyncCore/Public/ConcertTransactionEvents.h
@@ -227,8 +227,3 @@ struct FConcertTransactionRejectedEvent
 	UPROPERTY()
 	FGuid TransactionId;
 };
-
-#if UE_ENABLE_INCLUDE_ORDER_DEPRECATED_IN_5_2
-#include "CoreMinimal.h"
-#include "UObject/Package.h"
-#endif
diff --git a/Engine/Plugins/Developer/Concert/ConcertSync/ConcertSyncCore/Source/ConcertSyncCore/Public/HistoryEdition/ActivityDependencyGraph.h b/Engine/Plugins/Developer/Concert/ConcertSync/ConcertSyncCore/Source/ConcertSyncCore/Public/HistoryEdition/ActivityDependencyGraph.h
--- a/Engine/Plugins/Developer/Concert/ConcertSync/ConcertSyncCore/Source/ConcertSyncCore/Public/HistoryEdition/ActivityDependencyGraph.h
+++ b/Engine/Plugins/Developer/Concert/ConcertSync/ConcertSyncCore/Source/ConcertSyncCore/Public/HistoryEdition/ActivityDependencyGraph.h
@@ -82,10 +82,3 @@ namespace UE::ConcertSyncCore
 		FActivityNode& GetNodeByIdInternal(FActivityNodeID ID);
 	};
 }
-
-
-#if UE_ENABLE_INCLUDE_ORDER_DEPRECATED_IN_5_2
-#include "CoreMinimal.h"
-#include "Misc/Optional.h"
-#include "Templates/NonNullPointer.h"
-#endif
diff --git a/Engine/Plugins/Developer/Concert/ConcertSync/ConcertSyncCore/Source/ConcertSyncCore/Public/HistoryEdition/ActivityNode.h b/Engine/Plugins/Developer/Concert/ConcertSync/ConcertSyncCore/Source/ConcertSyncCore/Public/HistoryEdition/ActivityNode.h
--- a/Engine/Plugins/Developer/Concert/ConcertSync/ConcertSyncCore/Source/ConcertSyncCore/Public/HistoryEdition/ActivityNode.h
+++ b/Engine/Plugins/Developer/Concert/ConcertSync/ConcertSyncCore/Source/ConcertSyncCore/Public/HistoryEdition/ActivityNode.h
@@ -70,9 +70,3 @@ namespace UE::ConcertSyncCore
 		TArray<FActivityNodeID> AffectedChildren;
 	};
 }
-
-#if UE_ENABLE_INCLUDE_ORDER_DEPRECATED_IN_5_2
-#include "CoreMinimal.h"
-#include "Misc/Optional.h"
-#include "Templates/NonNullPointer.h"
-#endif
diff --git a/Engine/Plugins/Developer/Concert/ConcertSync/ConcertSyncCore/Source/ConcertSyncCore/Public/HistoryEdition/DebugDependencyGraph.h b/Engine/Plugins/Developer/Concert/ConcertSync/ConcertSyncCore/Source/ConcertSyncCore/Public/HistoryEdition/DebugDependencyGraph.h
--- a/Engine/Plugins/Developer/Concert/ConcertSync/ConcertSyncCore/Source/ConcertSyncCore/Public/HistoryEdition/DebugDependencyGraph.h
+++ b/Engine/Plugins/Developer/Concert/ConcertSync/ConcertSyncCore/Source/ConcertSyncCore/Public/HistoryEdition/DebugDependencyGraph.h
@@ -67,9 +67,3 @@ namespace UE::ConcertSyncCore::Graphviz
 	 */
 	CONCERTSYNCCORE_API FString ExportToGraphviz(const FActivityDependencyGraph& Graph, FMakeNodeTitle MakeNodeTitleFunc, FGetNodeStyle GetNodeStyleFunc, FGetEdgeStyle GetEdgeStyleFunc);
 }
-
-#if UE_ENABLE_INCLUDE_ORDER_DEPRECATED_IN_5_2
-#include "ActivityDependencyGraph.h"
-#include "ActivityGraphIDs.h"
-#include "CoreMinimal.h"
-#endif
diff --git a/Engine/Plugins/Developer/Concert/ConcertSync/ConcertSyncCore/Source/ConcertSyncCore/Public/HistoryEdition/DependencyGraphBuilder.h b/Engine/Plugins/Developer/Concert/ConcertSync/ConcertSyncCore/Source/ConcertSyncCore/Public/HistoryEdition/DependencyGraphBuilder.h
--- a/Engine/Plugins/Developer/Concert/ConcertSync/ConcertSyncCore/Source/ConcertSyncCore/Public/HistoryEdition/DependencyGraphBuilder.h
+++ b/Engine/Plugins/Developer/Concert/ConcertSync/ConcertSyncCore/Source/ConcertSyncCore/Public/HistoryEdition/DependencyGraphBuilder.h
@@ -15,8 +15,3 @@ namespace UE::ConcertSyncCore
 	 */
 	CONCERTSYNCCORE_API FActivityDependencyGraph BuildDependencyGraphFrom(const FConcertSyncSessionDatabase& SessionDatabase);
 }
-
-#if UE_ENABLE_INCLUDE_ORDER_DEPRECATED_IN_5_2
-#include "ActivityDependencyGraph.h"
-#include "CoreMinimal.h"
-#endif
diff --git a/Engine/Plugins/Developer/Concert/ConcertSync/ConcertSyncCore/Source/ConcertSyncCore/Public/HistoryEdition/HistoryAnalysis.h b/Engine/Plugins/Developer/Concert/ConcertSync/ConcertSyncCore/Source/ConcertSyncCore/Public/HistoryEdition/HistoryAnalysis.h
--- a/Engine/Plugins/Developer/Concert/ConcertSync/ConcertSyncCore/Source/ConcertSyncCore/Public/HistoryEdition/HistoryAnalysis.h
+++ b/Engine/Plugins/Developer/Concert/ConcertSync/ConcertSyncCore/Source/ConcertSyncCore/Public/HistoryEdition/HistoryAnalysis.h
@@ -54,9 +54,3 @@ namespace UE::ConcertSyncCore
 	 */
 	CONCERTSYNCCORE_API FHistoryAnalysisResult AnalyseActivityDependencies_BottomUp(const TSet<FActivityID>& ActivitiesToEdit, const FActivityDependencyGraph& DependencyGraph, bool bAddEditedAsHardDependencies = false);
 }
-
-
-#if UE_ENABLE_INCLUDE_ORDER_DEPRECATED_IN_5_2
-#include "ActivityDependencyGraph.h"
-#include "CoreMinimal.h"
-#endif
diff --git a/Engine/Plugins/Developer/Concert/ConcertSync/ConcertSyncCore/Source/ConcertSyncCore/Public/HistoryEdition/HistoryEdition.h b/Engine/Plugins/Developer/Concert/ConcertSync/ConcertSyncCore/Source/ConcertSyncCore/Public/HistoryEdition/HistoryEdition.h
--- a/Engine/Plugins/Developer/Concert/ConcertSync/ConcertSyncCore/Source/ConcertSyncCore/Public/HistoryEdition/HistoryEdition.h
+++ b/Engine/Plugins/Developer/Concert/ConcertSync/ConcertSyncCore/Source/ConcertSyncCore/Public/HistoryEdition/HistoryEdition.h
@@ -29,8 +29,3 @@ namespace UE::ConcertSyncCore
 	/** Utility functions that converts FHistoryDeletionRequirements into a single TSet. */
 	CONCERTSYNCCORE_API TSet<FActivityID> CombineRequirements(const FHistoryAnalysisResult& ToDelete);
 }
-
-#if UE_ENABLE_INCLUDE_ORDER_DEPRECATED_IN_5_2
-#include "CoreMinimal.h"
-#include "HistoryAnalysis.h"
-#endif
diff --git a/Engine/Plugins/Developer/TraceSourceFiltering/Source/SourceFilteringCore/Public/IDataSourceFilterSetInterface.h b/Engine/Plugins/Developer/TraceSourceFiltering/Source/SourceFilteringCore/Public/IDataSourceFilterSetInterface.h
--- a/Engine/Plugins/Developer/TraceSourceFiltering/Source/SourceFilteringCore/Public/IDataSourceFilterSetInterface.h
+++ b/Engine/Plugins/Developer/TraceSourceFiltering/Source/SourceFilteringCore/Public/IDataSourceFilterSetInterface.h
@@ -22,9 +22,3 @@ public:
 	/** Return the current Filter Set operation mode */
 	virtual EFilterSetMode GetFilterSetMode() const = 0;
 };
-
-
-#if UE_ENABLE_INCLUDE_ORDER_DEPRECATED_IN_5_2
-#include "DataSourceFiltering.h"
-#include "IDataSourceFilterInterface.h"
-#endif
diff --git a/Engine/Plugins/Developer/TraceSourceFiltering/Source/SourceFilteringTrace/Public/DataSourceFilter.h b/Engine/Plugins/Developer/TraceSourceFiltering/Source/SourceFilteringTrace/Public/DataSourceFilter.h
--- a/Engine/Plugins/Developer/TraceSourceFiltering/Source/SourceFilteringTrace/Public/DataSourceFilter.h
+++ b/Engine/Plugins/Developer/TraceSourceFiltering/Source/SourceFilteringTrace/Public/DataSourceFilter.h
@@ -38,9 +38,3 @@ protected:
 	/** Whether or not this filter is enabled */
 	bool bIsEnabled;
 };
-
-
-#if UE_ENABLE_INCLUDE_ORDER_DEPRECATED_IN_5_2
-#include "CoreMinimal.h"
-#include "GameFramework/Actor.h"
-#endif
diff --git a/Engine/Plugins/Developer/TraceSourceFiltering/Source/SourceFilteringTrace/Public/SourceFilterCollection.h b/Engine/Plugins/Developer/TraceSourceFiltering/Source/SourceFilteringTrace/Public/SourceFilterCollection.h
--- a/Engine/Plugins/Developer/TraceSourceFiltering/Source/SourceFilteringTrace/Public/SourceFilterCollection.h
+++ b/Engine/Plugins/Developer/TraceSourceFiltering/Source/SourceFilteringTrace/Public/SourceFilterCollection.h
@@ -121,8 +121,3 @@ protected:
 	
 	FSimpleMulticastDelegate SourceFiltersUpdatedDelegate;
 };
-
-#if UE_ENABLE_INCLUDE_ORDER_DEPRECATED_IN_5_2
-#include "DataSourceFilterSet.h"
-#include "Delegates/DelegateCombinations.h"
-#endif
diff --git a/Engine/Plugins/Developer/TraceSourceFiltering/Source/SourceFilteringTrace/Public/SourceFilterTrace.h b/Engine/Plugins/Developer/TraceSourceFiltering/Source/SourceFilteringTrace/Public/SourceFilterTrace.h
--- a/Engine/Plugins/Developer/TraceSourceFiltering/Source/SourceFilteringTrace/Public/SourceFilterTrace.h
+++ b/Engine/Plugins/Developer/TraceSourceFiltering/Source/SourceFilteringTrace/Public/SourceFilterTrace.h
@@ -109,12 +109,3 @@ protected:
 #define TRACE_FILTER_IDENTIFIER(Object) 0
 
 #endif
-
-#if UE_ENABLE_INCLUDE_ORDER_DEPRECATED_IN_5_2
-#include "DataSourceFilter.h"
-#include "DataSourceFilterSet.h"
-#include "DataSourceFiltering.h"
-#include "ObjectTrace.h"
-#include "Templates/SubclassOf.h"
-#include "UObject/ObjectKey.h"
-#endif
diff --git a/Engine/Plugins/Developer/TraceSourceFiltering/Source/SourceFilteringTrace/Public/TraceSourceFiltering.h b/Engine/Plugins/Developer/TraceSourceFiltering/Source/SourceFilteringTrace/Public/TraceSourceFiltering.h
--- a/Engine/Plugins/Developer/TraceSourceFiltering/Source/SourceFilteringTrace/Public/TraceSourceFiltering.h
+++ b/Engine/Plugins/Developer/TraceSourceFiltering/Source/SourceFilteringTrace/Public/TraceSourceFiltering.h
@@ -48,8 +48,3 @@ protected:
 	/** Mapping for all filtering commands from their name to respective FFilterCommand object */
 	TMap<FString, FFilterCommand> CommandMap;
 };
-
-#if UE_ENABLE_INCLUDE_ORDER_DEPRECATED_IN_5_2
-#include "CoreMinimal.h"
-#include "DataSourceFilter.h"
-#endif
diff --git a/Engine/Plugins/Developer/TraceSourceFiltering/Source/SourceFilteringTrace/Public/TraceSourceFilteringProjectSettings.h b/Engine/Plugins/Developer/TraceSourceFiltering/Source/SourceFilteringTrace/Public/TraceSourceFilteringProjectSettings.h
--- a/Engine/Plugins/Developer/TraceSourceFiltering/Source/SourceFilteringTrace/Public/TraceSourceFilteringProjectSettings.h
+++ b/Engine/Plugins/Developer/TraceSourceFiltering/Source/SourceFilteringTrace/Public/TraceSourceFilteringProjectSettings.h
@@ -22,8 +22,3 @@ public:
 	UPROPERTY(config, EditAnywhere, Category = TraceSourceFiltering, AdvancedDisplay, meta = (DisplayName = "Default Filter preset, which should be loaded during boot", RelativeToGameContentDir))
 	TSoftObjectPtr<USourceFilterCollection> DefaultFilterPreset;
 };
-
-#if UE_ENABLE_INCLUDE_ORDER_DEPRECATED_IN_5_2
-#include "DataSourceFilter.h"
-#include "SourceFilterCollection.h"
-#endif
diff --git a/Engine/Plugins/Developer/TraceSourceFiltering/Source/SourceFilteringTrace/Public/TraceSourceFilteringTestBase.h b/Engine/Plugins/Developer/TraceSourceFiltering/Source/SourceFilteringTrace/Public/TraceSourceFilteringTestBase.h
--- a/Engine/Plugins/Developer/TraceSourceFiltering/Source/SourceFilteringTrace/Public/TraceSourceFilteringTestBase.h
+++ b/Engine/Plugins/Developer/TraceSourceFiltering/Source/SourceFilteringTrace/Public/TraceSourceFilteringTestBase.h
@@ -154,11 +154,3 @@ protected:
 };
 
 #endif // WITH_AUTOMATION_TESTS
-
-#if UE_ENABLE_INCLUDE_ORDER_DEPRECATED_IN_5_2
-#include "PreviewScene.h"
-#include "TraceFilter.h"
-#include "TraceSourceFiltering.h"
-#include "TraceSourceFilteringProjectSettings.h"
-#include "TraceWorldFiltering.h"
-#endif
diff --git a/Engine/Plugins/Editor/AssetSearch/Source/Public/IAssetSearchModule.h b/Engine/Plugins/Editor/AssetSearch/Source/Public/IAssetSearchModule.h
--- a/Engine/Plugins/Editor/AssetSearch/Source/Public/IAssetSearchModule.h
+++ b/Engine/Plugins/Editor/AssetSearch/Source/Public/IAssetSearchModule.h
@@ -57,10 +57,3 @@ public:
 	/** Virtual destructor. */
 	virtual ~IAssetSearchModule() { }
 };
-
-#if UE_ENABLE_INCLUDE_ORDER_DEPRECATED_IN_5_2
-#include "Async/AsyncResult.h"
-#include "Async/Future.h"
-#include "CoreMinimal.h"
-#include "SearchSerializer.h"
-#endif
diff --git a/Engine/Plugins/Editor/AssetSearch/Source/Public/SearchSerializer.h b/Engine/Plugins/Editor/AssetSearch/Source/Public/SearchSerializer.h
--- a/Engine/Plugins/Editor/AssetSearch/Source/Public/SearchSerializer.h
+++ b/Engine/Plugins/Editor/AssetSearch/Source/Public/SearchSerializer.h
@@ -77,10 +77,3 @@ private:
 
 	TArray<UObject*> NestedAssets;
 };
-
-#if UE_ENABLE_INCLUDE_ORDER_DEPRECATED_IN_5_2
-#include "CoreMinimal.h"
-#include "JsonObjectConverter.h"
-#include "Policies/CondensedJsonPrintPolicy.h"
-#include "Serialization/JsonSerializer.h"
-#endif
diff --git a/Engine/Plugins/Editor/BlueprintHeaderView/Source/BlueprintHeaderView/Public/BlueprintHeaderView.h b/Engine/Plugins/Editor/BlueprintHeaderView/Source/BlueprintHeaderView/Public/BlueprintHeaderView.h
--- a/Engine/Plugins/Editor/BlueprintHeaderView/Source/BlueprintHeaderView/Public/BlueprintHeaderView.h
+++ b/Engine/Plugins/Editor/BlueprintHeaderView/Source/BlueprintHeaderView/Public/BlueprintHeaderView.h
@@ -46,10 +46,3 @@ private:
 	/** Handle to our delegate so we can remove it at module shutdown */
 	FDelegateHandle ContentBrowserExtenderDelegateHandle;
 };
-
-#if UE_ENABLE_INCLUDE_ORDER_DEPRECATED_IN_5_2
-#include "AssetRegistry/AssetData.h"
-#include "CoreMinimal.h"
-#include "Modules/ModuleManager.h"
-#include "Styling/SlateTypes.h"
-#endif
diff --git a/Engine/Plugins/Editor/BlueprintMaterialTextureNodes/Source/BlueprintMaterialTextureNodes/Public/BlueprintMaterialTextureNodes.h b/Engine/Plugins/Editor/BlueprintMaterialTextureNodes/Source/BlueprintMaterialTextureNodes/Public/BlueprintMaterialTextureNodes.h
--- a/Engine/Plugins/Editor/BlueprintMaterialTextureNodes/Source/BlueprintMaterialTextureNodes/Public/BlueprintMaterialTextureNodes.h
+++ b/Engine/Plugins/Editor/BlueprintMaterialTextureNodes/Source/BlueprintMaterialTextureNodes/Public/BlueprintMaterialTextureNodes.h
@@ -12,7 +12,3 @@ public:
 	virtual void StartupModule() override;
 	virtual void ShutdownModule() override;
 };
-
-#if UE_ENABLE_INCLUDE_ORDER_DEPRECATED_IN_5_2
-#include "Modules/ModuleManager.h"
-#endif
diff --git a/Engine/Plugins/Editor/ChangelistReview/Source/Public/ChangelistReviewModule.h b/Engine/Plugins/Editor/ChangelistReview/Source/Public/ChangelistReviewModule.h
--- a/Engine/Plugins/Editor/ChangelistReview/Source/Public/ChangelistReviewModule.h
+++ b/Engine/Plugins/Editor/ChangelistReview/Source/Public/ChangelistReviewModule.h
@@ -40,9 +40,3 @@ private:
 	TWeakPtr<SDockTab> ReviewTab;
 	TWeakPtr<SSourceControlReview> ReviewWidget;
 };
-
-#if UE_ENABLE_INCLUDE_ORDER_DEPRECATED_IN_5_2
-#include "UObject/StrongObjectPtr.h"
-#include "UObject/WeakObjectPtr.h"
-#include "Widgets/Docking/SDockTab.h"
-#endif
diff --git a/Engine/Plugins/Editor/ContentBrowser/ContentBrowserAliasDataSource/Source/ContentBrowserAliasDataSource/Public/ContentBrowserAliasDataSource.h b/Engine/Plugins/Editor/ContentBrowser/ContentBrowserAliasDataSource/Source/ContentBrowserAliasDataSource/Public/ContentBrowserAliasDataSource.h
--- a/Engine/Plugins/Editor/ContentBrowser/ContentBrowserAliasDataSource/Source/ContentBrowserAliasDataSource/Public/ContentBrowserAliasDataSource.h
+++ b/Engine/Plugins/Editor/ContentBrowser/ContentBrowserAliasDataSource/Source/ContentBrowserAliasDataSource/Public/ContentBrowserAliasDataSource.h
@@ -243,9 +243,3 @@ private:
 	/** Delegate broadcast after all aliases removed to give chance for systems to re-add aliases */
 	FSimpleMulticastDelegate RebuildAliasesDelegate;
 };
-
-#if UE_ENABLE_INCLUDE_ORDER_DEPRECATED_IN_5_2
-#include "AssetRegistry/AssetRegistryModule.h"
-#include "ContentBrowserItemPath.h"
-#include "CoreMinimal.h"
-#endif
diff --git a/Engine/Plugins/Editor/ContentBrowser/ContentBrowserAssetDataSource/Source/ContentBrowserAssetDataSource/Public/AssetFileContextMenu.h b/Engine/Plugins/Editor/ContentBrowser/ContentBrowserAssetDataSource/Source/ContentBrowserAssetDataSource/Public/AssetFileContextMenu.h
--- a/Engine/Plugins/Editor/ContentBrowser/ContentBrowserAssetDataSource/Source/ContentBrowserAssetDataSource/Public/AssetFileContextMenu.h
+++ b/Engine/Plugins/Editor/ContentBrowser/ContentBrowserAssetDataSource/Source/ContentBrowserAssetDataSource/Public/AssetFileContextMenu.h
@@ -249,8 +249,3 @@ private:
 	/** */
 	int32 ChunkIDSelected = 0;
 };
-
-#if UE_ENABLE_INCLUDE_ORDER_DEPRECATED_IN_5_2
-#include "CoreMinimal.h"
-#include "Input/Reply.h"
-#endif
diff --git a/Engine/Plugins/Editor/ContentBrowser/ContentBrowserAssetDataSource/Source/ContentBrowserAssetDataSource/Public/ContentBrowserAssetDataCore.h b/Engine/Plugins/Editor/ContentBrowser/ContentBrowserAssetDataSource/Source/ContentBrowserAssetDataSource/Public/ContentBrowserAssetDataCore.h
--- a/Engine/Plugins/Editor/ContentBrowser/ContentBrowserAssetDataSource/Source/ContentBrowserAssetDataSource/Public/ContentBrowserAssetDataCore.h
+++ b/Engine/Plugins/Editor/ContentBrowser/ContentBrowserAssetDataSource/Source/ContentBrowserAssetDataSource/Public/ContentBrowserAssetDataCore.h
@@ -219,8 +219,3 @@ CONTENTBROWSERASSETDATASOURCE_API FContentBrowserItemData CreateUnsupportedAsset
 	CONTENTBROWSERASSETDATASOURCE_API void PopulateAssetFileContextMenu(UContentBrowserDataSource* InOwnerDataSource, UToolMenu* InMenu, FAssetFileContextMenu& InAssetFileContextMenu);
 
 }
-
-#if UE_ENABLE_INCLUDE_ORDER_DEPRECATED_IN_5_2
-#include "ContentBrowserAssetDataPayload.h"
-#include "CoreMinimal.h"
-#endif
diff --git a/Engine/Plugins/Editor/ContentBrowser/ContentBrowserAssetDataSource/Source/ContentBrowserAssetDataSource/Public/ContentBrowserAssetDataSource.h b/Engine/Plugins/Editor/ContentBrowser/ContentBrowserAssetDataSource/Source/ContentBrowserAssetDataSource/Public/ContentBrowserAssetDataSource.h
--- a/Engine/Plugins/Editor/ContentBrowser/ContentBrowserAssetDataSource/Source/ContentBrowserAssetDataSource/Public/ContentBrowserAssetDataSource.h
+++ b/Engine/Plugins/Editor/ContentBrowser/ContentBrowserAssetDataSource/Source/ContentBrowserAssetDataSource/Public/ContentBrowserAssetDataSource.h
@@ -546,10 +546,3 @@ private:
 
 	FAssetDataSourceFilterCache FilterCache;
 };
-
-#if UE_ENABLE_INCLUDE_ORDER_DEPRECATED_IN_5_2
-#include "ContentBrowserAssetDataPayload.h"
-#include "CoreMinimal.h"
-#include "Input/Reply.h"
-#include "UObject/GCObject.h"
-#endif
diff --git a/Engine/Plugins/Editor/ContentBrowser/ContentBrowserClassDataSource/Source/ContentBrowserClassDataSource/Public/ContentBrowserClassDataCore.h b/Engine/Plugins/Editor/ContentBrowser/ContentBrowserClassDataSource/Source/ContentBrowserClassDataSource/Public/ContentBrowserClassDataCore.h
--- a/Engine/Plugins/Editor/ContentBrowser/ContentBrowserClassDataSource/Source/ContentBrowserClassDataSource/Public/ContentBrowserClassDataCore.h
+++ b/Engine/Plugins/Editor/ContentBrowser/ContentBrowserClassDataSource/Source/ContentBrowserClassDataSource/Public/ContentBrowserClassDataCore.h
@@ -77,8 +77,3 @@ namespace ContentBrowserClassData
 	CONTENTBROWSERCLASSDATASOURCE_API bool GetClassFileItemAttributes(const FContentBrowserClassFileItemDataPayload& InClassPayload, const bool InIncludeMetaData, FContentBrowserItemDataAttributeValues& OutAttributeValues);
 
 }
-
-#if UE_ENABLE_INCLUDE_ORDER_DEPRECATED_IN_5_2
-#include "ContentBrowserClassDataPayload.h"
-#include "CoreMinimal.h"
-#endif
diff --git a/Engine/Plugins/Editor/ContentBrowser/ContentBrowserClassDataSource/Source/ContentBrowserClassDataSource/Public/ContentBrowserClassDataSource.h b/Engine/Plugins/Editor/ContentBrowser/ContentBrowserClassDataSource/Source/ContentBrowserClassDataSource/Public/ContentBrowserClassDataSource.h
--- a/Engine/Plugins/Editor/ContentBrowser/ContentBrowserClassDataSource/Source/ContentBrowserClassDataSource/Public/ContentBrowserClassDataSource.h
+++ b/Engine/Plugins/Editor/ContentBrowser/ContentBrowserClassDataSource/Source/ContentBrowserClassDataSource/Public/ContentBrowserClassDataSource.h
@@ -125,9 +125,3 @@ private:
 
 	ICollectionManager* CollectionManager;
 };
-
-#if UE_ENABLE_INCLUDE_ORDER_DEPRECATED_IN_5_2
-#include "ContentBrowserClassDataPayload.h"
-#include "ContentBrowserDataMenuContexts.h"
-#include "CoreMinimal.h"
-#endif
diff --git a/Engine/Plugins/Editor/ContentBrowser/ContentBrowserFileDataSource/Source/ContentBrowserFileDataSource/Public/ContentBrowserFileDataCore.h b/Engine/Plugins/Editor/ContentBrowser/ContentBrowserFileDataSource/Source/ContentBrowserFileDataSource/Public/ContentBrowserFileDataCore.h
--- a/Engine/Plugins/Editor/ContentBrowser/ContentBrowserFileDataSource/Source/ContentBrowserFileDataSource/Public/ContentBrowserFileDataCore.h
+++ b/Engine/Plugins/Editor/ContentBrowser/ContentBrowserFileDataSource/Source/ContentBrowserFileDataSource/Public/ContentBrowserFileDataCore.h
@@ -140,8 +140,3 @@ namespace ContentBrowserFileData
 	};
 
 }
-
-#if UE_ENABLE_INCLUDE_ORDER_DEPRECATED_IN_5_2
-#include "ContentBrowserFileDataPayload.h"
-#include "CoreMinimal.h"
-#endif
diff --git a/Engine/Plugins/Editor/CryptoKeys/Source/CryptoKeys/Public/CryptoKeys.h b/Engine/Plugins/Editor/CryptoKeys/Source/CryptoKeys/Public/CryptoKeys.h
--- a/Engine/Plugins/Editor/CryptoKeys/Source/CryptoKeys/Public/CryptoKeys.h
+++ b/Engine/Plugins/Editor/CryptoKeys/Source/CryptoKeys/Public/CryptoKeys.h
@@ -10,7 +10,3 @@ namespace CryptoKeys
 {
 	CRYPTOKEYS_API void GenerateEncryptionKey(FString& OutBase64Key);
 }
-
-#if UE_ENABLE_INCLUDE_ORDER_DEPRECATED_IN_5_2
-#include "Modules/ModuleManager.h"
-#endif
diff --git a/Engine/Plugins/Editor/CryptoKeys/Source/CryptoKeysOpenSSL/Public/CryptoKeysOpenSSL.h b/Engine/Plugins/Editor/CryptoKeys/Source/CryptoKeysOpenSSL/Public/CryptoKeysOpenSSL.h
--- a/Engine/Plugins/Editor/CryptoKeys/Source/CryptoKeysOpenSSL/Public/CryptoKeysOpenSSL.h
+++ b/Engine/Plugins/Editor/CryptoKeys/Source/CryptoKeysOpenSSL/Public/CryptoKeysOpenSSL.h
@@ -8,8 +8,3 @@ namespace CryptoKeysOpenSSL
 	bool CRYPTOKEYSOPENSSL_API GenerateNewEncryptionKey(TArray<uint8>& OutKey);
 	bool CRYPTOKEYSOPENSSL_API GenerateNewSigningKey(TArray<uint8>& OutPublicExponent, TArray<uint8>& OutPrivateExponent, TArray<uint8>& OutModulus, int32 InNumKeyBits = 2048);
 }
-
-#if UE_ENABLE_INCLUDE_ORDER_DEPRECATED_IN_5_2
-#include "Containers/Array.h"
-#include "CoreMinimal.h"
-#endif
diff --git a/Engine/Plugins/Editor/DataValidation/Source/DataValidation/Public/DataValidationChangelist.h b/Engine/Plugins/Editor/DataValidation/Source/DataValidation/Public/DataValidationChangelist.h
--- a/Engine/Plugins/Editor/DataValidation/Source/DataValidation/Public/DataValidationChangelist.h
+++ b/Engine/Plugins/Editor/DataValidation/Source/DataValidation/Public/DataValidationChangelist.h
@@ -46,9 +46,3 @@ public:
 	
 	FText Description;
 };
-
-#if UE_ENABLE_INCLUDE_ORDER_DEPRECATED_IN_5_2
-#include "EditorSubsystem.h"
-#include "Engine/EngineTypes.h"
-#include "UObject/Interface.h"
-#endif
diff --git a/Engine/Plugins/Editor/DataValidation/Source/DataValidation/Public/DataValidationCommandlet.h b/Engine/Plugins/Editor/DataValidation/Source/DataValidation/Public/DataValidationCommandlet.h
--- a/Engine/Plugins/Editor/DataValidation/Source/DataValidation/Public/DataValidationCommandlet.h
+++ b/Engine/Plugins/Editor/DataValidation/Source/DataValidation/Public/DataValidationCommandlet.h
@@ -23,7 +23,3 @@ public:
 	// do the validation without creating a commandlet
 	static bool ValidateData(const FString& FullCommandLine);
 };
-
-#if UE_ENABLE_INCLUDE_ORDER_DEPRECATED_IN_5_2
-#include "UObject/Interface.h"
-#endif
diff --git a/Engine/Plugins/Editor/DataValidation/Source/DataValidation/Public/EditorValidatorBase.h b/Engine/Plugins/Editor/DataValidation/Source/DataValidation/Public/EditorValidatorBase.h
--- a/Engine/Plugins/Editor/DataValidation/Source/DataValidation/Public/EditorValidatorBase.h
+++ b/Engine/Plugins/Editor/DataValidation/Source/DataValidation/Public/EditorValidatorBase.h
@@ -156,16 +156,3 @@ private:
 
 	const FAssetData* CurrentAssetBeingValidated;
 };
-
-
-
-
-
-
-#if UE_ENABLE_INCLUDE_ORDER_DEPRECATED_IN_5_2
-#include "DataValidationModule.h"
-#include "EditorSubsystem.h"
-#include "Templates/SubclassOf.h"
-#include "UObject/Interface.h"
-#include "UObject/TextProperty.h"
-#endif
diff --git a/Engine/Plugins/Editor/DataValidation/Source/DataValidation/Public/EditorValidatorSubsystem.h b/Engine/Plugins/Editor/DataValidation/Source/DataValidation/Public/EditorValidatorSubsystem.h
--- a/Engine/Plugins/Editor/DataValidation/Source/DataValidation/Public/EditorValidatorSubsystem.h
+++ b/Engine/Plugins/Editor/DataValidation/Source/DataValidation/Public/EditorValidatorSubsystem.h
@@ -418,12 +418,3 @@ protected:
 
 	bool bNeedLoadingOfValidators = false;
 };
-
-
-#if UE_ENABLE_INCLUDE_ORDER_DEPRECATED_IN_5_2
-#include "AssetRegistry/AssetData.h"
-#include "EditorValidatorBase.h"
-#include "Engine/EngineTypes.h"
-#include "Templates/SubclassOf.h"
-#include "UObject/Interface.h"
-#endif
diff --git a/Engine/Plugins/Editor/EditorDebugTools/Source/EditorDebugTools/Public/EditorDebugToolsStyle.h b/Engine/Plugins/Editor/EditorDebugTools/Source/EditorDebugTools/Public/EditorDebugToolsStyle.h
--- a/Engine/Plugins/Editor/EditorDebugTools/Source/EditorDebugTools/Public/EditorDebugToolsStyle.h
+++ b/Engine/Plugins/Editor/EditorDebugTools/Source/EditorDebugTools/Public/EditorDebugToolsStyle.h
@@ -29,8 +29,3 @@ private:
 
 	static TSharedPtr< class FSlateStyleSet > StyleInstance;
 };
-
-#if UE_ENABLE_INCLUDE_ORDER_DEPRECATED_IN_5_2
-#include "CoreMinimal.h"
-#include "Styling/SlateStyle.h"
-#endif
