Remove deprecated LiveLinkClient code

LiveLinkMessageBus will no longer handle FLiveLinkSubjectDataMessage and FLiveLinkSubjectFrameMessage messages, you should upgrade your code to use the animation role.

diff --git a/Engine/Plugins/Animation/LiveLink/Source/LiveLink/Private/LiveLinkClient.cpp b/Engine/Plugins/Animation/LiveLink/Source/LiveLink/Private/LiveLinkClient.cpp
--- a/Engine/Plugins/Animation/LiveLink/Source/LiveLink/Private/LiveLinkClient.cpp
+++ b/Engine/Plugins/Animation/LiveLink/Source/LiveLink/Private/LiveLinkClient.cpp
@@ -1913,45 +1913,6 @@ FOnLiveLinkSubjectEvaluated& FLiveLinkClient::OnLiveLinkSubjectEvaluated()
 }
 #endif
 
-PRAGMA_DISABLE_DEPRECATION_WARNINGS
-/**
- * Function for Deprecation
- */
-void FLiveLinkClient::AquireLock_Deprecation() 
-{
-}
-
-void FLiveLinkClient::ReleaseLock_Deprecation()
-{
-}
-
-void FLiveLinkClient::ClearFrames_Deprecation(const FLiveLinkSubjectKey& InSubjectKey)
-{
-	if (Collection)
-	{
-		if (FLiveLinkCollectionSubjectItem* SubjectItem = Collection->FindSubject(InSubjectKey))
-		{
-			SubjectItem->GetSubject()->ClearFrames();
-		}
-	}
-}
-
-FLiveLinkSkeletonStaticData* FLiveLinkClient::GetSubjectAnimationStaticData_Deprecation(const FLiveLinkSubjectKey& InSubjectKey)
-{
-	if (Collection)
-	{
-		if (FLiveLinkCollectionSubjectItem* SubjectItem = Collection->FindSubject(InSubjectKey))
-		{
-			if (SubjectItem->GetSubject()->GetRole() == ULiveLinkAnimationRole::StaticClass() && !SubjectItem->bPendingKill)
-			{
-				return SubjectItem->GetSubject()->GetStaticData().Cast<FLiveLinkSkeletonStaticData>();
-			}
-		}
-	}
-
-	return nullptr;
-}
-
 void FLiveLinkClient::UnregisterGlobalSubjectFramesDelegate(FDelegateHandle& InStaticDataAddedHandle, FDelegateHandle& InFrameDataAddedHandle, bool bUseUnmappedData)
 {
 	if (FSubjectFramesAddedHandles* Handles = SubjectFrameAddedHandles.Find(ALL_SUBJECTS_DELEGATE_TOKEN))
@@ -1979,195 +1940,3 @@ bool FLiveLinkClient::RegisterGlobalSubjectFramesDelegate(const FOnLiveLinkSubje
 
 	return true;
 }
-
-/**
- * Function that are now deprecated
- */
-const TArray<FGuid>& FLiveLinkClient::GetSourceEntries() const
-{
-	static TArray<FGuid> CopiedSources;
-	CopiedSources = GetSources();
-	return CopiedSources;
-}
-
-void FLiveLinkClient::AddVirtualSubject(FName InNewVirtualSubjectName)
-{
-	AddVirtualSubject(InNewVirtualSubjectName, ULiveLinkAnimationVirtualSubject::StaticClass());
-}
-
-void FLiveLinkClient::AddVirtualSubject(FLiveLinkSubjectName VirtualSubjectName, TSubclassOf<ULiveLinkVirtualSubject> VirtualSubjectClass)
-{
-	const FLiveLinkSubjectKey DefaultSubjectKey(FLiveLinkSourceCollection::DefaultVirtualSubjectGuid, VirtualSubjectName);
-	AddVirtualSubject(DefaultSubjectKey, ULiveLinkAnimationVirtualSubject::StaticClass());
-}
-
-/**
- * Deprecated function
- */
-void FLiveLinkClient_Base_DEPRECATED::PushSubjectSkeleton(FGuid SourceGuid, FName SubjectName, const FLiveLinkRefSkeleton& RefSkeleton)
-{
-	// Backward compatibility with old API. Default to pushing animation data
-
-	FLiveLinkSubjectKey Key = FLiveLinkSubjectKey(SourceGuid, SubjectName);
-
-	RemoveSubject_AnyThread(Key);
-
-	FLiveLinkStaticDataStruct StaticData(FLiveLinkSkeletonStaticData::StaticStruct());
-	FLiveLinkSkeletonStaticData* SkeletonData = StaticData.Cast<FLiveLinkSkeletonStaticData>();
-	SkeletonData->SetBoneNames(RefSkeleton.GetBoneNames());
-	SkeletonData->SetBoneParents(RefSkeleton.GetBoneParents());
-	PushSubjectStaticData_AnyThread(Key, ULiveLinkAnimationRole::StaticClass(), MoveTemp(StaticData));
-}
-
-void UpdateForAnimationStatic(TArray<FName>& InOutCurveNames, const TArray<FLiveLinkCurveElement>& InCurveElements)
-{
-	InOutCurveNames.Reset(InCurveElements.Num());
-	for (const FLiveLinkCurveElement& Elem : InCurveElements)
-	{
-		InOutCurveNames.Add(Elem.CurveName);
-	}
-}
-
-void FLiveLinkClient_Base_DEPRECATED::PushSubjectData(FGuid InSourceGuid, FName InSubjectName, const FLiveLinkFrameData& InFrameData)
-{
-	FLiveLinkSubjectKey SubjectKey(InSourceGuid, InSubjectName);
-
-	//Update curve names in the static data for backward compatibility
-	int32 NumberOfPropertyNames = 0;
-	{
-		AquireLock_Deprecation();
-
-		FLiveLinkSkeletonStaticData* AnimationStaticData = GetSubjectAnimationStaticData_Deprecation(SubjectKey);
-		if (AnimationStaticData)
-		{
-			NumberOfPropertyNames = AnimationStaticData->PropertyNames.Num();
-			if (NumberOfPropertyNames == 0 && InFrameData.CurveElements.Num() > 0)
-			{
-				static const FName NAME_UpdateYourCode = "LiveLinkClient_PushSubjectData";
-				FLiveLinkLog::WarningOnce(NAME_UpdateYourCode, FLiveLinkSubjectKey(InSourceGuid, InSubjectName), TEXT("Upgrade your code. Curve elements count has changed from the previous frame. That will clear the previous frames of that subject."));
-
-				ClearFrames_Deprecation(SubjectKey);
-				UpdateForAnimationStatic(AnimationStaticData->PropertyNames, InFrameData.CurveElements);
-				NumberOfPropertyNames = AnimationStaticData->PropertyNames.Num();
-			}
-		}
-
-		ReleaseLock_Deprecation();
-	}
-
-	//Convert incoming data as animation data
-	FLiveLinkFrameDataStruct AnimationStruct(FLiveLinkAnimationFrameData::StaticStruct());
-	FLiveLinkAnimationFrameData& NewData = *AnimationStruct.Cast<FLiveLinkAnimationFrameData>();
-	NewData.MetaData = InFrameData.MetaData;
-	NewData.WorldTime = InFrameData.WorldTime;
-	NewData.Transforms = InFrameData.Transforms;
-
-	//Always match FrameData property count to StaticData property count.
-	//If StaticData has more properties than current FrameData, set non existent properties to Infinity
-	//If StaticData has less properties than current FrameData, only use a subset of the incoming properties
-	int32 MaxNumberOfProperties = FMath::Min(NumberOfPropertyNames, InFrameData.CurveElements.Num());
-	NewData.PropertyValues.SetNumZeroed(NumberOfPropertyNames);
-	for (int32 i = 0; i < MaxNumberOfProperties; ++i)
-	{
-		NewData.PropertyValues[i] = InFrameData.CurveElements[i].CurveValue;
-	}
-	for (int32 i = MaxNumberOfProperties; i < NumberOfPropertyNames; ++i)
-	{
-		NewData.PropertyValues[i] = std::numeric_limits<float>::infinity();
-	}
-	PushSubjectFrameData_AnyThread(SubjectKey, MoveTemp(AnimationStruct));
-}
-
-void FLiveLinkClient_Base_DEPRECATED::ClearSubject(FName InSubjectName)
-{
-	bool bRemovedSubject = false;
-	{
-		TArray<FLiveLinkSubjectKey> AllSubjects = GetSubjects(false, true);
-		for (const FLiveLinkSubjectKey& SubjectKey : AllSubjects)
-		{
-			if (SubjectKey.SubjectName == InSubjectName)
-			{
-				RemoveSubject_AnyThread(SubjectKey);
-			}
-		}
-	}
-}
-
-void FLiveLinkClient_Base_DEPRECATED::GetSubjectNames(TArray<FName>& SubjectNames)
-{
-	TArray<FLiveLinkSubjectKey> SubjectKeys = GetSubjects(false, true);
-	SubjectNames.Reset(SubjectKeys.Num());
-
-	for (const FLiveLinkSubjectKey& SubjectKey : SubjectKeys)
-	{
-		SubjectNames.Add(SubjectKey.SubjectName);
-	}
-}
-
-const FLiveLinkSubjectFrame* FLiveLinkClient_Base_DEPRECATED::GetSubjectData(FName InSubjectName)
-{
-	//Old getters default to Animation role and copies data into old data structure
-	static const FName NAME_UpdateYourCode = "LiveLinkClient_GetSubjectData";
-	FLiveLinkLog::WarningOnce(NAME_UpdateYourCode, FLiveLinkSubjectKey(FGuid(), InSubjectName), TEXT("Upgrade your code. There is no way to deprecate GetSubjectData without creating new memory."));
-	return nullptr;
-}
-
-const FLiveLinkSubjectFrame* FLiveLinkClient_Base_DEPRECATED::GetSubjectDataAtWorldTime(FName InSubjectName, double InWorldTime)
-{
-	//Old getters default to Animation role and copies data into old data structure
-	static const FName NAME_UpdateYourCode = "LiveLinkClient_GetSubjectDataAtWorldTime";
-	FLiveLinkLog::WarningOnce(NAME_UpdateYourCode, FLiveLinkSubjectKey(FGuid(), InSubjectName), TEXT("Upgrade your code. There is no way to deprecate GetSubjectDataAtWorldTime without creating new memory."));
-	return nullptr;
-}
-
-const FLiveLinkSubjectFrame* FLiveLinkClient_Base_DEPRECATED::GetSubjectDataAtSceneTime(FName InSubjectName, const FTimecode& InTimecode)
-{
-	//Old getters default to Animation role and copies data into old data structure
-	static const FName NAME_UpdateYourCode = "LiveLinkClient_GetSubjectDataAtSceneTime";
-	FLiveLinkLog::WarningOnce(NAME_UpdateYourCode, FLiveLinkSubjectKey(FGuid(), InSubjectName), TEXT("Upgrade your code. There is no way to deprecate GetSubjectDataAtSceneTime without creating new memory."));
-	return nullptr;
-}
-
-bool FLiveLinkClient_Base_DEPRECATED::EvaluateFrameAtSceneTime_AnyThread(FLiveLinkSubjectName SubjectName, const FTimecode& SceneTime, TSubclassOf<ULiveLinkRole> DesiredRole, FLiveLinkSubjectFrameData& OutFrame)
-{
-	return static_cast<ILiveLinkClient*>(this)->EvaluateFrameAtSceneTime_AnyThread(SubjectName, FQualifiedFrameTime(SceneTime, FApp::GetTimecodeFrameRate()), DesiredRole, OutFrame);
-}
-
-const TArray<FLiveLinkFrame>* FLiveLinkClient_Base_DEPRECATED::GetSubjectRawFrames(FName InSubjectName)
-{
-	static const FName NAME_UpdateYourCode = "LiveLinkClient_GetSubjectRawFrames";
-	FLiveLinkLog::WarningOnce(NAME_UpdateYourCode, FLiveLinkSubjectKey(FGuid(), InSubjectName), TEXT("Upgrade your code. There is no way to deprecate GetSubjectRawFrames without creating new memory."));
-	return nullptr;
-}
-
-void FLiveLinkClient_Base_DEPRECATED::ClearSubjectsFrames(FName SubjectName)
-{
-	ClearSubjectsFrames_AnyThread(SubjectName);
-}
-
-void FLiveLinkClient_Base_DEPRECATED::ClearAllSubjectsFrames()
-{
-	ClearAllSubjectsFrames_AnyThread();
-}
-
-TSubclassOf<ULiveLinkRole> FLiveLinkClient_Base_DEPRECATED::GetSubjectRole(const FLiveLinkSubjectKey& SubjectKey) const
-{
-	return GetSubjectRole_AnyThread(SubjectKey);
-}
-
-TSubclassOf<ULiveLinkRole> FLiveLinkClient_Base_DEPRECATED::GetSubjectRole(FLiveLinkSubjectName SubjectName) const
-{
-	return GetSubjectRole_AnyThread(SubjectName);
-}
-
-bool FLiveLinkClient_Base_DEPRECATED::DoesSubjectSupportsRole(const FLiveLinkSubjectKey& SubjectKey, TSubclassOf<ULiveLinkRole> SupportedRole) const
-{
-	return DoesSubjectSupportsRole_AnyThread(SubjectKey, SupportedRole);
-}
-
-bool FLiveLinkClient_Base_DEPRECATED::DoesSubjectSupportsRole(FLiveLinkSubjectName SubjectName, TSubclassOf<ULiveLinkRole> SupportedRole) const
-{
-	return DoesSubjectSupportsRole_AnyThread(SubjectName, SupportedRole);
-}
-
-PRAGMA_ENABLE_DEPRECATION_WARNINGS
diff --git a/Engine/Plugins/Animation/LiveLink/Source/LiveLink/Private/LiveLinkComponent.cpp b/Engine/Plugins/Animation/LiveLink/Source/LiveLink/Private/LiveLinkComponent.cpp
--- a/Engine/Plugins/Animation/LiveLink/Source/LiveLink/Private/LiveLinkComponent.cpp
+++ b/Engine/Plugins/Animation/LiveLink/Source/LiveLink/Private/LiveLinkComponent.cpp
@@ -4,7 +4,6 @@
 #include "Components/SkeletalMeshComponent.h"
 #include "Features/IModularFeatures.h"
 #include "ILiveLinkClient.h"
-#include "Roles/LiveLinkAnimationRole.h"
 #include "Roles/LiveLinkAnimationBlueprintStructs.h"
 
 #include UE_INLINE_GENERATED_CPP_BY_NAME(LiveLinkComponent)
@@ -64,116 +63,3 @@ bool ULiveLinkComponent::HasLiveLinkClient()
 
 	return (LiveLinkClient != nullptr);
 }
-
-void ULiveLinkComponent::GetAvailableSubjectNames(TArray<FName>& SubjectNames)
-{
-	if (HasLiveLinkClient())
-	{
-		TArray<FLiveLinkSubjectKey> SubjectKeys = LiveLinkClient->GetSubjects(false, true);
-		SubjectNames.Reset(SubjectKeys.Num());
-		for (const FLiveLinkSubjectKey& SubjectKey : SubjectKeys)
-		{
-			SubjectNames.Add(SubjectKey.SubjectName);
-		}
-	}
-	else
-	{
-		SubjectNames.Reset();
-	}
-}
-
-void ULiveLinkComponent::GetSubjectData(const FName SubjectName, bool& bSuccess, FSubjectFrameHandle& SubjectFrameHandle)
-{
-	bSuccess = false;
-	if (HasLiveLinkClient())
-	{
-		FLiveLinkSubjectFrameData FrameData;
-		if (LiveLinkClient->EvaluateFrame_AnyThread(SubjectName, ULiveLinkAnimationRole::StaticClass(), FrameData))
-		{
-			const FLiveLinkSkeletonStaticData* SkeletonData = FrameData.StaticData.Cast<FLiveLinkSkeletonStaticData>();
-			const FLiveLinkAnimationFrameData* AnimationFrameData = FrameData.FrameData.Cast<FLiveLinkAnimationFrameData>();
-			FLiveLinkBlueprintDataStruct BlueprintDataWrapper(FSubjectFrameHandle::StaticStruct(), &SubjectFrameHandle);
-
-			if (SkeletonData == nullptr)
-			{
-				FFrame::KismetExecutionMessage(TEXT("Could not get subject data. Static data was invalid."), ELogVerbosity::Error);
-				return;
-			}
-
-			if (AnimationFrameData == nullptr)
-			{
-				FFrame::KismetExecutionMessage(TEXT("Could not get subject data. Frame data was invalid."), ELogVerbosity::Error);
-				return;
-			}
-
-			bSuccess = ULiveLinkAnimationRole::StaticClass()->GetDefaultObject<ULiveLinkAnimationRole>()->InitializeBlueprintData(FrameData, BlueprintDataWrapper);
-		}
-
-	}
-}
-
-void ULiveLinkComponent::GetSubjectDataAtWorldTime(const FName SubjectName, const float WorldTime, bool& bSuccess, FSubjectFrameHandle& SubjectFrameHandle)
-{
-PRAGMA_DISABLE_DEPRECATION_WARNINGS
-	GetSubjectDataAtTime(SubjectName, (double)WorldTime, bSuccess, SubjectFrameHandle);
-PRAGMA_ENABLE_DEPRECATION_WARNINGS
-}
-
-void ULiveLinkComponent::GetSubjectDataAtTime(const FName SubjectName, const double WorldTime, bool& bSuccess, FSubjectFrameHandle& SubjectFrameHandle)
-{
-	bSuccess = false;
-	if (HasLiveLinkClient())
-	{
-		FLiveLinkSubjectFrameData FrameData;
-		if (LiveLinkClient->EvaluateFrameAtWorldTime_AnyThread(SubjectName, WorldTime, ULiveLinkAnimationRole::StaticClass(), FrameData))
-		{
-			const FLiveLinkSkeletonStaticData* SkeletonData = FrameData.StaticData.Cast<FLiveLinkSkeletonStaticData>();
-			const FLiveLinkAnimationFrameData* AnimationFrameData = FrameData.FrameData.Cast<FLiveLinkAnimationFrameData>();
-			FLiveLinkBlueprintDataStruct BlueprintDataWrapper(FSubjectFrameHandle::StaticStruct(), &SubjectFrameHandle);
-
-			if (SkeletonData == nullptr)
-			{
-				FFrame::KismetExecutionMessage(TEXT("Could not get subject data. Static data was invalid."), ELogVerbosity::Error);
-				return;
-			}
-
-			if (AnimationFrameData == nullptr)
-			{
-				FFrame::KismetExecutionMessage(TEXT("Could not get subject data. Frame data was invalid."), ELogVerbosity::Error);
-				return;
-			}
-
-			bSuccess = ULiveLinkAnimationRole::StaticClass()->GetDefaultObject<ULiveLinkAnimationRole>()->InitializeBlueprintData(FrameData, BlueprintDataWrapper);
-		}
-	}
-}
-
-void ULiveLinkComponent::GetSubjectDataAtSceneTime(const FName SubjectName, const FTimecode& SceneTime, bool& bSuccess, FSubjectFrameHandle& SubjectFrameHandle)
-{
-	bSuccess = false;
-	if (HasLiveLinkClient())
-	{
-		FLiveLinkSubjectFrameData FrameData;
-		if (LiveLinkClient->EvaluateFrameAtSceneTime_AnyThread(SubjectName, FQualifiedFrameTime(SceneTime, FApp::GetTimecodeFrameRate()), ULiveLinkAnimationRole::StaticClass(), FrameData))
-		{
-			const FLiveLinkSkeletonStaticData* SkeletonData = FrameData.StaticData.Cast<FLiveLinkSkeletonStaticData>();
-			const FLiveLinkAnimationFrameData* AnimationFrameData = FrameData.FrameData.Cast<FLiveLinkAnimationFrameData>();
-			FLiveLinkBlueprintDataStruct BlueprintDataWrapper(FSubjectFrameHandle::StaticStruct(), &SubjectFrameHandle);
-
-			if (SkeletonData == nullptr)
-			{
-				FFrame::KismetExecutionMessage(TEXT("Could not get subject data. Static data was invalid."), ELogVerbosity::Error);
-				return;
-			}
-
-			if (AnimationFrameData == nullptr)
-			{
-				FFrame::KismetExecutionMessage(TEXT("Could not get subject data. Frame data was invalid."), ELogVerbosity::Error);
-				return;
-			}
-			
-			bSuccess = ULiveLinkAnimationRole::StaticClass()->GetDefaultObject<ULiveLinkAnimationRole>()->InitializeBlueprintData(FrameData, BlueprintDataWrapper);
-		}
-	}
-}
-
diff --git a/Engine/Plugins/Animation/LiveLink/Source/LiveLink/Private/LiveLinkDrivenComponent.cpp b/Engine/Plugins/Animation/LiveLink/Source/LiveLink/Private/LiveLinkDrivenComponent.cpp
deleted file mode 100644
diff --git a/Engine/Plugins/Animation/LiveLink/Source/LiveLink/Private/LiveLinkMessageBusSource.cpp b/Engine/Plugins/Animation/LiveLink/Source/LiveLink/Private/LiveLinkMessageBusSource.cpp
--- a/Engine/Plugins/Animation/LiveLink/Source/LiveLink/Private/LiveLinkMessageBusSource.cpp
+++ b/Engine/Plugins/Animation/LiveLink/Source/LiveLink/Private/LiveLinkMessageBusSource.cpp
@@ -357,34 +357,9 @@ void FLiveLinkMessageBusSource::InitializeMessageEndpoint(FMessageEndpointBuilde
 {
 PRAGMA_DISABLE_DEPRECATION_WARNINGS
 	MessageEndpoint = EndpointBuilder
-	.Handling<FLiveLinkSubjectDataMessage>(this, &FLiveLinkMessageBusSource::HandleSubjectData)
-	.Handling<FLiveLinkSubjectFrameMessage>(this, &FLiveLinkMessageBusSource::HandleSubjectFrame)
 	.Handling<FLiveLinkHeartbeatMessage>(this, &FLiveLinkMessageBusSource::HandleHeartbeat)
 	.Handling<FLiveLinkClearSubject>(this, &FLiveLinkMessageBusSource::HandleClearSubject)
 	.ReceivingOnAnyThread()
 	.WithCatchall(this, &FLiveLinkMessageBusSource::InternalHandleMessage);
 PRAGMA_ENABLE_DEPRECATION_WARNINGS
 }
-
-PRAGMA_DISABLE_DEPRECATION_WARNINGS
-void FLiveLinkMessageBusSource::HandleSubjectData(const FLiveLinkSubjectDataMessage& Message, const TSharedRef<IMessageContext, ESPMode::ThreadSafe>& Context)
-{
-	UpdateConnectionLastActive();
-
-	Client->PushSubjectSkeleton(SourceGuid, Message.SubjectName, Message.RefSkeleton);
-}
-PRAGMA_ENABLE_DEPRECATION_WARNINGS
-
-PRAGMA_DISABLE_DEPRECATION_WARNINGS
-void FLiveLinkMessageBusSource::HandleSubjectFrame(const FLiveLinkSubjectFrameMessage& Message, const TSharedRef<IMessageContext, ESPMode::ThreadSafe>& Context)
-{
-	UpdateConnectionLastActive();
-
-	FLiveLinkFrameData FrameData;
-	FrameData.Transforms = Message.Transforms;
-	FrameData.CurveElements = Message.Curves;
-	FrameData.MetaData = Message.MetaData;
-	FrameData.WorldTime = FLiveLinkWorldTime(Message.Time);
-	Client->PushSubjectData(SourceGuid, Message.SubjectName, FrameData);
-}
-PRAGMA_ENABLE_DEPRECATION_WARNINGS
diff --git a/Engine/Plugins/Animation/LiveLink/Source/LiveLink/Public/LiveLinkClient.h b/Engine/Plugins/Animation/LiveLink/Source/LiveLink/Public/LiveLinkClient.h
--- a/Engine/Plugins/Animation/LiveLink/Source/LiveLink/Public/LiveLinkClient.h
+++ b/Engine/Plugins/Animation/LiveLink/Source/LiveLink/Public/LiveLinkClient.h
@@ -26,41 +26,7 @@ struct FLiveLinkSubjectTimeSyncData
 	FFrameRate SampleFrameRate;
 };
 
-PRAGMA_DISABLE_DEPRECATION_WARNINGS
-struct FLiveLinkSkeletonStaticData;
-
-class LIVELINK_API FLiveLinkClient_Base_DEPRECATED : public ILiveLinkClient
-{
-public:
-	//~ Begin ILiveLinkClient implementation
-	virtual void PushSubjectSkeleton(FGuid SourceGuid, FName SubjectName, const FLiveLinkRefSkeleton& RefSkeleton) override;
-	virtual void PushSubjectData(FGuid SourceGuid, FName SubjectName, const FLiveLinkFrameData& FrameData) override;
-	virtual void ClearSubject(FName SubjectName) override;
-	virtual void GetSubjectNames(TArray<FName>& SubjectNames) override;
-	virtual const FLiveLinkSubjectFrame* GetSubjectData(FName InSubjectName) override;
-	virtual const FLiveLinkSubjectFrame* GetSubjectDataAtWorldTime(FName InSubjectName, double InWorldTime) override;
-	virtual const FLiveLinkSubjectFrame* GetSubjectDataAtSceneTime(FName InSubjectName, const FTimecode& InSceneTime) override;
-	virtual bool EvaluateFrameAtSceneTime_AnyThread(FLiveLinkSubjectName SubjectName, const FTimecode& SceneTime, TSubclassOf<ULiveLinkRole> DesiredRole, FLiveLinkSubjectFrameData& OutFrame) override;
-	virtual const TArray<FLiveLinkFrame>* GetSubjectRawFrames(FName SubjectName) override;
-	virtual void ClearSubjectsFrames(FName SubjectName) override;
-	virtual void ClearAllSubjectsFrames() override;
-	virtual TSubclassOf<ULiveLinkRole> GetSubjectRole(const FLiveLinkSubjectKey& SubjectKey) const override;
-	virtual TSubclassOf<ULiveLinkRole> GetSubjectRole(FLiveLinkSubjectName SubjectName) const override;
-	virtual bool DoesSubjectSupportsRole(const FLiveLinkSubjectKey& SubjectKey, TSubclassOf<ULiveLinkRole> SupportedRole) const override;
-	virtual bool DoesSubjectSupportsRole(FLiveLinkSubjectName SubjectName, TSubclassOf<ULiveLinkRole> SupportedRole) const override;
-
-	//~ End ILiveLinkClient implementation
-
-protected:
-	virtual void AquireLock_Deprecation() = 0;
-	virtual void ReleaseLock_Deprecation() = 0;
-	virtual void ClearFrames_Deprecation(const FLiveLinkSubjectKey& SubjectKey) = 0;
-	virtual FLiveLinkSkeletonStaticData* GetSubjectAnimationStaticData_Deprecation(const FLiveLinkSubjectKey& SubjectKey) = 0;
-};
-PRAGMA_ENABLE_DEPRECATION_WARNINGS
-
-
-class LIVELINK_API FLiveLinkClient : public FLiveLinkClient_Base_DEPRECATED
+class LIVELINK_API FLiveLinkClient : public ILiveLinkClient
 {
 public:
 	/** Default constructor that setups LiveLink to use the SamplingInput delegate to tick. */
@@ -206,32 +172,7 @@ public:
 	/** Push subject static data with additional metadata. */
 	void PushPendingSubject_AnyThread(FPendingSubjectStatic&& PendingSubject);
 
-	UE_DEPRECATED(4.23, "FLiveLinkClient::GetSourceTypeForEntry is deprecated. Please use GetSourceType instead!")
-	FText GetSourceTypeForEntry(FGuid EntryGuid) const { return GetSourceType(EntryGuid); }
-	UE_DEPRECATED(4.23, "FLiveLinkClient::GetMachineNameForEntry is deprecated. Please use GetSourceMachineName instead!")
-	FText GetMachineNameForEntry(FGuid EntryGuid) const { return GetSourceMachineName(EntryGuid); }
-	UE_DEPRECATED(4.23, "FLiveLinkClient::GetEntryStatusForEntry is deprecated. Please use GetSourceStatus instead!")
-	FText GetEntryStatusForEntry(FGuid EntryGuid) const { return GetSourceStatus(EntryGuid); }
-	UE_DEPRECATED(4.23, "FLiveLinkClient::GetSourceEntries is deprecated. Please use ILiveLinkClient::GetSources instead!")
-	const TArray<FGuid>& GetSourceEntries() const;
-	UE_DEPRECATED(4.23, "FLiveLinkClient::AddVirtualSubject(FName) is deprecated. Please use AddVirtualSubject(FName, TSubClass<ULiveLinkVirtualSubject>) instead!")
-	void AddVirtualSubject(FName NewVirtualSubjectName);
-	UE_DEPRECATED(4.25, "FLiveLinkClient::AddVirtualSubject(FLiveLinkSubjectName, TSubclassOf<ULiveLinkVirtualSubject>) is deprecated. Please use bool AddVirtualSubject(FLiveLinkSubjectKey, TSubClass<ULiveLinkVirtualSubject>) instead!")
-	void AddVirtualSubject(FLiveLinkSubjectName VirtualSubjectName, TSubclassOf<ULiveLinkVirtualSubject> VirtualSubjectClass);
-	UE_DEPRECATED(4.23, "FLiveLinkClient::OnLiveLinkSourcesChanged is deprecated. Please use OnLiveLinkSourceAdded instead!")
-	FDelegateHandle RegisterSourcesChangedHandle(const FSimpleMulticastDelegate::FDelegate& SourcesChanged);
-	UE_DEPRECATED(4.23, "FLiveLinkClient::OnLiveLinkSourcesChanged is deprecated. Please use OnLiveLinkSourceAdded instead!")
-	void UnregisterSourcesChangedHandle(FDelegateHandle Handle);
-
 protected:
-	//~ Begin FLiveLinkClient_Base_DEPRECATED implementation
-	using FLiveLinkClient_Base_DEPRECATED::EvaluateFrameAtSceneTime_AnyThread;
-	virtual void AquireLock_Deprecation() override;
-	virtual void ReleaseLock_Deprecation() override;
-	virtual void ClearFrames_Deprecation(const FLiveLinkSubjectKey& SubjectKey) override;
-	virtual FLiveLinkSkeletonStaticData* GetSubjectAnimationStaticData_Deprecation(const FLiveLinkSubjectKey& SubjectKey) override;
-	//~ End FLiveLinkClient_Base_DEPRECATED implementation
-
 	/**
 	 * Add delegates that will be triggered for all subjects.
 	 * @param InOnStaticDataAdded The delegate for when static data is added.
diff --git a/Engine/Plugins/Animation/LiveLink/Source/LiveLink/Public/LiveLinkComponent.h b/Engine/Plugins/Animation/LiveLink/Source/LiveLink/Public/LiveLinkComponent.h
--- a/Engine/Plugins/Animation/LiveLink/Source/LiveLink/Public/LiveLinkComponent.h
+++ b/Engine/Plugins/Animation/LiveLink/Source/LiveLink/Public/LiveLinkComponent.h
@@ -33,33 +33,6 @@ public:
 	// This Event is triggered any time new LiveLink data is available, including in the editor
 	UPROPERTY(BlueprintAssignable, Category = "LiveLink")
 	FLiveLinkTickSignature OnLiveLinkUpdated;
-
-	// Returns a list of available Subject Names for LiveLink
-	UE_DEPRECATED(4.23, "GetAvailableSubjectNames is deprecated, use GetLiveLinkEnabledSubjectNames.")
-	UFUNCTION(BlueprintCallable, Category = "LiveLink", meta=(DeprecatedFunction, DeprecationMessage="GetAvailableSubjectNames is deprecated, use GetLiveLinkEnabledSubjectNames."))
-	void GetAvailableSubjectNames(TArray<FName>& SubjectNames);
-	
-	// Returns a handle to the current frame of data in LiveLink for a given subject along with a boolean for whether a frame was found.
-	// Returns a handle to an empty frame if no frame of data is found.
-	UE_DEPRECATED(4.23, "GetSubjectData is deprecated, use EvaluateLiveLinkFrame.")
-	UFUNCTION(BlueprintCallable, Category = "LiveLink", meta=(DeprecatedFunction, DeprecationMessage="GetSubjectData is deprecated, EvaluateLiveLinkFrame."))
-	void GetSubjectData(const FName SubjectName, bool& bSuccess, FSubjectFrameHandle& SubjectFrameHandle);
-
-	// Returns a handle to the current frame of data in LiveLink for a given subject along with a boolean for whether a frame was found.
-	// Returns a handle to an empty frame if no frame of data is found.
-	UE_DEPRECATED(4.23, "GetSubjectDataAtWorldTime is deprecated, use EvaluateLiveLinkFrameAtWorldTime.")
-	UFUNCTION(BlueprintCallable, Category = "LiveLink", meta=(DeprecatedFunction, DeprecationMessage="GetSubjectDataAtWorldTime is deprecated, use EvaluateLiveLinkFrameAtWorldTime."))
-	void GetSubjectDataAtWorldTime(const FName SubjectName, const float WorldTime, bool& bSuccess, FSubjectFrameHandle& SubjectFrameHandle);
-	
-	UE_DEPRECATED(4.23, "GetSubjectDataAtTime is deprecated, use EvaluateLiveLinkFrameAtSceneTime.")
-	void GetSubjectDataAtTime(const FName SubjectName, const double WorldTime, bool& bSuccess, FSubjectFrameHandle& SubjectFrameHandle);
-
-	// Returns a handle to the frame of data in LiveLink for a given subject at the specified time along with a boolean for whether a frame was found.
-	// Returns a handle to an empty frame if no frame of data is found.
-	UE_DEPRECATED(4.23, "GetSubjectDataAtSceneTime is deprecated, use EvaluateLiveLinkFrameAtSceneTime.")
-	UFUNCTION(BlueprintCallable, Category = "LiveLink", meta=(DeprecatedFunction, DeprecationMessage="GetSubjectDataAtSceneTime is deprecated, use EvaluateLiveLinkFrameAtSceneTime."))
-	void GetSubjectDataAtSceneTime(const FName SubjectName, const FTimecode& SceneTime, bool& bSuccess, FSubjectFrameHandle& SubjectFrameHandle);
-
 private:
 	bool HasLiveLinkClient();
 
diff --git a/Engine/Plugins/Animation/LiveLink/Source/LiveLink/Public/LiveLinkDrivenComponent.h b/Engine/Plugins/Animation/LiveLink/Source/LiveLink/Public/LiveLinkDrivenComponent.h
--- a/Engine/Plugins/Animation/LiveLink/Source/LiveLink/Public/LiveLinkDrivenComponent.h
+++ b/Engine/Plugins/Animation/LiveLink/Source/LiveLink/Public/LiveLinkDrivenComponent.h
@@ -4,16 +4,15 @@
 
 #include "Components/ActorComponent.h"
 #include "LiveLinkTypes.h"
-#include "LiveLinkClientReference.h"
 #include "LiveLinkDrivenComponent.generated.h"
 
 /** A component that applies data from Live Link to the owning actor */
 UCLASS(ClassGroup = "LiveLink", deprecated)
 class LIVELINK_API UDEPRECATED_LiveLinkDrivenComponent : public UActorComponent
 {
-	GENERATED_UCLASS_BODY()
-
 public:
+	GENERATED_BODY()
+
 	// The name of the live link subject to take data from
 	UPROPERTY(EditAnywhere, Category = "Live Link", meta = (ShowOnlyInnerProperties))
 	FLiveLinkSubjectName SubjectName;
@@ -29,15 +28,4 @@ public:
 	// Should the transform from live link be treated as relative or world space
 	UPROPERTY(EditAnywhere, Category = "Live Link")
 	bool bSetRelativeLocation;
-
-	// UActorComponent interface
-	virtual void OnRegister() override;
-	virtual void OnUnregister() override;
-	virtual void TickComponent(float DeltaTime, enum ELevelTick TickType, FActorComponentTickFunction *ThisTickFunction) override;
-	// End UActorComponent interface
-
-private:
-
-	// Reference to the live link client so that we can get data about our subject
-	FLiveLinkClientReference ClientRef;
-};
\ No newline at end of file
+};
diff --git a/Engine/Plugins/Animation/LiveLink/Source/LiveLink/Public/LiveLinkMessageBusSource.h b/Engine/Plugins/Animation/LiveLink/Source/LiveLink/Public/LiveLinkMessageBusSource.h
--- a/Engine/Plugins/Animation/LiveLink/Source/LiveLink/Public/LiveLinkMessageBusSource.h
+++ b/Engine/Plugins/Animation/LiveLink/Source/LiveLink/Public/LiveLinkMessageBusSource.h
@@ -107,10 +107,6 @@ private:
 
 	//~ Message bus message handlers
 	PRAGMA_DISABLE_DEPRECATION_WARNINGS
-	void HandleSubjectData(const FLiveLinkSubjectDataMessage& Message, const TSharedRef<IMessageContext, ESPMode::ThreadSafe>& Context);
-	void HandleSubjectFrame(const FLiveLinkSubjectFrameMessage& Message, const TSharedRef<IMessageContext, ESPMode::ThreadSafe>& Context);
-	PRAGMA_ENABLE_DEPRECATION_WARNINGS
-
 	void HandleHeartbeat(const FLiveLinkHeartbeatMessage& Message, const TSharedRef<IMessageContext, ESPMode::ThreadSafe>& Context);
 	void HandleClearSubject(const FLiveLinkClearSubject& Message, const TSharedRef<IMessageContext, ESPMode::ThreadSafe>& Context);
 	void InternalHandleMessage(const TSharedRef<IMessageContext, ESPMode::ThreadSafe>& Context);
diff --git a/Engine/Plugins/Animation/LiveLink/Source/LiveLinkEditor/Private/LiveLinkClientPanelViews.cpp b/Engine/Plugins/Animation/LiveLink/Source/LiveLinkEditor/Private/LiveLinkClientPanelViews.cpp
--- a/Engine/Plugins/Animation/LiveLink/Source/LiveLinkEditor/Private/LiveLinkClientPanelViews.cpp
+++ b/Engine/Plugins/Animation/LiveLink/Source/LiveLinkEditor/Private/LiveLinkClientPanelViews.cpp
@@ -171,7 +171,7 @@ FText FLiveLinkSubjectUIEntry::GetItemText() const
 
 TSubclassOf<ULiveLinkRole> FLiveLinkSubjectUIEntry::GetItemRole() const
 {
-	return IsSubject() ? Client->GetSubjectRole(SubjectKey) : TSubclassOf<ULiveLinkRole>();
+	return IsSubject() ? Client->GetSubjectRole_AnyThread(SubjectKey) : TSubclassOf<ULiveLinkRole>();
 }
 
 void FLiveLinkSubjectUIEntry::RemoveFromClient() const
diff --git a/Engine/Plugins/Animation/LiveLink/Source/LiveLinkEditor/Private/SLiveLinkDataView.cpp b/Engine/Plugins/Animation/LiveLink/Source/LiveLinkEditor/Private/SLiveLinkDataView.cpp
--- a/Engine/Plugins/Animation/LiveLink/Source/LiveLinkEditor/Private/SLiveLinkDataView.cpp
+++ b/Engine/Plugins/Animation/LiveLink/Source/LiveLinkEditor/Private/SLiveLinkDataView.cpp
@@ -179,7 +179,7 @@ void SLiveLinkDataView::Tick(const FGeometry& AllottedGeometry, const double InC
 
 		if (Client->IsSubjectEnabled(SubjectKey, true))
 		{
-			TSubclassOf<ULiveLinkRole> SubjectRole = Client->GetSubjectRole(SubjectKey.SubjectName);
+			TSubclassOf<ULiveLinkRole> SubjectRole = Client->GetSubjectRole_AnyThread(SubjectKey.SubjectName);
 			if (SubjectRole)
 			{
 				FLiveLinkSubjectFrameData SubjectData;
diff --git a/Engine/Plugins/Animation/LiveLinkHub/Source/LiveLinkHub/Private/UI/Widgets/SLiveLinkTimecode.cpp b/Engine/Plugins/Animation/LiveLinkHub/Source/LiveLinkHub/Private/UI/Widgets/SLiveLinkTimecode.cpp
--- a/Engine/Plugins/Animation/LiveLinkHub/Source/LiveLinkHub/Private/UI/Widgets/SLiveLinkTimecode.cpp
+++ b/Engine/Plugins/Animation/LiveLinkHub/Source/LiveLinkHub/Private/UI/Widgets/SLiveLinkTimecode.cpp
@@ -523,17 +523,20 @@ void SLiveLinkTimecode::MakeTimecodeMenuItems(FMenuBuilder& MenuBuilder)
 			EUserInterfaceActionType::Check);
 	}
 
-	TArray<FName> Subjects;
-	WorkingClient->GetSubjectNames(Subjects);
-	for (FName Subject : Subjects)
+	constexpr bool bIncludeDisabledSubject = false;
+    constexpr bool bIncludeVirtualSubject = true;
+	TArray<FLiveLinkSubjectKey> Subjects = WorkingClient->GetSubjects(bIncludeDisabledSubject, bIncludeVirtualSubject);
+
+
+	for (const FLiveLinkSubjectKey& Subject : Subjects)
 	{
 		FTextBuilder SubjectText;
-		SubjectText.AppendLineFormat(LOCTEXT("LiveLinkHubTimecodeSourceSubject_Tooltip", "{0}'s timecode"), FText::FromName(Subject));
+		SubjectText.AppendLineFormat(LOCTEXT("LiveLinkHubTimecodeSourceSubject_Tooltip", "{0}'s timecode"), FText::FromName(Subject.SubjectName));
 		MenuBuilder.AddMenuEntry(
-			FText::FromName(Subject),
+			FText::FromName(Subject.SubjectName),
 			SubjectText.ToText(),
 			FSlateIcon(),
-			GenerateUIAction(Subject),
+			GenerateUIAction(Subject.SubjectName),
 			NAME_None,
 			EUserInterfaceActionType::Check);
 	}
@@ -584,23 +587,26 @@ void SLiveLinkTimecode::MakeCustomTimeStepMenuItems(FMenuBuilder& MenuBuilder)
 	FNewMenuDelegate SubMenuDelegate = FNewMenuDelegate::CreateLambda(
 		[this](FMenuBuilder& SubMenuBuilder)
 		{
-			TArray<FName> Subjects;
-			WorkingClient->GetSubjectNames(Subjects);
-			for (FName Subject : Subjects)
+			constexpr bool bIncludeDisabledSubject = false;
+			constexpr bool bIncludeVirtualSubject = true;
+
+			TArray<FLiveLinkSubjectKey> Subjects = WorkingClient->GetSubjects(bIncludeDisabledSubject, bIncludeVirtualSubject);
+
+			for (const FLiveLinkSubjectKey& Subject : Subjects)
 			{
 				FTextBuilder SubjectText;
 				SubjectText.AppendLineFormat(
 					LOCTEXT("LiveLinkHubCustomTimeStepSourceSubject_Tooltip", "Use {0} as a frame lock source."),
-					FText::FromName(Subject)
+					FText::FromName(Subject.SubjectName)
 				);
 
 				FUIAction Action = FUIAction(
-					FExecuteAction::CreateSP(this, &SLiveLinkTimecode::SetCustomTimeStepSubject, Subject),
+					FExecuteAction::CreateSP(this, &SLiveLinkTimecode::SetCustomTimeStepSubject, Subject.SubjectName.Name),
 					FCanExecuteAction::CreateLambda([] { return true; }),
-					FIsActionChecked::CreateLambda([this, Subject] { return Subject == ActiveSubjectCustomTimeStep; }));
+					FIsActionChecked::CreateLambda([this, Subject] { return Subject.SubjectName == ActiveSubjectCustomTimeStep; }));
 
 				SubMenuBuilder.AddMenuEntry(
-					FText::FromName(Subject),
+					FText::FromName(Subject.SubjectName),
 					SubjectText.ToText(),
 					FSlateIcon(),
 					Action,
diff --git a/Engine/Source/Runtime/LiveLinkInterface/Public/ILiveLinkClient.h b/Engine/Source/Runtime/LiveLinkInterface/Public/ILiveLinkClient.h
--- a/Engine/Source/Runtime/LiveLinkInterface/Public/ILiveLinkClient.h
+++ b/Engine/Source/Runtime/LiveLinkInterface/Public/ILiveLinkClient.h
@@ -46,93 +46,13 @@ DECLARE_TS_MULTICAST_DELEGATE_ThreeParams(FOnLiveLinkSubjectFrameDataAdded, FLiv
 DECLARE_TS_MULTICAST_DELEGATE_FiveParams(FOnLiveLinkSubjectEvaluated, FLiveLinkSubjectKey /*InSubjectKey*/, TSubclassOf<ULiveLinkRole> /*RequestedRole*/, const FLiveLinkTime& /*RequestedTime*/, bool /*bResult*/, const FLiveLinkTime& /*EvaluatedFrameTime*/)
 
 
-PRAGMA_DISABLE_DEPRECATION_WARNINGS
-class ILiveLinkClient_Base_DEPRECATED : public IModularFeature
-{
-public:
-	UE_DEPRECATED(4.23, "ILiveLinkClient::PushSubjectSkeleton is deprecated. Please use ILiveLinkClient::PushSubjectStaticData with the proper LiveLinkRole's static data structure instead!")
-	virtual void PushSubjectSkeleton(FGuid SourceGuid, FName SubjectName, const FLiveLinkRefSkeleton& RefSkeleton) {}
-
-	UE_DEPRECATED(4.23, "ILiveLinkClient::PushSubjectData is deprecated. Please use ILiveLinkClient::PushSubjectFrameData with the proper LiveLinkRole's frame data structure instead!")
-	virtual void PushSubjectData(FGuid SourceGuid, FName SubjectName, const FLiveLinkFrameData& FrameData) {}
-
-	UE_DEPRECATED(4.23, "ILiveLinkClient::ClearSubject is deprecated. Please use ILiveLinkClient::RemoveSubject instead!")
-	virtual void ClearSubject(FName SubjectName) {}
-
-	UE_DEPRECATED(4.23, "ILiveLinkClient::GetSubjectNames is deprecated. Please use ILiveLinkClient::GetSubjects instead!")
-	virtual void GetSubjectNames(TArray<FName>& SubjectNames) {}
-
-	UE_DEPRECATED(4.23, "ILiveLinkClient::GetActiveSubjects is deprecated. Please use ILiveLinkClient::GetSubjects instead!")
-	virtual TArray<FLiveLinkSubjectKey> GetActiveSubjects() { return TArray<FLiveLinkSubjectKey>(); }
-
-	UE_DEPRECATED(4.23, "ILiveLinkClient::GetSubjectData is deprecated. Please use ILiveLinkClient::EvaluateFrame with the desired LiveLinkRole instead!")
-	virtual const FLiveLinkSubjectFrame* GetSubjectData(FName SubjectName) { return nullptr; }
-
-	UE_DEPRECATED(4.23, "ILiveLinkClient::GetSubjectDataAtWorldTime is deprecated. Please use ILiveLinkClient::EvaluateFrameAtWorldTime with the desired LiveLinkRole instead!")
-	virtual const FLiveLinkSubjectFrame* GetSubjectDataAtWorldTime(FName SubjectName, double WorldTime) { return nullptr; }
-
-	UE_DEPRECATED(4.23, "ILiveLinkClient::GetSubjectDataAtSceneTime is deprecated. Please use ILiveLinkClient::EvaluateFrameAtSceneTime with the desired LiveLinkRole instead!")
-	virtual const FLiveLinkSubjectFrame* GetSubjectDataAtSceneTime(FName SubjectName, const FTimecode& SceneTime) { return nullptr; }
-
-	UE_DEPRECATED(4.23, "ILiveLinkClient::GetSubjectRawFrames is deprecated. It is no longer supported.")
-	virtual const TArray<FLiveLinkFrame>* GetSubjectRawFrames(FName SubjectName) { return nullptr; }
-
-	UE_DEPRECATED(4.23, "ILiveLinkClient::ClearSubjectsFrames is deprecated. Please use ILiveLinkClient::ClearSubjectsFrames_AnyThread with the desired LiveLinkRole instead!")
-	virtual void ClearSubjectsFrames(FName SubjectName) {}
-
-	UE_DEPRECATED(4.23, "ILiveLinkClient::ClearAllSubjectsFrames is deprecated. Please use ILiveLinkClient::ClearAllSubjectsFrames_AnyThread with the desired LiveLinkRole instead!")
-	virtual void ClearAllSubjectsFrames() {}
-
-	UE_DEPRECATED(4.23, "ILiveLinkClient::GetSaveFrames is deprecated. Please register using RegisterForSubjectFrames to start receiving subject frames instead!")
-	virtual bool GetSaveFrames() const { return false; }
-
-	UE_DEPRECATED(4.23, "ILiveLinkClient::SetSaveFrames is deprecated. Please register using RegisterForSubjectFrames to start receiving subject frames instead!")
-	virtual bool SetSaveFrames(bool InSave) { return false; }
-
-	UE_DEPRECATED(4.23, "ILiveLinkClient::StartRecordingLiveLink is deprecated. Please register using RegisterForSubjectFrames to start receiving subject frames instead!")
-	virtual FGuid StartRecordingLiveLink(const TArray<FName>& SubjectNames) { return FGuid(); }
-
-	UE_DEPRECATED(4.23, "ILiveLinkClient::StartRecordingLiveLink is deprecated. Please register using RegisterForSubjectFrames to start receiving subject frames instead!")
-	virtual FGuid StartRecordingLiveLink(const FName& SubjectName) { return FGuid(); }
-
-	UE_DEPRECATED(4.23, "ILiveLinkClient::StopRecordingLiveLinkData is deprecated. Please unregister using UnregisterSubjectFramesHandle to stop receiving subject frames instead!")
-	virtual void StopRecordingLiveLinkData(const FGuid &InGuid, const TArray<FName>& SubjectNames) {}
-
-	UE_DEPRECATED(4.23, "ILiveLinkClient::StopRecordingLiveLinkData is deprecated. Please unregister using UnregisterSubjectFramesHandle to stop receiving subject frames instead!")
-	virtual void StopRecordingLiveLinkData(const FGuid &InGuid, const FName& SubjectName) {}
-
-	UE_DEPRECATED(4.23, "ILiveLinkClient::GetAndFreeLastRecordedFrames is deprecated. Please register using RegisterForSubjectFrames to receive subject frames instead!")
-	virtual void GetAndFreeLastRecordedFrames(const FGuid& InHandlerGuid, FName SubjectName, TArray<FLiveLinkFrame> &OutFrames) {}
-
-	UE_DEPRECATED(4.23, "ILiveLinkClient::RegisterSubjectsChangedHandle is deprecated. Please register using OnLiveLinkSubjectAdded and OnLiveLinkSubjectRemoved instead!")
-	virtual FDelegateHandle RegisterSubjectsChangedHandle(const FSimpleMulticastDelegate::FDelegate& SubjectsChanged) { return FDelegateHandle(); }
-
-	UE_DEPRECATED(4.23, "ILiveLinkClient::UnregisterSubjectsChangedHandle is deprecated. Please unregister using OnLiveLinkSubjectAdded and OnLiveLinkSubjectRemoved instead!")
-	virtual void UnregisterSubjectsChangedHandle(FDelegateHandle Handle) {}
-
-	UE_DEPRECATED(5.0, "ILiveLinkClient::GetSubjectRole is deprecated. Please use GetSubjectRole_AnyThread instead!")
-	virtual TSubclassOf<ULiveLinkRole> GetSubjectRole(const FLiveLinkSubjectKey& SubjectKey) const = 0;
-
-	UE_DEPRECATED(5.0, "ILiveLinkClient::GetSubjectRole is deprecated. Please use GetSubjectRole_AnyThread instead!")
-	virtual TSubclassOf<ULiveLinkRole> GetSubjectRole(FLiveLinkSubjectName SubjectName) const = 0;
-
-	UE_DEPRECATED(5.0, "ILiveLinkClient::DoesSubjectSupportsRole is deprecated. Please use DoesSubjectSupportsRole_AnyThread instead!")
-	virtual bool DoesSubjectSupportsRole(const FLiveLinkSubjectKey& SubjectKey, TSubclassOf<ULiveLinkRole> SupportedRole) const = 0;
-
-	UE_DEPRECATED(5.0, "ILiveLinkClient::DoesSubjectSupportsRole is deprecated. Please use DoesSubjectSupportsRole_AnyThread instead!")
-	virtual bool DoesSubjectSupportsRole(FLiveLinkSubjectName SubjectName, TSubclassOf<ULiveLinkRole> SupportedRole) const = 0;
-
-};
-PRAGMA_ENABLE_DEPRECATION_WARNINGS
-
-
 /**
  * Interface for streaming and consuming data from external sources into UE4.
  * A LiveLinkSource, may stream multiple LiveLinkSubject.
  * Pushing and evaluating data can be executed on any thread. The other functions must be executed on the Game Thread.
  * Subject may shared name between sources, but only 1 of those subjects may be enabled.
  */
-class LIVELINKINTERFACE_API ILiveLinkClient : public ILiveLinkClient_Base_DEPRECATED
+class LIVELINKINTERFACE_API ILiveLinkClient : public IModularFeature
 {
 public:
 	static const int32 LIVELINK_VERSION = 2;
