Fix CIS warning caused by mismatched deprecation macro

diff --git a/Engine/Plugins/Animation/LiveLink/Source/LiveLink/Public/LiveLinkMessageBusSource.h b/Engine/Plugins/Animation/LiveLink/Source/LiveLink/Public/LiveLinkMessageBusSource.h
--- a/Engine/Plugins/Animation/LiveLink/Source/LiveLink/Public/LiveLinkMessageBusSource.h
+++ b/Engine/Plugins/Animation/LiveLink/Source/LiveLink/Public/LiveLinkMessageBusSource.h
@@ -106,7 +106,6 @@ private:
 	TSharedPtr<FMessageEndpoint, ESPMode::ThreadSafe> CreateAndInitializeMessageEndpoint();
 
 	//~ Message bus message handlers
-	PRAGMA_DISABLE_DEPRECATION_WARNINGS
 	void HandleHeartbeat(const FLiveLinkHeartbeatMessage& Message, const TSharedRef<IMessageContext, ESPMode::ThreadSafe>& Context);
 	void HandleClearSubject(const FLiveLinkClearSubject& Message, const TSharedRef<IMessageContext, ESPMode::ThreadSafe>& Context);
 	void InternalHandleMessage(const TSharedRef<IMessageContext, ESPMode::ThreadSafe>& Context);
