Extend Live Link Pre Processors interface so that const static data is provided for each frame

diff --git a/Engine/Plugins/Animation/LiveLink/Source/LiveLink/Private/LiveLinkClient.cpp b/Engine/Plugins/Animation/LiveLink/Source/LiveLink/Private/LiveLinkClient.cpp
--- a/Engine/Plugins/Animation/LiveLink/Source/LiveLink/Private/LiveLinkClient.cpp
+++ b/Engine/Plugins/Animation/LiveLink/Source/LiveLink/Private/LiveLinkClient.cpp
@@ -291,15 +291,15 @@ void FLiveLinkClient::HandleSubjectRebroadcast(ILiveLinkSubject* InSubject, cons
 				// Make a copy of the data for use by the rebroadcaster
 				FLiveLinkFrameDataStruct FrameDataCopy;
 				FrameDataCopy.InitializeWith(InFrameData);
-
+				
+				FLiveLinkStaticDataStruct StaticDataCopy;
+				StaticDataCopy.InitializeWith(InSubject->GetStaticData());
+				
 				if (bPreProcessRebroadcastFrames)
 				{
-					InSubject->PreprocessFrame(FrameDataCopy);
+					InSubject->PreprocessFrame(StaticDataCopy, FrameDataCopy);
 				}
 
-				FLiveLinkStaticDataStruct StaticDataCopy;
-	            StaticDataCopy.InitializeWith(InSubject->GetStaticData());
-
 				TSubclassOf<ULiveLinkRole> SubjectRole = InSubject->GetRole();
 
 				if (bTranslateRebroadcastFrames)
diff --git a/Engine/Plugins/Animation/LiveLink/Source/LiveLink/Private/LiveLinkSubject.cpp b/Engine/Plugins/Animation/LiveLink/Source/LiveLink/Private/LiveLinkSubject.cpp
--- a/Engine/Plugins/Animation/LiveLink/Source/LiveLink/Private/LiveLinkSubject.cpp
+++ b/Engine/Plugins/Animation/LiveLink/Source/LiveLink/Private/LiveLinkSubject.cpp
@@ -431,20 +431,20 @@ void FLiveLinkSubject::AddFrameData(FLiveLinkFrameDataStruct&& InFrameData)
 		break;
 	}
 
-	//Make sure frame should be inserted. 
-	if(FrameIndex >= 0)
+	// Make sure frame should be inserted. 
+	if (FrameIndex >= 0)
 	{
 		for (ULiveLinkFramePreProcessor::FWorkerSharedPtr PreProcessor : FramePreProcessors)
 		{
-			PreProcessor->PreProcessFrame(InFrameData);
+			PreProcessor->PreProcessFrame(GetStaticData(), InFrameData);
 		}
 
-		//Assign identifier to incoming frame and insert it where it belongs
+		// Assign identifier to incoming frame and insert it where it belongs
 		const FLiveLinkFrameIdentifier ThisFrameIdentifier = NextIdentifier++;
 		ReceivedOrderedFrames.Enqueue(ThisFrameIdentifier);
 		InFrameData.GetBaseData()->FrameId = ThisFrameIdentifier;
 		
-		//Extra verbose logging about new frame being added
+		// Extra verbose logging about new frame being added
 		const FTimecode FrameTC = FTimecode::FromFrameNumber(InFrameData.GetBaseData()->MetaData.SceneTime.Time.GetFrame(), InFrameData.GetBaseData()->MetaData.SceneTime.Rate);
 		UE_LOG(LogLiveLink, Verbose, TEXT("Subject '%s' adding frame ID %d at index %d - Timecode:[%s.%0.3f] - SourceTime: %0.4f, Offset: %0.6f, CorrectedTime: %0.4f"), *SubjectKey.SubjectName.ToString(), ThisFrameIdentifier, FrameIndex, *FrameTC.ToString(), InFrameData.GetBaseData()->MetaData.SceneTime.Time.GetSubFrame(), InFrameData.GetBaseData()->WorldTime.GetSourceTime(), InFrameData.GetBaseData()->WorldTime.GetOffset(), InFrameData.GetBaseData()->WorldTime.GetOffsettedTime());
 		
@@ -472,6 +472,14 @@ void FLiveLinkSubject::PreprocessFrame(FLiveLinkFrameDataStruct& InOutFrameData)
 	}
 }
 
+void FLiveLinkSubject::PreprocessFrame(const FLiveLinkStaticDataStruct& InStaticData, FLiveLinkFrameDataStruct& InOutFrameData)
+{
+	for (ULiveLinkFramePreProcessor::FWorkerSharedPtr PreProcessor : FramePreProcessors)
+	{
+		PreProcessor->PreProcessFrame(InStaticData, InOutFrameData);
+	}
+}
+
 bool FLiveLinkSubject::ValidateFrameData(const FLiveLinkFrameDataStruct& InFrameData)
 {
 	if (!StaticData.IsValid())
diff --git a/Engine/Plugins/Animation/LiveLink/Source/LiveLink/Public/LiveLinkSubject.h b/Engine/Plugins/Animation/LiveLink/Source/LiveLink/Public/LiveLinkSubject.h
--- a/Engine/Plugins/Animation/LiveLink/Source/LiveLink/Public/LiveLinkSubject.h
+++ b/Engine/Plugins/Animation/LiveLink/Source/LiveLink/Public/LiveLinkSubject.h
@@ -77,7 +77,11 @@ public:
 	virtual bool IsRebroadcasted() const override { return bRebroadcastSubject; }
 	virtual bool HasStaticDataBeenRebroadcasted() const override { return bRebroadcastStaticDataSent; }
 	virtual void SetStaticDataAsRebroadcasted(const bool bInSent) override { bRebroadcastStaticDataSent = bInSent; }
+
+	UE_DEPRECATED(5.6, "Replaced with a new version of PreprocessFrame that also provides const static data.")
 	virtual void PreprocessFrame(FLiveLinkFrameDataStruct& InOutFrameData) override;
+	
+	virtual void PreprocessFrame(const FLiveLinkStaticDataStruct &InStaticData, FLiveLinkFrameDataStruct& InOutFrameData) override;
 
 protected:
 	virtual const FLiveLinkSubjectFrameData& GetFrameSnapshot() const override { return FrameSnapshot; }
diff --git a/Engine/Source/Runtime/LiveLinkInterface/Public/ILiveLinkSubject.h b/Engine/Source/Runtime/LiveLinkInterface/Public/ILiveLinkSubject.h
--- a/Engine/Source/Runtime/LiveLinkInterface/Public/ILiveLinkSubject.h
+++ b/Engine/Source/Runtime/LiveLinkInterface/Public/ILiveLinkSubject.h
@@ -56,9 +56,13 @@ public:
 	/** Sets the static data for this subject as rebroadcasted */
 	virtual void SetStaticDataAsRebroadcasted(const bool bInSent) = 0;
 
-	/** Apply this subject's preprocessors to a frame data. */
+	/** Apply this subject's preprocessors to frame data. */
+	UE_DEPRECATED(5.6, "Replaced with a new version of PreprocessFrame that also provides const static data.")
 	virtual void PreprocessFrame(FLiveLinkFrameDataStruct& InOutFrameData) {};
 
+	/** Apply this subject's preprocessors to frame data. Provides static data as a const reference.*/
+	virtual void PreprocessFrame(const FLiveLinkStaticDataStruct& InStaticData, FLiveLinkFrameDataStruct& InOutFrameData) {};
+
 	/** Apply a remapper to a frame data. Called after preprocessing.*/
 	virtual void RemapFrame(FLiveLinkSkeletonStaticData& InOutSkeletonData, FLiveLinkAnimationFrameData& InOutFrameData) {}
 
diff --git a/Engine/Source/Runtime/LiveLinkInterface/Public/LiveLinkFramePreProcessor.h b/Engine/Source/Runtime/LiveLinkInterface/Public/LiveLinkFramePreProcessor.h
--- a/Engine/Source/Runtime/LiveLinkInterface/Public/LiveLinkFramePreProcessor.h
+++ b/Engine/Source/Runtime/LiveLinkInterface/Public/LiveLinkFramePreProcessor.h
@@ -21,8 +21,14 @@
 class ILiveLinkFramePreProcessorWorker
 {
 public:
+	virtual ~ILiveLinkFramePreProcessorWorker() = default;
 	virtual TSubclassOf<ULiveLinkRole> GetRole() const = 0;
-	virtual bool PreProcessFrame(FLiveLinkFrameDataStruct& InOutFrame) const = 0;
+
+	/** Implement this method if you only require the frame data. **/
+	virtual bool PreProcessFrame(FLiveLinkFrameDataStruct& InOutFrame) const { return true; }
+
+	/** Implement this method if you require access to static and frame data. Note that the static data is a const reference. **/
+	virtual bool PreProcessFrame(const FLiveLinkStaticDataStruct& InStaticData, FLiveLinkFrameDataStruct& InOutFrame) { return PreProcessFrame(InOutFrame); }
 };
 
 
