Add GetDisplayName for engine custom time step to allow giving more details to the user when fps stats are enabled.

diff --git a/Engine/Plugins/Animation/LiveLink/Source/LiveLink/Private/LiveLinkCustomTimeStep.cpp b/Engine/Plugins/Animation/LiveLink/Source/LiveLink/Private/LiveLinkCustomTimeStep.cpp
--- a/Engine/Plugins/Animation/LiveLink/Source/LiveLink/Private/LiveLinkCustomTimeStep.cpp
+++ b/Engine/Plugins/Animation/LiveLink/Source/LiveLink/Private/LiveLinkCustomTimeStep.cpp
@@ -7,7 +7,10 @@
 #include "ILiveLinkClient.h"
 #include "Misc/App.h"
 
-
+namespace LiveLinkCustomTimeStepUtils
+{
+	static const FString NoSubjectSelected = TEXT("<NO_SUBJECT>");
+}
 
 bool ULiveLinkCustomTimeStep::Initialize(UEngine* InEngine)
 {
@@ -103,6 +106,16 @@ FFrameRate ULiveLinkCustomTimeStep::GetSyncRate() const
 	return LiveLinkDataRate;
 }
 
+FString ULiveLinkCustomTimeStep::GetDisplayName() const
+{
+	const FString PlainName = GetFName().GetPlainNameString();
+	const FText FrameRate = GetFixedFrameRate().ToPrettyText();
+	const int32 NameIndex = GetFName().GetNumber();
+	const FString& SubjectName = !SubjectKey.SubjectName.IsNone() ? SubjectKey.SubjectName.ToString() : LiveLinkCustomTimeStepUtils::NoSubjectSelected;
+
+	return FString::Printf(TEXT("%s %s - %s (%d)"), *PlainName, *FrameRate.ToString(), *SubjectName, NameIndex);
+}
+
 uint32 ULiveLinkCustomTimeStep::GetLastSyncCountDelta() const
 {
 	return LastSyncCountDelta;
diff --git a/Engine/Plugins/Animation/LiveLink/Source/LiveLink/Public/LiveLinkCustomTimeStep.h b/Engine/Plugins/Animation/LiveLink/Source/LiveLink/Public/LiveLinkCustomTimeStep.h
--- a/Engine/Plugins/Animation/LiveLink/Source/LiveLink/Public/LiveLinkCustomTimeStep.h
+++ b/Engine/Plugins/Animation/LiveLink/Source/LiveLink/Public/LiveLinkCustomTimeStep.h
@@ -45,6 +45,7 @@ public:
 	LIVELINK_API virtual ECustomTimeStepSynchronizationState GetSynchronizationState() const override;
 	LIVELINK_API virtual FFrameRate GetFixedFrameRate() const override;
 	LIVELINK_API virtual FFrameRate GetSyncRate() const override;
+	LIVELINK_API virtual FString GetDisplayName() const override;
 
 	//~ UGenlockedCustomTimeStep interface
 	LIVELINK_API virtual uint32 GetLastSyncCountDelta() const override;
diff --git a/Engine/Plugins/Animation/LiveLinkHub/Source/LiveLinkHubMessaging/Public/LiveLinkHubMessages.h b/Engine/Plugins/Animation/LiveLinkHub/Source/LiveLinkHubMessaging/Public/LiveLinkHubMessages.h
--- a/Engine/Plugins/Animation/LiveLinkHub/Source/LiveLinkHubMessaging/Public/LiveLinkHubMessages.h
+++ b/Engine/Plugins/Animation/LiveLinkHub/Source/LiveLinkHubMessaging/Public/LiveLinkHubMessages.h
@@ -54,7 +54,7 @@ struct LIVELINKHUBMESSAGING_API FLiveLinkHubCustomTimeStepSettings
 
 	/** Desired frame rate to lock the editor. This corresponds to the LiveLinkDataRate in ULiveLinkCustomTimeStep */
 	UPROPERTY()
-	FFrameRate CustomTimeStepRate;
+	FFrameRate CustomTimeStepRate = FFrameRate(60, 1);
 
 	/** Assign the frame lock settings to the engine. */
 	void AssignCustomTimeStepToEngine() const;
diff --git a/Engine/Source/Runtime/Engine/Classes/Engine/EngineCustomTimeStep.h b/Engine/Source/Runtime/Engine/Classes/Engine/EngineCustomTimeStep.h
--- a/Engine/Source/Runtime/Engine/Classes/Engine/EngineCustomTimeStep.h
+++ b/Engine/Source/Runtime/Engine/Classes/Engine/EngineCustomTimeStep.h
@@ -52,6 +52,9 @@ public:
 	/** The state of the CustomTimeStep. */
 	ENGINE_API virtual ECustomTimeStepSynchronizationState GetSynchronizationState() const PURE_VIRTUAL(UEngineCustomTimeStep::GetSynchronizationState, return ECustomTimeStepSynchronizationState::Closed;);
 
+	/** Get the display name of the custom time step. Allows child classes to provide more context to users. */
+	ENGINE_API virtual FString GetDisplayName() const { return GetName(); };
+
 public:
 	/** Default behaviour of the engine. Update FApp::LastTime */
 	static ENGINE_API void UpdateApplicationLastTime();
diff --git a/Engine/Source/Runtime/Engine/Private/UnrealEngine.cpp b/Engine/Source/Runtime/Engine/Private/UnrealEngine.cpp
--- a/Engine/Source/Runtime/Engine/Private/UnrealEngine.cpp
+++ b/Engine/Source/Runtime/Engine/Private/UnrealEngine.cpp
@@ -18533,7 +18533,7 @@ int32 UEngine::RenderStatFPS(UWorld* World, FViewport* Viewport, FCanvas* Canvas
 	if (CustomTimeStep)
 	{
 		ECustomTimeStepSynchronizationState State = CustomTimeStep->GetSynchronizationState();
-		FString CustomTimeStepName = CustomTimeStep->GetName();
+		FString CustomTimeStepName = CustomTimeStep->GetDisplayName();
 		int32 NewX = X - Font->GetStringSize(*CustomTimeStepName);
 		switch (State)
 		{
diff --git a/Engine/Source/Runtime/TimeManagement/Private/FixedFrameRateCustomTimeStep.cpp b/Engine/Source/Runtime/TimeManagement/Private/FixedFrameRateCustomTimeStep.cpp
--- a/Engine/Source/Runtime/TimeManagement/Private/FixedFrameRateCustomTimeStep.cpp
+++ b/Engine/Source/Runtime/TimeManagement/Private/FixedFrameRateCustomTimeStep.cpp
@@ -9,12 +9,20 @@
 
 #include UE_INLINE_GENERATED_CPP_BY_NAME(FixedFrameRateCustomTimeStep)
 
-
 UFixedFrameRateCustomTimeStep::UFixedFrameRateCustomTimeStep(const FObjectInitializer& ObjectInitializer)
 	: Super(ObjectInitializer)
 {
 }
 
+FString UFixedFrameRateCustomTimeStep::GetDisplayName() const
+{
+	const FString PlainName = GetFName().GetPlainNameString();
+	const FText FrameRate = GetFixedFrameRate().ToPrettyText();
+	const int32 NameIndex = GetFName().GetNumber();
+
+	return FString::Printf(TEXT("%s %s (%d)"), *PlainName, *FrameRate.ToString(), NameIndex);
+}
+
 void UFixedFrameRateCustomTimeStep::WaitForFixedFrameRate() const
 {
 	UpdateApplicationLastTime();
@@ -54,4 +62,3 @@ FFrameRate UFixedFrameRateCustomTimeStep::GetFixedFrameRate_PureVirtual() const
 {
 	return FFrameRate(24, 1);
 }
-
diff --git a/Engine/Source/Runtime/TimeManagement/Public/FixedFrameRateCustomTimeStep.h b/Engine/Source/Runtime/TimeManagement/Public/FixedFrameRateCustomTimeStep.h
--- a/Engine/Source/Runtime/TimeManagement/Public/FixedFrameRateCustomTimeStep.h
+++ b/Engine/Source/Runtime/TimeManagement/Public/FixedFrameRateCustomTimeStep.h
@@ -24,6 +24,10 @@ class UFixedFrameRateCustomTimeStep : public UEngineCustomTimeStep
 	GENERATED_UCLASS_BODY()
 
 public:
+	//~ Begin UEngineCustomTimeStep interface
+	TIMEMANAGEMENT_API virtual FString GetDisplayName() const override;
+	//~ End UEngineCustomTimeStep
+
 	/** Get The fixed FrameRate */
 	TIMEMANAGEMENT_API virtual FFrameRate GetFixedFrameRate() const PURE_VIRTUAL(UFixedFrameRateCustomTimeStep::GetFixedFrameRate, return GetFixedFrameRate_PureVirtual(););
 
