LiveLink Locator Role

-Adds locators role primarily for optical motion capture  markers.
-Provides option to send unstructured data (changing number of markers per frame with no labels)
-Axis swap pre-processor to operate on this new type of data.

diff --git a/Engine/Plugins/Animation/LiveLink/Source/LiveLink/Private/PreProcessor/LiveLinkAxisSwitchPreProcessor.cpp b/Engine/Plugins/Animation/LiveLink/Source/LiveLink/Private/PreProcessor/LiveLinkAxisSwitchPreProcessor.cpp
--- a/Engine/Plugins/Animation/LiveLink/Source/LiveLink/Private/PreProcessor/LiveLinkAxisSwitchPreProcessor.cpp
+++ b/Engine/Plugins/Animation/LiveLink/Source/LiveLink/Private/PreProcessor/LiveLinkAxisSwitchPreProcessor.cpp
@@ -1,14 +1,18 @@
 // Copyright Epic Games, Inc. All Rights Reserved.
 
 #include "PreProcessor/LiveLinkAxisSwitchPreProcessor.h"
+
 #include "Roles/LiveLinkAnimationRole.h"
 #include "Roles/LiveLinkAnimationTypes.h"
 #include "Roles/LiveLinkTransformRole.h"
 #include "Roles/LiveLinkTransformTypes.h"
+#include "Roles/LiveLinkLocatorRole.h"
+#include "Roles/LiveLinkLocatorTypes.h"
 #include "UObject/ReleaseObjectVersion.h"
 
 #include UE_INLINE_GENERATED_CPP_BY_NAME(LiveLinkAxisSwitchPreProcessor)
 
+
 namespace LiveLinkAxisSwitchPreProcessor
 {
 	float LiveLinkAxisToVectorMember(ELiveLinkAxis InAxis, const FVector& Origin)
@@ -112,6 +116,50 @@ namespace LiveLinkAxisSwitchPreProcessor
 	}
 }
 
+/**
+ * 
+ */
+namespace LiveLinkLocationSwitchPreprocessor
+{
+	float LiveLinkAxisToVectorMember(ELiveLinkAxis InAxis, const FVector& Origin)
+	{
+		const uint8 Value = static_cast<uint8>(InAxis) % 3;
+		return Origin[Value];
+	}
+
+	float AxisSign(ELiveLinkAxis InAxis)
+	{
+		return static_cast<uint8>(InAxis) < 3 ? 1.f : -1.f;
+	}
+
+	ELiveLinkAxis AxisAbs(ELiveLinkAxis InAxis)
+	{
+		const uint8 Value = static_cast<uint8>(InAxis) % 3;
+		return static_cast<ELiveLinkAxis>(Value);
+	} 
+
+	void SwitchTranslation(FVector& Translation, ELiveLinkAxis FrontAxis, ELiveLinkAxis RightAxis, ELiveLinkAxis UpAxis, bool bUseOffsetPosition, FVector OffsetPosition, bool bUseOffsetOrientation, FRotator OffsetOrientation)
+	{
+		FVector Location = Translation;
+		FVector NewLocation;
+		NewLocation[0] = LiveLinkAxisToVectorMember(FrontAxis, Location) * AxisSign(FrontAxis);
+		NewLocation[1] = LiveLinkAxisToVectorMember(RightAxis, Location) * AxisSign(RightAxis);
+		NewLocation[2] = LiveLinkAxisToVectorMember(UpAxis, Location) * AxisSign(UpAxis);
+
+		// Apply any offsets after the remapping
+		if (bUseOffsetPosition)
+		{
+			NewLocation += OffsetPosition;
+		}
+		if (bUseOffsetOrientation)
+		{
+			NewLocation = OffsetOrientation.RotateVector(NewLocation); //Note this is executed after the axes are flipped/swapped
+		}
+
+		Translation = NewLocation;
+	}
+}
+
 /**
  * ULiveLinkTransformAxisSwitchPreProcessor::FLiveLinkTransformAxisSwitchPreProcessorWorker
  */
@@ -244,3 +292,43 @@ void ULiveLinkTransformAxisSwitchPreProcessor::Serialize(FArchive& Ar)
 #endif
 }
 
+/**
+ * ULiveLinkTransformAxisSwitchPreProcessor::FLiveLinkLocatorAxisSwitchPreProcessorWorker
+ */
+
+TSubclassOf<ULiveLinkRole> ULiveLinkLocatorAxisSwitchPreProcessor::FLiveLinkLocatorAxisSwitchPreProcessorWorker::GetRole() const
+	{
+		return ULiveLinkLocatorRole::StaticClass();
+	}
+
+bool ULiveLinkLocatorAxisSwitchPreProcessor::FLiveLinkLocatorAxisSwitchPreProcessorWorker::PreProcessFrame(FLiveLinkFrameDataStruct& InOutFrame) const
+{
+	FLiveLinkLocatorFrameData& LocatorData = *InOutFrame.Cast<FLiveLinkLocatorFrameData>();
+
+	for (FVector& Location : LocatorData.Locators)
+	{
+		LiveLinkLocationSwitchPreprocessor::SwitchTranslation(Location, FrontAxis, RightAxis, UpAxis, bUseOffsetPosition, OffsetPosition, bUseOffsetOrientation, OffsetOrientation);
+	}
+	return true;
+}
+
+ULiveLinkFramePreProcessor::FWorkerSharedPtr ULiveLinkLocatorAxisSwitchPreProcessor::FetchWorker()
+{
+	if (!Instance.IsValid())
+	{
+		Instance = MakeShared<FLiveLinkLocatorAxisSwitchPreProcessorWorker, ESPMode::ThreadSafe>();
+		Instance->FrontAxis = FrontAxis;
+		Instance->RightAxis = RightAxis;
+		Instance->UpAxis = UpAxis;
+		Instance->bUseOffsetPosition = bUseOffsetPosition;
+		Instance->bUseOffsetOrientation = bUseOffsetOrientation;
+		Instance->OffsetPosition = OffsetPosition;
+		Instance->OffsetOrientation = OffsetOrientation;
+	}
+	return Instance;
+}
+
+TSubclassOf<ULiveLinkRole> ULiveLinkLocatorAxisSwitchPreProcessor::GetRole() const
+{
+	return ULiveLinkLocatorRole::StaticClass();
+}
diff --git a/Engine/Plugins/Animation/LiveLink/Source/LiveLink/Public/PreProcessor/LiveLinkAxisSwitchPreProcessor.h b/Engine/Plugins/Animation/LiveLink/Source/LiveLink/Public/PreProcessor/LiveLinkAxisSwitchPreProcessor.h
--- a/Engine/Plugins/Animation/LiveLink/Source/LiveLink/Public/PreProcessor/LiveLinkAxisSwitchPreProcessor.h
+++ b/Engine/Plugins/Animation/LiveLink/Source/LiveLink/Public/PreProcessor/LiveLinkAxisSwitchPreProcessor.h
@@ -117,3 +117,26 @@ public:
 	virtual TSubclassOf<ULiveLinkRole> GetRole() const override;
 	virtual ULiveLinkFramePreProcessor::FWorkerSharedPtr FetchWorker() override;
 };
+
+/**
+ * Allows to switch any axis of an incoming locator subject with another axis.
+ * @note For example the Z-Axis of an incoming transform can be set to the (optionally negated) Y-Axis of the transform in UE.
+ * @note This implies that translation, rotation and scale will be affected by switching an axis.
+ */
+UCLASS(meta = (DisplayName = "Locator Axis Switch"))
+class LIVELINK_API ULiveLinkLocatorAxisSwitchPreProcessor : public ULiveLinkTransformAxisSwitchPreProcessor
+{
+	GENERATED_BODY()
+
+public:
+	class FLiveLinkLocatorAxisSwitchPreProcessorWorker : public ULiveLinkTransformAxisSwitchPreProcessor::FLiveLinkTransformAxisSwitchPreProcessorWorker
+	{
+	public:
+		virtual TSubclassOf<ULiveLinkRole> GetRole() const override;
+		virtual bool PreProcessFrame(FLiveLinkFrameDataStruct& InOutFrame) const override;
+	};
+
+	public:
+	virtual TSubclassOf<ULiveLinkRole> GetRole() const override;
+	virtual ULiveLinkFramePreProcessor::FWorkerSharedPtr FetchWorker() override;
+};
\ No newline at end of file
diff --git a/Engine/Source/Runtime/LiveLinkInterface/Private/Roles/LiveLinkLocatorRole.cpp b/Engine/Source/Runtime/LiveLinkInterface/Private/Roles/LiveLinkLocatorRole.cpp
new file mode 100644
--- /dev/null
+++ b/Engine/Source/Runtime/LiveLinkInterface/Private/Roles/LiveLinkLocatorRole.cpp
@@ -0,0 +1,74 @@
+// Copyright Epic Games, Inc. All Rights Reserved.
+
+
+#include "Roles/LiveLinkLocatorRole.h"
+
+#include "LiveLinkTypes.h"
+#include "Roles/LiveLinkLocatorTypes.h"
+
+#include UE_INLINE_GENERATED_CPP_BY_NAME(LiveLinkLocatorRole)
+
+#define LOCTEXT_NAMESPACE "LiveLinkRole"
+
+
+UScriptStruct* ULiveLinkLocatorRole::GetStaticDataStruct() const
+{
+	return FLiveLinkLocatorStaticData::StaticStruct();
+}
+
+UScriptStruct* ULiveLinkLocatorRole::GetFrameDataStruct() const
+{
+	return FLiveLinkLocatorFrameData::StaticStruct();
+}
+
+UScriptStruct* ULiveLinkLocatorRole::GetBlueprintDataStruct() const
+{
+	return FLiveLinkLocatorBlueprintData::StaticStruct();
+}
+
+bool ULiveLinkLocatorRole::InitializeBlueprintData(const FLiveLinkSubjectFrameData& InSourceData,
+	FLiveLinkBlueprintDataStruct& OutBlueprintData) const
+{
+	bool bSuccess = false;
+
+	FLiveLinkLocatorBlueprintData* BlueprintData = OutBlueprintData.Cast<FLiveLinkLocatorBlueprintData>();
+	const FLiveLinkLocatorStaticData* StaticData = InSourceData.StaticData.Cast<FLiveLinkLocatorStaticData>();
+	const FLiveLinkLocatorFrameData* FrameData = InSourceData.FrameData.Cast<FLiveLinkLocatorFrameData>();
+	if(BlueprintData && StaticData && FrameData)
+	{
+		GetStaticDataStruct()->CopyScriptStruct(&BlueprintData->StaticData, StaticData);
+		GetFrameDataStruct()->CopyScriptStruct(&BlueprintData->FrameData, FrameData);
+		bSuccess = true;
+	}
+	return bSuccess;
+}
+
+FText ULiveLinkLocatorRole::GetDisplayName() const
+{
+	return LOCTEXT("LiveLinkLocator", "Locator");
+}
+
+bool ULiveLinkLocatorRole::IsFrameDataValid(const FLiveLinkStaticDataStruct& InStaticData, const FLiveLinkFrameDataStruct& InFrameData, bool& bOutShouldLogWarning) const
+{
+	bool bResult = Super::IsFrameDataValid(InStaticData, InFrameData, bOutShouldLogWarning);
+	if(bResult)
+	{
+		const FLiveLinkLocatorStaticData* StaticData = InStaticData.Cast<FLiveLinkLocatorStaticData>();
+		const FLiveLinkLocatorFrameData* FrameData = InFrameData.Cast<FLiveLinkLocatorFrameData>();
+		if(!StaticData->bUnlabelledData)
+		{
+			bResult = StaticData && FrameData && StaticData->LocatorNames.Num() == FrameData->Locators.Num();
+		}
+		else if(StaticData->LocatorNames.Num()>0 && StaticData->bUnlabelledData)
+		{
+			bResult = false; //Always fail if we are expecting unlabelled data and there are marker names/labels. 
+		}
+		else
+		{
+			bResult = StaticData->bUnlabelledData;
+		}
+	}
+	return bResult;
+}
+
+#undef LOCTEXT_NAMESPACE
diff --git a/Engine/Source/Runtime/LiveLinkInterface/Private/Roles/LiveLinkLocatorTypes.cpp b/Engine/Source/Runtime/LiveLinkInterface/Private/Roles/LiveLinkLocatorTypes.cpp
new file mode 100644
--- /dev/null
+++ b/Engine/Source/Runtime/LiveLinkInterface/Private/Roles/LiveLinkLocatorTypes.cpp
@@ -0,0 +1,3 @@
+// Copyright Epic Games, Inc. All Rights Reserved.
+
+#include "Roles/LiveLinkLocatorTypes.h"
diff --git a/Engine/Source/Runtime/LiveLinkInterface/Public/Roles/LiveLinkLocatorRole.h b/Engine/Source/Runtime/LiveLinkInterface/Public/Roles/LiveLinkLocatorRole.h
new file mode 100644
--- /dev/null
+++ b/Engine/Source/Runtime/LiveLinkInterface/Public/Roles/LiveLinkLocatorRole.h
@@ -0,0 +1,32 @@
+// Copyright Epic Games, Inc. All Rights Reserved.
+
+#pragma once
+
+#include "Roles/LiveLinkBasicRole.h"
+#include "CoreMinimal.h"
+#include "LiveLinkTypes.h"
+#include "UObject/ObjectMacros.h"
+#include "UObject/UObjectGlobals.h"
+
+#include "LiveLinkLocatorRole.generated.h"
+
+/**
+ * Role associated with Locator / Marker data. This is intended for optical motion capture marker clouds. 
+ */
+UCLASS(Blueprintable, meta = (DisplayName = "Locator Role"))
+class LIVELINKINTERFACE_API ULiveLinkLocatorRole : public ULiveLinkBasicRole
+{
+	GENERATED_BODY()
+
+public:
+	//~ Begin ULiveLinkRole interface
+	virtual UScriptStruct* GetStaticDataStruct() const override;
+	virtual UScriptStruct* GetFrameDataStruct() const override;
+	virtual UScriptStruct* GetBlueprintDataStruct() const override;
+
+	virtual bool InitializeBlueprintData(const FLiveLinkSubjectFrameData& InSourceData, FLiveLinkBlueprintDataStruct& OutBlueprintData) const override;
+
+	virtual FText GetDisplayName() const override;
+	virtual bool IsFrameDataValid(const FLiveLinkStaticDataStruct& InStaticData, const FLiveLinkFrameDataStruct& InFrameData, bool& bOutShouldLogWarning) const override;
+	//~ End ULiveLinkRole interface
+};
diff --git a/Engine/Source/Runtime/LiveLinkInterface/Public/Roles/LiveLinkLocatorTypes.h b/Engine/Source/Runtime/LiveLinkInterface/Public/Roles/LiveLinkLocatorTypes.h
new file mode 100644
--- /dev/null
+++ b/Engine/Source/Runtime/LiveLinkInterface/Public/Roles/LiveLinkLocatorTypes.h
@@ -0,0 +1,70 @@
+// Copyright Epic Games, Inc. All Rights Reserved.
+
+#pragma once
+
+#include "LiveLinkTypes.h"
+
+#include "LiveLinkLocatorTypes.generated.h"
+
+/**
+ * Static data for Locator purposes. Contains data about locators that should not change every frame. If data is unlabelled markers, Locator names array must be empty.
+ */
+USTRUCT(BlueprintType)
+struct FLiveLinkLocatorStaticData : public FLiveLinkBaseStaticData
+{
+	GENERATED_BODY()
+
+public:
+
+	/*Set the locator names*/
+	void SetLocatorNames(const TArray<FName>& InLocatorNames) { LocatorNames = InLocatorNames; }
+
+	/*Get the locator names*/
+	const TArray<FName>& GetLocatorNames() const { return LocatorNames; }
+
+	/*Set as Unlabelled data*/
+	void SetUnlabelledData(const bool& InbUnlabelledData) {bUnlabelledData = InbUnlabelledData; }
+
+	UPROPERTY(EditAnywhere, BlueprintReadWrite, Category = "Locator")
+	TArray<FName> LocatorNames;
+
+	/*
+	 * Set this to true if you wish to send an unstructured number of locators that can vary from one frame to the next. 
+	 */
+	UPROPERTY(EditAnywhere, BlueprintReadWrite, Category = "Locator")
+	bool bUnlabelledData=false;
+};
+
+	/**
+	 * Dynamic data for Animation purposes. 
+	 */	
+	USTRUCT(BlueprintType)
+	struct FLiveLinkLocatorFrameData : public FLiveLinkBaseFrameData
+{
+	GENERATED_BODY()
+public:
+
+	/**
+	 * Array of locations for each locator/marker
+	 */
+	UPROPERTY(EditAnywhere, BlueprintReadWrite, Category = "Locator")
+	TArray<FVector> Locators;
+};
+
+/**
+ * Facility structure to handle data access in Blueprint
+ */
+
+USTRUCT(BlueprintType)
+struct LIVELINKINTERFACE_API FLiveLinkLocatorBlueprintData : public FLiveLinkBaseBlueprintData
+{
+	GENERATED_BODY()
+
+	/** Static data should not change every frame */
+	UPROPERTY(EditAnywhere, BlueprintReadWrite, Category = "Locator")
+	FLiveLinkLocatorStaticData StaticData;
+
+	UPROPERTY(EditAnywhere, BlueprintReadWrite, Category = "Locator")
+	FLiveLinkLocatorFrameData FrameData;
+	
+};
\ No newline at end of file
