LiveLinkClient - Add delegate to inform when a subject's state has changed.

diff --git a/Engine/Plugins/Animation/LiveLink/Source/LiveLink/Private/LiveLinkClient.cpp b/Engine/Plugins/Animation/LiveLink/Source/LiveLink/Private/LiveLinkClient.cpp
--- a/Engine/Plugins/Animation/LiveLink/Source/LiveLink/Private/LiveLinkClient.cpp
+++ b/Engine/Plugins/Animation/LiveLink/Source/LiveLink/Private/LiveLinkClient.cpp
@@ -1906,6 +1906,11 @@ FOnLiveLinkSubjectChangedDelegate& FLiveLinkClient::OnLiveLinkSubjectAdded()
 	return Collection->OnLiveLinkSubjectAdded();
 }
 
+FOnLiveLinkSubjectStateChanged& FLiveLinkClient::OnLiveLinkSubjectStateChanged()
+{
+	return Collection->OnLiveLinkSubjectStateChanged();
+}
+
 #if WITH_EDITOR
 FOnLiveLinkSubjectEvaluated& FLiveLinkClient::OnLiveLinkSubjectEvaluated()
 {
diff --git a/Engine/Plugins/Animation/LiveLink/Source/LiveLink/Private/LiveLinkSourceCollection.cpp b/Engine/Plugins/Animation/LiveLink/Source/LiveLink/Private/LiveLinkSourceCollection.cpp
--- a/Engine/Plugins/Animation/LiveLink/Source/LiveLink/Private/LiveLinkSourceCollection.cpp
+++ b/Engine/Plugins/Animation/LiveLink/Source/LiveLink/Private/LiveLinkSourceCollection.cpp
@@ -229,6 +229,11 @@ void FLiveLinkSourceCollection::AddSubject(FLiveLinkCollectionSubjectItem InSubj
 {
 	FLiveLinkSubjectKey Key = InSubject.Key;
 
+	if (FLiveLinkSubject* LiveLinkSubject = InSubject.GetLiveSubject())
+	{
+		LiveLinkSubject->OnStateChanged().BindRaw(this, &FLiveLinkSourceCollection::HandleSubjectStateChanged, Key);
+	}
+
 	{
 		UE::TUniqueLock Lock(SubjectsLock);
 		Subjects.Add(MoveTemp(InSubject));
@@ -442,3 +447,8 @@ void FLiveLinkSourceCollection::ForEachSource(TFunctionRef<void(const FLiveLinkC
 		VisitorFunc(Source);
 	}
 }
+
+void FLiveLinkSourceCollection::HandleSubjectStateChanged(ELiveLinkSubjectState NewState, FLiveLinkSubjectKey SubjectKey)
+{
+	BroadcastOnGameThread(OnLiveLinkSubjectStateChanged(), SubjectKey, NewState);
+}
\ No newline at end of file
diff --git a/Engine/Plugins/Animation/LiveLink/Source/LiveLink/Private/LiveLinkSubject.cpp b/Engine/Plugins/Animation/LiveLink/Source/LiveLink/Private/LiveLinkSubject.cpp
--- a/Engine/Plugins/Animation/LiveLink/Source/LiveLink/Private/LiveLinkSubject.cpp
+++ b/Engine/Plugins/Animation/LiveLink/Source/LiveLink/Private/LiveLinkSubject.cpp
@@ -26,6 +26,7 @@ FLiveLinkSubject::FSubjectEvaluationStatistics::FSubjectEvaluationStatistics()
 
 FLiveLinkSubject::FLiveLinkSubject(TSharedPtr<FLiveLinkTimedDataInput> InTimedDataGroup)
 	: TimedDataGroup(InTimedDataGroup)
+	, State(ELiveLinkSubjectState::Unknown)
 {
 	check(InTimedDataGroup.IsValid());
 }
@@ -205,6 +206,8 @@ void FLiveLinkSubject::Update()
 		const FTimecode FrameTC = FTimecode::FromFrameNumber(FrameSnapshot.FrameData.GetBaseData()->MetaData.SceneTime.Time.GetFrame(), FrameSnapshot.FrameData.GetBaseData()->MetaData.SceneTime.Rate);
 		UE_LOG(LogLiveLink, Verbose, TEXT("Subject '%s' created snapshot with frame ID %d - Timecode:[%s.%0.3f] - SourceTime: %0.4f, Offset: %0.6f, CorrectedTime: %0.4f"), *SubjectKey.SubjectName.ToString(), FrameSnapshot.FrameData.GetBaseData()->FrameId, *FrameTC.ToString(), FrameSnapshot.FrameData.GetBaseData()->MetaData.SceneTime.Time.GetSubFrame(), FrameSnapshot.FrameData.GetBaseData()->WorldTime.GetSourceTime(), FrameSnapshot.FrameData.GetBaseData()->WorldTime.GetOffset(), FrameSnapshot.FrameData.GetBaseData()->WorldTime.GetOffsettedTime());
 	}
+
+	UpdateState();
 }
 
 void FLiveLinkSubject::ClearFrames()
@@ -1058,6 +1061,19 @@ void FLiveLinkSubject::RemoveFrames(int32 InCount)
 	}
 }
 
+void FLiveLinkSubject::UpdateState()
+{
+	const ELiveLinkSubjectState PreviousState = State;
+
+	const bool bHasValidFrame = (FApp::GetCurrentTime() - GetLastPushTime() < GetDefault<ULiveLinkSettings>()->GetTimeWithoutFrameToBeConsiderAsInvalid());
+	State = (bHasValidFrame && HasValidFrameSnapshot()) ? ELiveLinkSubjectState::Connected : ELiveLinkSubjectState::Unresponsive;
+
+	if (State != PreviousState)
+	{
+		OnStateChangedDelegate.ExecuteIfBound(State);
+	}
+}
+
 void FLiveLinkSubject::SetStaticData(TSubclassOf<ULiveLinkRole> InRole, FLiveLinkStaticDataStruct&& InStaticData)
 {
 	if (Role == nullptr)
@@ -1270,8 +1286,16 @@ FText FLiveLinkSubject::GetDisplayName() const
 
 ETimedDataInputState FLiveLinkSubject::GetState() const
 {
-	bool bHasValidFrame = (FApp::GetCurrentTime() - GetLastPushTime() < GetDefault<ULiveLinkSettings>()->GetTimeWithoutFrameToBeConsiderAsInvalid());
-	return (bHasValidFrame && HasValidFrameSnapshot()) ? ETimedDataInputState::Connected : ETimedDataInputState::Unresponsive;
+	using enum ELiveLinkSubjectState;
+	switch (State)
+	{
+		case Connected: return ETimedDataInputState::Connected;
+		case Unresponsive: return ETimedDataInputState::Unresponsive;
+		case Disconnected: return ETimedDataInputState::Disconnected;
+		default:
+			ensureMsgf(false, TEXT("Unhandled ELiveLinkSubjectState::%d"), State);
+			return ETimedDataInputState::Disconnected;
+	}
 }
 
 FTimedDataChannelSampleTime FLiveLinkSubject::GetOldestDataTime() const
diff --git a/Engine/Plugins/Animation/LiveLink/Source/LiveLink/Public/LiveLinkClient.h b/Engine/Plugins/Animation/LiveLink/Source/LiveLink/Public/LiveLinkClient.h
--- a/Engine/Plugins/Animation/LiveLink/Source/LiveLink/Public/LiveLinkClient.h
+++ b/Engine/Plugins/Animation/LiveLink/Source/LiveLink/Public/LiveLinkClient.h
@@ -108,6 +108,7 @@ public:
 	virtual FOnLiveLinkSourceChangedDelegate& OnLiveLinkSourceRemoved() override;
 	virtual FOnLiveLinkSubjectChangedDelegate& OnLiveLinkSubjectAdded() override;
 	virtual FOnLiveLinkSubjectChangedDelegate& OnLiveLinkSubjectRemoved() override;
+	virtual FOnLiveLinkSubjectStateChanged& OnLiveLinkSubjectStateChanged() override;
 #if WITH_EDITOR
 	virtual FOnLiveLinkSubjectEvaluated& OnLiveLinkSubjectEvaluated() override;
 #endif
diff --git a/Engine/Plugins/Animation/LiveLink/Source/LiveLink/Public/LiveLinkSourceCollection.h b/Engine/Plugins/Animation/LiveLink/Source/LiveLink/Public/LiveLinkSourceCollection.h
--- a/Engine/Plugins/Animation/LiveLink/Source/LiveLink/Public/LiveLinkSourceCollection.h
+++ b/Engine/Plugins/Animation/LiveLink/Source/LiveLink/Public/LiveLinkSourceCollection.h
@@ -151,6 +151,7 @@ public:
 	FOnLiveLinkSourceChangedDelegate& OnLiveLinkSourceRemoved() { return OnLiveLinkSourceRemovedDelegate; }
 	FOnLiveLinkSubjectChangedDelegate& OnLiveLinkSubjectAdded() { return OnLiveLinkSubjectAddedDelegate; }
 	FOnLiveLinkSubjectChangedDelegate& OnLiveLinkSubjectRemoved() { return OnLiveLinkSubjectRemovedDelegate; }
+	FOnLiveLinkSubjectStateChanged& OnLiveLinkSubjectStateChanged() { return OnLiveLinkSubjectStateChangedDelegate; }
 
 	/** Utility method used to broadcast delegates on the game thread if this function is called on a different thread. */
 	template <typename DelegateType, typename... ArgTypes>
@@ -169,6 +170,11 @@ public:
 		}
 	}
 
+private:
+
+	/** Handles broadcasting the inner subject state change delegate to all listeners of the collection's delegate. */
+	void HandleSubjectStateChanged(ELiveLinkSubjectState NewState, FLiveLinkSubjectKey SubjectKey);
+
 private:
 	TArray<FLiveLinkCollectionSourceItem> Sources;
 
@@ -192,6 +198,9 @@ private:
 	/** Notify when a client subject's is removed */
 	FOnLiveLinkSubjectChangedDelegate OnLiveLinkSubjectRemovedDelegate;
 
+	/** Notify when a client subject's state has changed. (ie. it became stale) */
+	FOnLiveLinkSubjectStateChanged OnLiveLinkSubjectStateChangedDelegate;
+
 	/** Lock to stop multiple threads accessing the Subjects from the collection at the same time */
 	mutable UE::FRecursiveMutex SubjectsLock;
 
diff --git a/Engine/Plugins/Animation/LiveLink/Source/LiveLink/Public/LiveLinkSubject.h b/Engine/Plugins/Animation/LiveLink/Source/LiveLink/Public/LiveLinkSubject.h
--- a/Engine/Plugins/Animation/LiveLink/Source/LiveLink/Public/LiveLinkSubject.h
+++ b/Engine/Plugins/Animation/LiveLink/Source/LiveLink/Public/LiveLinkSubject.h
@@ -15,30 +15,33 @@ class FLiveLinkTimedDataInput;
 class ULiveLinkSubjectSettings;
 struct FLiveLinkSubjectTimeSyncData;
 
+enum class ELiveLinkSubjectState : uint8;
 
 struct FLiveLinkInterpolationInfo;
 
 struct FLiveLinkTimeSynchronizationData
 {
-	// Whether or not synchronization has been established.
+	/** Whether or not synchronization has been established. */
 	bool bHasEstablishedSync = false;
 
-	// The frame in our buffer where a rollover was detected. Only applicable for time synchronized sources.
+	/** The frame in our buffer where a rollover was detected. Only applicable for time synchronized sources. */
 	int32 RolloverFrame = INDEX_NONE;
 
-	// Frame offset that will be used for this source.
+	/** Frame offset that will be used for this source. */
 	int32 Offset = 0;
 
-	// Frame Time value modulus. When this value is not set, we assume no rollover occurs.
+	/** Frame Time value modulus. When this value is not set, we assume no rollover occurs. */
 	TOptional<FFrameTime> RolloverModulus;
 
-	// Frame rate used as the base for synchronization.
+	/** Frame rate used as the base for synchronization. */
 	FFrameRate SyncFrameRate;
 
-	// Frame time that synchronization was established (relative to SynchronizationFrameRate).
+	/** Frame time that synchronization was established (relative to SynchronizationFrameRate). */
 	FFrameTime SyncStartTime;
 };
 
+DECLARE_DELEGATE_OneParam(FOnStateChanged, ELiveLinkSubjectState /*NewState*/);
+
 /**
  * Manages subject manipulation either to add or get frame data for specific roles
  */
@@ -110,10 +113,10 @@ public:
 
 	bool HasStaticData() const;
 
-	// Handling setting a new static data. Create role data if not found in map.
+	/** Handling setting a new static data. Create role data if not found in map. */
 	void SetStaticData(TSubclassOf<ULiveLinkRole> InRole, FLiveLinkStaticDataStruct&& InStaticData);
 
-	// Add a frame of data from a FLiveLinkFrameData
+	/** Add a frame of data from a FLiveLinkFrameData */
 	void AddFrameData(FLiveLinkFrameDataStruct&& InFrameData);
 
 	void CacheSettings(ULiveLinkSourceSettings* SourceSetting, ULiveLinkSubjectSettings* SubjectSetting);
@@ -133,44 +136,50 @@ public:
 	/** Clear the override static data for this subject. */
 	void ClearOverrideStaticData_AnyThread();
 
+	/** Get the delegate triggered when the state changes. */
+	FOnStateChanged& OnStateChanged() { return OnStateChangedDelegate; }
+
 private:
 	int32 FindNewFrame_WorldTime(const FLiveLinkWorldTime& FrameTime) const;
 	int32 FindNewFrame_WorldTimeInternal(const FLiveLinkWorldTime& FrameTime) const;
 	int32 FindNewFrame_SceneTime(const FQualifiedFrameTime& FrameTime, const FLiveLinkWorldTime& WorldTime) const;
 	int32 FindNewFrame_Latest(const FLiveLinkWorldTime& FrameTime) const;
 
-	// Reorder frame with the same timecode and create subframes
+	/** Reorder frame with the same timecode and create subframes */
 	void AdjustSubFrame_SceneTime(int32 FrameIndex);
 
-	// Populate OutFrame with a frame based off of the supplied time (pre offsetted)
+	/** Populate OutFrame with a frame based off of the supplied time (pre offsetted) */
 	bool GetFrameAtWorldTime(const double InSeconds, FLiveLinkSubjectFrameData& OutFrame);
 	bool GetFrameAtWorldTime_Closest(const double InSeconds, FLiveLinkSubjectFrameData& OutFrame);
 	bool GetFrameAtWorldTime_Interpolated(const double InSeconds, FLiveLinkSubjectFrameData& OutFrame);
 
-	// Populate OutFrame with a frame based off of the supplied scene time (pre offsetted).
+	/** Populate OutFrame with a frame based off of the supplied scene time (pre offsetted). */
 	bool GetFrameAtSceneTime(const FQualifiedFrameTime& InSceneTime, FLiveLinkSubjectFrameData& OutFrame);
 	bool GetFrameAtSceneTime_Closest(const FQualifiedFrameTime& InSceneTime, FLiveLinkSubjectFrameData& OutFrame);
 	bool GetFrameAtSceneTime_Interpolated(const FQualifiedFrameTime& InSceneTime, FLiveLinkSubjectFrameData& OutFrame);
 	
-	// Verify interpolation result to update our internal statistics
+	/** Verify interpolation result to update our internal statistics */
 	void VerifyInterpolationInfo(const FLiveLinkInterpolationInfo& InterpolationInfo);
 
-	// Populate OutFrame with the latest frame.
+	/** Populate OutFrame with the latest frame. */
 	bool GetLatestFrame(FLiveLinkSubjectFrameData& OutFrame);
 
 	void ResetFrame(FLiveLinkSubjectFrameData& OutFrame) const;
 
-	// Update our internal statistics
+	/** Update our internal statistics */
 	void IncreaseFrameDroppedStat();
 	void IncreaseBufferUnderFlowStat();
 	void IncreaseBufferOverFlowStat();
 	void UpdateEvaluationData(const FTimedDataInputEvaluationData& EvaluationData);
 
-	//Remove frames from our buffer - based on receiving order
+	/** Remove frames from our buffer - based on receiving order */
 	void RemoveFrames(int32 Count);
 
+	/** Update the cached state for this subject. */
+	void UpdateState();
+
 protected:
-	// The role the subject was build with
+	/** The role the subject was build with */
 	TSubclassOf<ULiveLinkRole> Role;
 
 	TArray<ULiveLinkFramePreProcessor::FWorkerSharedPtr> FramePreProcessors;
@@ -202,65 +211,70 @@ private:
 		FSubjectEvaluationStatistics& operator=(const FSubjectEvaluationStatistics&) = delete;
 	};
 
-	// Static data of the subject
+	/** Static data of the subject */
 	FLiveLinkStaticDataStruct StaticData;
 
-	// Override static data, set by the remapper.
+	/** Override static data, set by the remapper. */
 	TOptional<FLiveLinkStaticDataStruct> OverrideStaticData;
 
-	// Frames added to the subject
+	/** Frames added to the subject */
 	TArray<FLiveLinkFrameDataStruct> FrameData;
 
-	// Contains identifier of each frame in the order they were received
+	/** Contains identifier of each frame in the order they were received */
 	TQueue<FLiveLinkFrameIdentifier> ReceivedOrderedFrames;
 
-	// Next identifier to assign to next received frame
+	/** Next identifier to assign to next received frame */
 	FLiveLinkFrameIdentifier NextIdentifier = 0;
 
-	// Current frame snapshot of the evaluation
+	/** Current frame snapshot of the evaluation */
 	FLiveLinkSubjectFrameData FrameSnapshot;
 
-	// Name of the subject
+	/** Name of the subject */
 	FLiveLinkSubjectKey SubjectKey;
 
-	// Timed data input group for the subject
+	/** Timed data input group for the subject */
 	TWeakPtr<FLiveLinkTimedDataInput> TimedDataGroup;
 
-	// Connection settings specified by user
+	/** Connection settings specified by user */
 	FLiveLinkCachedSettings CachedSettings;
 
-	// Override mode, determined by frame data
+	/** Override mode, determined by frame data */
 	TOptional<ELiveLinkSourceMode> ModeOverride;
 
-	// Last time a frame was pushed
+	/** Last time a frame was pushed */
 	double LastPushTime = 0.0;
 
-	// Logging stats is enabled by default. If monitor opens at a later stage,previous stats will be able to be seen
+	/** Logging stats is enabled by default. If monitor opens at a later stage,previous stats will be able to be seen */
 	bool bIsStatLoggingEnabled = true;
 
-	// Some stats compiled by the subject.
+	/** Some stats compiled by the subject. */
 	FSubjectEvaluationStatistics EvaluationStatistics;
 
-	// Last Timecode FrameRate received
+	/** Last Timecode FrameRate received */
 	FFrameRate LastTimecodeFrameRate;
 
-	// If enabled, rebroadcast this subject
+	/** If enabled, rebroadcast this subject */
     bool bRebroadcastSubject = false;
 
-	// If true, static data has been sent for this rebroadcast
+	/** If true, static data has been sent for this rebroadcast */
 	bool bRebroadcastStaticDataSent = false;
 
-	// If true, override static data may remap when caching settings.
+	/** If true, override static data may remap when caching settings. */
 	bool bNeedsStaticRemap = false;
 
 	/** Flag set to clear the override static data for a subject. */
 	std::atomic<bool> bClearOverrideStaticData = false;
+
+	/** Current state of this subject. */
+	ELiveLinkSubjectState State;
+
+	/** Delegate called when the state of this subject has changed. */
+	FOnStateChanged OnStateChangedDelegate;
 	
 	/** 
 	 * Evaluation can be done on any thread so we need to protect statistic logging 
 	 * Some stats requires more than atomic sized vars so a critical section is used to protect when necessary
 	 */
-
 	mutable FCriticalSection StatisticCriticalSection;
 
 	/** Used to protect access to translators, preprocessors and interpolation processors since they can be set and accessed in different threads. */
diff --git a/Engine/Source/Runtime/LiveLinkInterface/Public/ILiveLinkClient.h b/Engine/Source/Runtime/LiveLinkInterface/Public/ILiveLinkClient.h
--- a/Engine/Source/Runtime/LiveLinkInterface/Public/ILiveLinkClient.h
+++ b/Engine/Source/Runtime/LiveLinkInterface/Public/ILiveLinkClient.h
@@ -44,6 +44,7 @@ DECLARE_TS_MULTICAST_DELEGATE_OneParam(FOnLiveLinkSubjectFrameDataReceived, cons
 DECLARE_TS_MULTICAST_DELEGATE_ThreeParams(FOnLiveLinkSubjectStaticDataAdded, FLiveLinkSubjectKey /*InSubjectKey*/, TSubclassOf<ULiveLinkRole> /*SubjectRole*/, const FLiveLinkStaticDataStruct& /*InStaticData*/)
 DECLARE_TS_MULTICAST_DELEGATE_ThreeParams(FOnLiveLinkSubjectFrameDataAdded, FLiveLinkSubjectKey /*InSubjectKey*/, TSubclassOf<ULiveLinkRole> /*SubjectRole*/, const FLiveLinkFrameDataStruct& /*InFrameData*/)
 DECLARE_TS_MULTICAST_DELEGATE_FiveParams(FOnLiveLinkSubjectEvaluated, FLiveLinkSubjectKey /*InSubjectKey*/, TSubclassOf<ULiveLinkRole> /*RequestedRole*/, const FLiveLinkTime& /*RequestedTime*/, bool /*bResult*/, const FLiveLinkTime& /*EvaluatedFrameTime*/)
+DECLARE_TS_MULTICAST_DELEGATE_TwoParams(FOnLiveLinkSubjectStateChanged, FLiveLinkSubjectKey /*InSubjectKey*/, ELiveLinkSubjectState /*NewState*/);
 
 
 /**
@@ -307,6 +308,9 @@ public:
 	/** Notify when a subject has been removed */
 	virtual FOnLiveLinkSubjectChangedDelegate& OnLiveLinkSubjectRemoved() = 0;
 
+	/** Notify when a subject's state has changed. */
+	virtual FOnLiveLinkSubjectStateChanged& OnLiveLinkSubjectStateChanged() = 0;
+
 #if WITH_EDITOR
 	/** Notify the debug interface when a subject has been evaluated. Only available in editor and used for debugging purpose. */
 	virtual FOnLiveLinkSubjectEvaluated& OnLiveLinkSubjectEvaluated() = 0;
