LiveLinkSubject - Set LiveLinkState to Disconnected once initialized.

#rnx

diff --git a/Engine/Plugins/Animation/LiveLink/Source/LiveLink/Private/LiveLinkSubject.cpp b/Engine/Plugins/Animation/LiveLink/Source/LiveLink/Private/LiveLinkSubject.cpp
--- a/Engine/Plugins/Animation/LiveLink/Source/LiveLink/Private/LiveLinkSubject.cpp
+++ b/Engine/Plugins/Animation/LiveLink/Source/LiveLink/Private/LiveLinkSubject.cpp
@@ -36,6 +36,8 @@ void FLiveLinkSubject::Initialize(FLiveLinkSubjectKey InSubjectKey, TSubclassOf<
 	SubjectKey = InSubjectKey;
 	Role = InRole;
 
+	State = ELiveLinkSubjectState::Disconnected;
+
 	FrameData.Reset();
 	ReceivedOrderedFrames.Empty();
 	LastTimecodeFrameRate = FFrameRate(1, -1); //Initialized as invalid for first detection
