LiveLink - Allow combining child roles of basic roles into animation virtual subjects

diff --git a/Engine/Plugins/Animation/LiveLink/Source/LiveLink/Private/LiveLinkAnimationVirtualSubject.cpp b/Engine/Plugins/Animation/LiveLink/Source/LiveLink/Private/LiveLinkAnimationVirtualSubject.cpp
--- a/Engine/Plugins/Animation/LiveLink/Source/LiveLink/Private/LiveLinkAnimationVirtualSubject.cpp
+++ b/Engine/Plugins/Animation/LiveLink/Source/LiveLink/Private/LiveLinkAnimationVirtualSubject.cpp
@@ -123,7 +123,7 @@ bool ULiveLinkAnimationVirtualSubject::AreSubjectsValid(const TArray<FLiveLinkSu
 		if (FoundPtr)
 		{
 			TSubclassOf<ULiveLinkRole> SubjectRole = LiveLinkClient->GetSubjectRole_AnyThread(*FoundPtr);
-			bValid = SubjectRole && (SubjectRole->IsChildOf(ULiveLinkAnimationRole::StaticClass()) || SubjectRole == ULiveLinkBasicRole::StaticClass());
+			bValid = SubjectRole && (SubjectRole->IsChildOf(ULiveLinkAnimationRole::StaticClass()) || SubjectRole->IsChildOf(ULiveLinkBasicRole::StaticClass()));
 		}
 
 		if (!bValid)
@@ -147,7 +147,7 @@ bool ULiveLinkAnimationVirtualSubject::BuildSubjectSnapshot(TArray<FLiveLinkSubj
 		TSubclassOf<ULiveLinkRole> SubjectRole = LiveLinkClient->GetSubjectRole_AnyThread(SubjectName);
 		
 		// Evaluate for basic role if that's the subject's role.
-		TSubclassOf<ULiveLinkRole> DesiredRole = SubjectRole && SubjectRole == ULiveLinkBasicRole::StaticClass() ? SubjectRole : GetRole();
+		TSubclassOf<ULiveLinkRole> DesiredRole = SubjectRole && SubjectRole->IsChildOf(ULiveLinkBasicRole::StaticClass()) ? SubjectRole : GetRole();
 		if (!LiveLinkClient->EvaluateFrame_AnyThread(SubjectName, DesiredRole, NextSnapshot))
 		{
 			bSnapshotDone = false;
diff --git a/Engine/Plugins/Animation/LiveLink/Source/LiveLinkEditor/Private/LiveLinkVirtualSubjectDetailCustomization.cpp b/Engine/Plugins/Animation/LiveLink/Source/LiveLinkEditor/Private/LiveLinkVirtualSubjectDetailCustomization.cpp
--- a/Engine/Plugins/Animation/LiveLink/Source/LiveLinkEditor/Private/LiveLinkVirtualSubjectDetailCustomization.cpp
+++ b/Engine/Plugins/Animation/LiveLink/Source/LiveLinkEditor/Private/LiveLinkVirtualSubjectDetailCustomization.cpp
@@ -57,7 +57,7 @@ void FLiveLinkVirtualSubjectDetailCustomization::CustomizeDetails(IDetailLayoutB
 		for (const FLiveLinkSubjectKey& SubjectKey : SubjectKeys)
 		{
 			TSubclassOf<ULiveLinkRole> LiveLinkRole = Client->GetSubjectRole_AnyThread(SubjectKey);
-			if (LiveLinkRole && (LiveLinkRole->IsChildOf(ULiveLinkAnimationRole::StaticClass()) || LiveLinkRole == ULiveLinkBasicRole::StaticClass()))
+			if (LiveLinkRole && (LiveLinkRole->IsChildOf(ULiveLinkAnimationRole::StaticClass()) || LiveLinkRole->IsChildOf(ULiveLinkBasicRole::StaticClass())))
 			{
 				SubjectsListItems.AddUnique(MakeShared<FName>(SubjectKey.SubjectName.Name));
 			}
