LiveLinkHub - Fix ensure when removing a subject with Metahuman plugin active

#rnx

diff --git a/Engine/Plugins/Animation/LiveLink/Source/LiveLink/Private/LiveLinkSourceCollection.cpp b/Engine/Plugins/Animation/LiveLink/Source/LiveLink/Private/LiveLinkSourceCollection.cpp
--- a/Engine/Plugins/Animation/LiveLink/Source/LiveLink/Private/LiveLinkSourceCollection.cpp
+++ b/Engine/Plugins/Animation/LiveLink/Source/LiveLink/Private/LiveLinkSourceCollection.cpp
@@ -120,12 +120,12 @@ void FLiveLinkSourceCollection::RemoveSource(FGuid InSourceGuid)
 
 			if (bRemovedSubject)
 			{
-				OnLiveLinkSubjectsChanged().Broadcast();
+				BroadcastOnGameThread(OnLiveLinkSubjectsChanged());
 			}
 
 			Sources.RemoveAtSwap(SourceIndex);
-			OnLiveLinkSourceRemoved().Broadcast(InSourceGuid);
-			OnLiveLinkSourcesChanged().Broadcast();
+			BroadcastOnGameThread(OnLiveLinkSourceRemoved(), InSourceGuid);
+			BroadcastOnGameThread(OnLiveLinkSourcesChanged());
 		}
 	}
 }
