Convert existing mutexes to recursive mutexes in LiveLinkSourceCollection to avoid deadlocks

diff --git a/Engine/Plugins/Animation/LiveLink/Source/LiveLink/Private/LiveLinkSourceCollection.cpp b/Engine/Plugins/Animation/LiveLink/Source/LiveLink/Private/LiveLinkSourceCollection.cpp
--- a/Engine/Plugins/Animation/LiveLink/Source/LiveLink/Private/LiveLinkSourceCollection.cpp
+++ b/Engine/Plugins/Animation/LiveLink/Source/LiveLink/Private/LiveLinkSourceCollection.cpp
@@ -86,7 +86,7 @@ FLiveLinkSourceCollection::FLiveLinkSourceCollection()
 
 void FLiveLinkSourceCollection::AddSource(FLiveLinkCollectionSourceItem InSource)
 {
-	UE::TUniqueLock Lock(SourcesLock);
+	UE::TUniqueLock<UE::FRecursiveMutex> Lock(SourcesLock);
 	FLiveLinkCollectionSourceItem& SourceItem = Sources.Add_GetRef(MoveTemp(InSource));
 
 	OnLiveLinkSourceAdded().Broadcast(SourceItem.Guid);
@@ -135,7 +135,7 @@ void FLiveLinkSourceCollection::RemoveAllSources()
 {
 	bool bHasRemovedSubject = false;
 	{
-		UE::TUniqueLock Lock(SubjectsLock);
+		UE::TUniqueLock<UE::FRecursiveMutex> Lock(SubjectsLock);
 		for (int32 Index = Subjects.Num() - 1; Index >= 0; --Index)
 		{
 			bHasRemovedSubject = true;
@@ -152,7 +152,7 @@ void FLiveLinkSourceCollection::RemoveAllSources()
 
 	bool bHasRemovedSource = false;
 	{
-		UE::TUniqueLock Lock(SubjectsLock);
+		UE::TUniqueLock<UE::FRecursiveMutex> Lock(SubjectsLock);
 		for (int32 Index = Sources.Num() - 1; Index >= 0; --Index)
 		{
 			if (Sources[Index].Guid != FLiveLinkSourceCollection::DefaultVirtualSubjectGuid)
@@ -174,7 +174,7 @@ void FLiveLinkSourceCollection::RemoveAllSources()
 
 FLiveLinkCollectionSourceItem* FLiveLinkSourceCollection::FindSource(TSharedPtr<ILiveLinkSource> InSource)
 {
-	UE::TUniqueLock Lock(SourcesLock);
+	UE::TUniqueLock<UE::FRecursiveMutex> Lock(SourcesLock);
 	return Sources.FindByPredicate([InSource](const FLiveLinkCollectionSourceItem& Other) { return Other.Source == InSource; });
 }
 
@@ -187,7 +187,7 @@ const FLiveLinkCollectionSourceItem* FLiveLinkSourceCollection::FindSource(TShar
 
 FLiveLinkCollectionSourceItem* FLiveLinkSourceCollection::FindSource(FGuid InSourceGuid)
 {
-	UE::TUniqueLock Lock(SourcesLock);
+	UE::TUniqueLock<UE::FRecursiveMutex> Lock(SourcesLock);
 	return Sources.FindByPredicate([InSourceGuid](const FLiveLinkCollectionSourceItem& Other) { return Other.Guid == InSourceGuid; });
 }
 
@@ -200,7 +200,7 @@ const FLiveLinkCollectionSourceItem* FLiveLinkSourceCollection::FindSource(FGuid
 
 FLiveLinkCollectionSourceItem* FLiveLinkSourceCollection::FindVirtualSource(FName VirtualSourceName)
 {
-	UE::TUniqueLock Lock(SourcesLock);
+	UE::TUniqueLock<UE::FRecursiveMutex> Lock(SourcesLock);
 	return Sources.FindByPredicate([VirtualSourceName](const FLiveLinkCollectionSourceItem& Other)
 	{
 		if (Other.IsVirtualSource())
@@ -221,7 +221,7 @@ const FLiveLinkCollectionSourceItem* FLiveLinkSourceCollection::FindVirtualSourc
 
 int32 FLiveLinkSourceCollection::NumSources() const
 {
-	UE::TUniqueLock Lock(SourcesLock);
+	UE::TUniqueLock<UE::FRecursiveMutex> Lock(SourcesLock);
 	return Sources.Num();
 }
 
@@ -235,7 +235,7 @@ void FLiveLinkSourceCollection::AddSubject(FLiveLinkCollectionSubjectItem InSubj
 	}
 
 	{
-		UE::TUniqueLock Lock(SubjectsLock);
+		UE::TUniqueLock<UE::FRecursiveMutex> Lock(SubjectsLock);
 		Subjects.Add(MoveTemp(InSubject));
 	}
 
@@ -247,7 +247,7 @@ void FLiveLinkSourceCollection::AddSubject(FLiveLinkCollectionSubjectItem InSubj
 void FLiveLinkSourceCollection::RemoveSubject(FLiveLinkSubjectKey InSubjectKey)
 {
 	{
-		UE::TUniqueLock Lock(SubjectsLock);
+		UE::TUniqueLock<UE::FRecursiveMutex> Lock(SubjectsLock);
 		int32 IndexOf = Subjects.IndexOfByPredicate([InSubjectKey](const FLiveLinkCollectionSubjectItem& Other) { return Other.Key == InSubjectKey; });
 		if (IndexOf != INDEX_NONE)
 		{
@@ -262,7 +262,7 @@ void FLiveLinkSourceCollection::RemoveSubject(FLiveLinkSubjectKey InSubjectKey)
 
 FLiveLinkCollectionSubjectItem* FLiveLinkSourceCollection::FindSubject(FLiveLinkSubjectKey InSubjectKey)
 {
-	UE::TUniqueLock Lock(SubjectsLock);
+	UE::TUniqueLock<UE::FRecursiveMutex> Lock(SubjectsLock);
 	return Subjects.FindByPredicate([InSubjectKey](const FLiveLinkCollectionSubjectItem& Other) { return Other.Key == InSubjectKey; });
 }
 
@@ -275,19 +275,19 @@ const FLiveLinkCollectionSubjectItem* FLiveLinkSourceCollection::FindSubject(FLi
 
 const FLiveLinkCollectionSubjectItem* FLiveLinkSourceCollection::FindSubject(FLiveLinkSubjectName SubjectName) const
 {
-	UE::TUniqueLock Lock(SubjectsLock);
+	UE::TUniqueLock<UE::FRecursiveMutex> Lock(SubjectsLock);
 	return Subjects.FindByPredicate([SubjectName](const FLiveLinkCollectionSubjectItem& Other) { return Other.Key.SubjectName == SubjectName;  });
 }
 
 const FLiveLinkCollectionSubjectItem* FLiveLinkSourceCollection::FindEnabledSubject(FLiveLinkSubjectName InSubjectName) const
 {
-	UE::TUniqueLock Lock(SubjectsLock);
+	UE::TUniqueLock<UE::FRecursiveMutex> Lock(SubjectsLock);
 	return Subjects.FindByPredicate([InSubjectName](const FLiveLinkCollectionSubjectItem& Other) { return Other.Key.SubjectName == InSubjectName && Other.bEnabled && !Other.bPendingKill; });
 }
 
 int32 FLiveLinkSourceCollection::NumSubjects() const
 {
-	UE::TUniqueLock Lock(SubjectsLock);
+	UE::TUniqueLock<UE::FRecursiveMutex> Lock(SubjectsLock);
 	return Subjects.Num();
 }
 
@@ -310,7 +310,7 @@ void FLiveLinkSourceCollection::SetSubjectEnabled(FLiveLinkSubjectKey InSubjectK
 		// clear all bEnabled only if found
 		if (FLiveLinkCollectionSubjectItem* NewEnabledItem = FindSubject(InSubjectKey))
 		{
-			UE::TUniqueLock Lock(SubjectsLock);
+			UE::TUniqueLock<UE::FRecursiveMutex> Lock(SubjectsLock);
 
 			NewEnabledItem->bEnabled = true;
 
@@ -325,7 +325,7 @@ void FLiveLinkSourceCollection::SetSubjectEnabled(FLiveLinkSubjectKey InSubjectK
 	}
 	else
 	{
-		UE::TUniqueLock Lock(SubjectsLock);
+		UE::TUniqueLock<UE::FRecursiveMutex> Lock(SubjectsLock);
 		for (FLiveLinkCollectionSubjectItem& SubjectItem : Subjects)
 		{
 			if (SubjectItem.Key.SubjectName == InSubjectKey.SubjectName)
@@ -380,11 +380,11 @@ void FLiveLinkSourceCollection::RemovePendingKill()
 bool FLiveLinkSourceCollection::RequestShutdown()
 {
 	{
-		UE::TUniqueLock Lock(SubjectsLock);
+		UE::TUniqueLock<UE::FRecursiveMutex> Lock(SubjectsLock);
 		Subjects.Reset();
 	}
 
-	UE::TUniqueLock Lock(SourcesLock);
+	UE::TUniqueLock<UE::FRecursiveMutex> Lock(SourcesLock);
 	for (int32 SourceIndex = Sources.Num() - 1; SourceIndex >= 0; --SourceIndex)
 	{
 		FLiveLinkCollectionSourceItem& SourceItem = Sources[SourceIndex];
@@ -429,7 +429,7 @@ void FLiveLinkSourceCollection::ForEachSubject(TFunctionRef<void(const FLiveLink
 
 void FLiveLinkSourceCollection::ForEachSource(TFunctionRef<void(FLiveLinkCollectionSourceItem&)> VisitorFunc)
 {
-	UE::TUniqueLock Lock(SourcesLock);
+	UE::TUniqueLock<UE::FRecursiveMutex> Lock(SourcesLock);
 
 	for (FLiveLinkCollectionSourceItem& Source : Sources)
 	{
@@ -440,7 +440,7 @@ void FLiveLinkSourceCollection::ForEachSource(TFunctionRef<void(FLiveLinkCollect
 
 void FLiveLinkSourceCollection::ForEachSource(TFunctionRef<void(const FLiveLinkCollectionSourceItem&)> VisitorFunc) const
 {
-	UE::TUniqueLock Lock(SourcesLock);
+	UE::TUniqueLock<UE::FRecursiveMutex> Lock(SourcesLock);
 
 	for (const FLiveLinkCollectionSourceItem& Source : Sources)
 	{
