Make sure we have valid static data before creating a subsection to store the data.

diff --git a/Engine/Plugins/Animation/LiveLink/Source/LiveLinkSequencer/Private/SequencerRecorderSections/MovieSceneLiveLinkTrackRecorder.cpp b/Engine/Plugins/Animation/LiveLink/Source/LiveLinkSequencer/Private/SequencerRecorderSections/MovieSceneLiveLinkTrackRecorder.cpp
--- a/Engine/Plugins/Animation/LiveLink/Source/LiveLinkSequencer/Private/SequencerRecorderSections/MovieSceneLiveLinkTrackRecorder.cpp
+++ b/Engine/Plugins/Animation/LiveLink/Source/LiveLinkSequencer/Private/SequencerRecorderSections/MovieSceneLiveLinkTrackRecorder.cpp
@@ -117,7 +117,7 @@ void UMovieSceneLiveLinkTrackRecorder::CreateTracks()
 																	, SubjectRole
 																	, StaticData.Get());
 
-	if(bRegistered)
+	if(bRegistered && StaticData->IsValid())
 	{
 		LiveLinkTrack = DoesLiveLinkTrackExist(SubjectName, SubjectRole);
 		if (!LiveLinkTrack.IsValid())
