Allow pushing frame data on virtual subjects

diff --git a/Engine/Plugins/Animation/LiveLink/Source/LiveLink/Private/LiveLinkClient.cpp b/Engine/Plugins/Animation/LiveLink/Source/LiveLink/Private/LiveLinkClient.cpp
--- a/Engine/Plugins/Animation/LiveLink/Source/LiveLink/Private/LiveLinkClient.cpp
+++ b/Engine/Plugins/Animation/LiveLink/Source/LiveLink/Private/LiveLinkClient.cpp
@@ -451,6 +451,7 @@ FGuid FLiveLinkClient::AddVirtualSubjectSource(FName SourceName)
 		NewSettings->SourceName = SourceName;
 		Data.Setting = TStrongObjectPtr(NewSettings);
 		Data.bIsVirtualSource = true;
+		Data.TimedData = MakeShared<FLiveLinkTimedDataInput>(this, Guid);
 		Collection->AddSource(MoveTemp(Data));
 
 		Source->ReceiveClient(this, Guid);
@@ -709,12 +710,6 @@ void FLiveLinkClient::PushSubjectStaticData_Internal(FPendingSubjectStatic&& Sub
 		return;
 	}
 
-	if (SourceItem->IsVirtualSource())
-	{
-		FLiveLinkLog::Error(TEXT("Trying to add frame data to the virtual subject '%s'."), *SubjectStaticData.Role->GetName(), *SubjectStaticData.SubjectKey.SubjectName.ToString());
-		return;
-	}
-
 	FLiveLinkSubject* LiveLinkSubject = nullptr;
 	{
 		if (FLiveLinkCollectionSubjectItem* SubjectItem = Collection->FindSubject(SubjectStaticData.SubjectKey))
