LiveLink - Allow reconnecting to UE instances running in -game

diff --git a/Engine/Plugins/Animation/LiveLink/Source/LiveLink/Private/LiveLinkMessageBusSource.cpp b/Engine/Plugins/Animation/LiveLink/Source/LiveLink/Private/LiveLinkMessageBusSource.cpp
--- a/Engine/Plugins/Animation/LiveLink/Source/LiveLink/Private/LiveLinkMessageBusSource.cpp
+++ b/Engine/Plugins/Animation/LiveLink/Source/LiveLink/Private/LiveLinkMessageBusSource.cpp
@@ -18,6 +18,15 @@
 #include "MessageEndpointBuilder.h"
 #include "Misc/App.h"
 
+
+int32 GLiveLinkMessageBusSourceReconnectAfterTimeout = 1;
+static FAutoConsoleVariableRef CVarLiveLinkMessageBusSourceReconnectAfterTimeout(
+	TEXT("LiveLink.MessageBus.Source.ReconnectAfterTimeout"),
+	GLiveLinkMessageBusSourceReconnectAfterTimeout,
+	TEXT("When enabled, when the connection times out, it will try to re-connect instead of removing the source."),
+	ECVF_Default);
+
+
 FText FLiveLinkMessageBusSource::ValidSourceStatus()
 {
 	return NSLOCTEXT("LiveLinkMessageBusSource", "ActiveStatus", "Active");
@@ -36,10 +45,13 @@ FText FLiveLinkMessageBusSource::TimeoutSourceStatus()
 FLiveLinkMessageBusSource::FLiveLinkMessageBusSource(const FText& InSourceType, const FText& InSourceMachineName, const FMessageAddress& InConnectionAddress, double InMachineTimeOffset)
 	: ConnectionAddress(InConnectionAddress)
 	, bIsValid(false)
+	, Client(nullptr)
 	, SourceType(InSourceType)
 	, SourceMachineName(InSourceMachineName)
 	, ConnectionLastActive(0.0)
 	, MachineTimeOffset(InMachineTimeOffset)
+	, bDiscovering(false)
+	, bInitialized(false )
 {}
 
 void FLiveLinkMessageBusSource::InitializeSettings(ULiveLinkSourceSettings* Settings)
@@ -50,28 +62,8 @@ void FLiveLinkMessageBusSource::ReceiveClient(ILiveLinkClient* InClient, FGuid I
 {
 	Client = InClient;
 	SourceGuid = InSourceGuid;
-	bIsValid = true;
-
-	for (const TSubclassOf<ULiveLinkRole>& RoleClass : FLiveLinkRoleTrait::GetRoles())
-	{
-		RoleInstances.Add(RoleClass->GetDefaultObject<ULiveLinkRole>());
-	}
-
-	MessageEndpoint = CreateAndInitializeMessageEndpoint();
-
-	if (ConnectionAddress.IsValid())
-	{
-		SendConnectMessage();
-	}
-	else
-	{
-#if WITH_LIVELINK_DISCOVERY_MANAGER_THREAD
-		ILiveLinkModule::Get().GetMessageBusDiscoveryManager().AddDiscoveryMessageRequest();
-#endif
-		bIsValid = false;
-	}
 
-	UpdateConnectionLastActive();
+	Initialize();
 }
 
 void FLiveLinkMessageBusSource::Update()
@@ -79,6 +71,8 @@ void FLiveLinkMessageBusSource::Update()
 	if (!ConnectionAddress.IsValid())
 	{
 #if WITH_LIVELINK_DISCOVERY_MANAGER_THREAD
+		StartDiscovering();
+
 		FLiveLinkMessageBusDiscoveryManager& DiscoveryManager = ILiveLinkModule::Get().GetMessageBusDiscoveryManager();
 		for (const FProviderPollResultPtr& Result : DiscoveryManager.GetDiscoveryResults())
 		{
@@ -87,7 +81,7 @@ void FLiveLinkMessageBusSource::Update()
 				ConnectionAddress = Result->Address;
 				SourceMachineName = FText::FromString(Result->MachineName);
 				MachineTimeOffset = Result->MachineTimeOffset;
-				DiscoveryManager.RemoveDiscoveryMessageRequest();
+				StopDiscovering();
 				SendConnectMessage();
 				UpdateConnectionLastActive();
 				break;
@@ -106,7 +100,13 @@ void FLiveLinkMessageBusSource::Update()
 			const double DeadSourceTimeout = GetDeadSourceTimeout();
 			if (CurrentTime - ConnectionLastActive > DeadSourceTimeout)
 			{
-				if (RequestSourceShutdown())
+				RequestSourceShutdown();
+
+				if (GLiveLinkMessageBusSourceReconnectAfterTimeout)
+				{
+					Initialize();
+				}
+				else
 				{
 					Client->RemoveSource(SourceGuid);
 				}
@@ -115,6 +115,35 @@ void FLiveLinkMessageBusSource::Update()
 	}
 }
 
+void FLiveLinkMessageBusSource::Initialize()
+{
+	if (bInitialized)
+	{
+		return;
+	}
+
+	for (const TSubclassOf<ULiveLinkRole>& RoleClass : FLiveLinkRoleTrait::GetRoles())
+	{
+		RoleInstances.Add(RoleClass->GetDefaultObject<ULiveLinkRole>());
+	}
+
+	MessageEndpoint = CreateAndInitializeMessageEndpoint();
+
+	if (ConnectionAddress.IsValid())
+	{
+		SendConnectMessage();
+	}
+	else
+	{
+		StartDiscovering();
+		bIsValid = false;
+	}
+
+	UpdateConnectionLastActive();
+
+	bInitialized = true;
+}
+
 void FLiveLinkMessageBusSource::InternalHandleMessage(const TSharedRef<IMessageContext, ESPMode::ThreadSafe>& Context)
 {
 	UpdateConnectionLastActive();
@@ -315,6 +344,32 @@ FORCEINLINE void FLiveLinkMessageBusSource::UpdateConnectionLastActive()
 	ConnectionLastActive = FPlatformTime::Seconds();
 }
 
+void FLiveLinkMessageBusSource::StartDiscovering()
+{
+#if WITH_LIVELINK_DISCOVERY_MANAGER_THREAD
+	if (bDiscovering)
+	{
+		return;
+	}
+
+	FLiveLinkMessageBusDiscoveryManager& DiscoveryManager = ILiveLinkModule::Get().GetMessageBusDiscoveryManager();
+	DiscoveryManager.AddDiscoveryMessageRequest();
+#endif
+}
+
+void FLiveLinkMessageBusSource::StopDiscovering()
+{
+#if WITH_LIVELINK_DISCOVERY_MANAGER_THREAD
+	if (!bDiscovering)
+	{
+		return;
+	}
+
+	FLiveLinkMessageBusDiscoveryManager& DiscoveryManager = ILiveLinkModule::Get().GetMessageBusDiscoveryManager();
+	DiscoveryManager.RemoveDiscoveryMessageRequest();
+#endif
+}
+
 void FLiveLinkMessageBusSource::SendConnectMessage()
 {
 	FLiveLinkConnectMessage* ConnectMessage = FMessageEndpoint::MakeMessage<FLiveLinkConnectMessage>();
@@ -326,13 +381,12 @@ void FLiveLinkMessageBusSource::SendConnectMessage()
 
 bool FLiveLinkMessageBusSource::RequestSourceShutdown()
 {
-#if WITH_LIVELINK_DISCOVERY_MANAGER_THREAD
-	FLiveLinkMessageBusDiscoveryManager& DiscoveryManager = ILiveLinkModule::Get().GetMessageBusDiscoveryManager();
-	if (DiscoveryManager.IsRunning() && !ConnectionAddress.IsValid())
+	if (!bInitialized)
 	{
-		DiscoveryManager.RemoveDiscoveryMessageRequest();
+		return true;
 	}
-#endif
+
+	StopDiscovering();
 
 	FLiveLinkHeartbeatEmitter& HeartbeatEmitter = ILiveLinkModule::Get().GetHeartbeatEmitter();
 	HeartbeatEmitter.StopHeartbeat(ConnectionAddress, MessageEndpoint);
@@ -343,6 +397,11 @@ bool FLiveLinkMessageBusSource::RequestSourceShutdown()
 		MessageEndpoint->Disable();
 	}
 	MessageEndpoint.Reset();
+	ConnectionAddress.Invalidate();
+
+	RoleInstances.Empty();
+
+	bInitialized = false;
 
 	return true;
 }
diff --git a/Engine/Plugins/Animation/LiveLink/Source/LiveLink/Public/LiveLinkMessageBusSource.h b/Engine/Plugins/Animation/LiveLink/Source/LiveLink/Public/LiveLinkMessageBusSource.h
--- a/Engine/Plugins/Animation/LiveLink/Source/LiveLink/Public/LiveLinkMessageBusSource.h
+++ b/Engine/Plugins/Animation/LiveLink/Source/LiveLink/Public/LiveLinkMessageBusSource.h
@@ -117,6 +117,16 @@ private:
 	// Threadsafe update of the last active time
 	FORCEINLINE void UpdateConnectionLastActive();
 
+	// Start discovering sources
+	void StartDiscovering();
+
+	// Stop discovering sources
+	void StopDiscovering();
+
+	// Initialize (can be used after a shutdown request)
+	void Initialize();
+
+private:
 	// List of the roles available when the bus was opened
 	TArray<TWeakObjectPtr<ULiveLinkRole>> RoleInstances;
 
@@ -131,4 +141,11 @@ private:
 
 	// Offset between sender's machine engine time and receiver's machine engine time
 	double MachineTimeOffset;
+
+	// Actively discovering the source
+	bool bDiscovering;
+
+	// Flag to keep track of initialization state
+	bool bInitialized;
 };
+
