LiveLink: Reduce default reconnection time from 30s to 5s.

This will make it reconnect more quickly to live link senders, now that the default behavior is to try to reconnect instead of dropping them from the list, which was causing unexpected issues for example when rebroadcasting from Editor instances to nDisplay nodes and the Editor would either be restarted or in a blocking state, and the nodes would stop using the streame data.

diff --git a/Engine/Plugins/Animation/LiveLink/Source/LiveLink/Private/LiveLinkSettings.cpp b/Engine/Plugins/Animation/LiveLink/Source/LiveLink/Private/LiveLinkSettings.cpp
--- a/Engine/Plugins/Animation/LiveLink/Source/LiveLink/Private/LiveLinkSettings.cpp
+++ b/Engine/Plugins/Animation/LiveLink/Source/LiveLink/Private/LiveLinkSettings.cpp
@@ -19,7 +19,7 @@ ULiveLinkSettings::ULiveLinkSettings()
 	, MessageBusPingRequestFrequency(1.0)
 	, MessageBusHeartbeatFrequency(1.0)
 	, MessageBusHeartbeatTimeout(2.0)
-	, MessageBusTimeBeforeRemovingInactiveSource(30.0)
+	, MessageBusTimeBeforeRemovingInactiveSource(5.0)
 	, TimeWithoutFrameToBeConsiderAsInvalid(0.5)
 	, ValidColor(FLinearColor::Green)
 	, InvalidColor(FLinearColor::Yellow)
