[AutoRTFM] Allow LiveLink code to be reached by a Verse transaction.

Verse can trigger cinematic playback and LiveLink is part of the animation system.

diff --git a/Engine/Plugins/Animation/LiveLink/Source/LiveLink/LiveLink.Build.cs b/Engine/Plugins/Animation/LiveLink/Source/LiveLink/LiveLink.Build.cs
--- a/Engine/Plugins/Animation/LiveLink/Source/LiveLink/LiveLink.Build.cs
+++ b/Engine/Plugins/Animation/LiveLink/Source/LiveLink/LiveLink.Build.cs
@@ -38,8 +38,6 @@ namespace UnrealBuildTool.Rules
 					"UnrealEd",
 				});
 			}
-
-			bDisableAutoRTFMInstrumentation = true;
 		}
 	}
 }
