Live Link  - Add context menu button to pause a subject
Paused subjects will retain their last subject snapshot

diff --git a/Engine/Plugins/Animation/LiveLink/Source/LiveLink/Private/LiveLinkBlueprintLibrary.cpp b/Engine/Plugins/Animation/LiveLink/Source/LiveLink/Private/LiveLinkBlueprintLibrary.cpp
--- a/Engine/Plugins/Animation/LiveLink/Source/LiveLink/Private/LiveLinkBlueprintLibrary.cpp
+++ b/Engine/Plugins/Animation/LiveLink/Source/LiveLink/Private/LiveLinkBlueprintLibrary.cpp
@@ -289,6 +289,16 @@ TSubclassOf<ULiveLinkRole> ULiveLinkBlueprintLibrary::GetLiveLinkSubjectRole(con
 	return TSubclassOf<ULiveLinkRole>();
 }
 
+void ULiveLinkBlueprintLibrary::PauseSubject(const FLiveLinkSubjectName SubjectName)
+{
+	IModularFeatures& ModularFeatures = IModularFeatures::Get();
+	if (ModularFeatures.IsModularFeatureAvailable(ILiveLinkClient::ModularFeatureName))
+	{
+		ILiveLinkClient& LiveLinkClient = ModularFeatures.GetModularFeature<ILiveLinkClient>(ILiveLinkClient::ModularFeatureName);
+		return LiveLinkClient.PauseSubject_AnyThread(SubjectName);
+	}
+}
+
 bool ULiveLinkBlueprintLibrary::EvaluateLiveLinkFrame(FLiveLinkSubjectRepresentation InSubjectRepresentation, FLiveLinkBaseBlueprintData& OutBlueprintData)
 {
 	// We should never hit this!  stubs to avoid NoExport on the class.
diff --git a/Engine/Plugins/Animation/LiveLink/Source/LiveLink/Private/LiveLinkClient.cpp b/Engine/Plugins/Animation/LiveLink/Source/LiveLink/Private/LiveLinkClient.cpp
--- a/Engine/Plugins/Animation/LiveLink/Source/LiveLink/Private/LiveLinkClient.cpp
+++ b/Engine/Plugins/Animation/LiveLink/Source/LiveLink/Private/LiveLinkClient.cpp
@@ -170,34 +170,12 @@ void FLiveLinkClient::BuildThisTicksSubjectSnapshot()
 				LiveSubject->Update();
 				EnabledSubjects.Add(SubjectItem.Key.SubjectName, SubjectItem.Key);
 
-				// Update Source FrameRate from first enabled subject with valid data.
-				if (LiveSubject->HasValidFrameSnapshot() && SubjectSettings->FrameRate.IsValid() && !TaggedSources.Contains(SubjectItem.Key.Source))
+				if (LiveSubject->IsPaused())
 				{
-					SourceSettings->BufferSettings.DetectedFrameRate = SubjectSettings->FrameRate;
-					TaggedSources.Add(SubjectItem.Key.Source);
+					HandleSubjectRebroadcast(LiveSubject, LiveSubject->GetFrameSnapshot().FrameData);
 				}
-			}
-			else
-			{
-				LiveSubject->ClearFrames();
-			}
-		}
-	});
-
-	Collection->ForEachSubject([this, &TaggedSources](FLiveLinkCollectionSourceItem& SourceItem, FLiveLinkCollectionSubjectItem& SubjectItem)
-	{
-		if (FLiveLinkSubject* LiveSubject = SubjectItem.GetLiveSubject())
-		{
-			if (SubjectItem.bEnabled)
-			{
-				ULiveLinkSourceSettings* SourceSettings = SourceItem.Setting.Get();
-				ULiveLinkSubjectSettings* SubjectSettings = SubjectItem.GetLinkSettings();
 
-				LiveSubject->CacheSettings(SourceSettings, SubjectSettings);
-				LiveSubject->Update();
-				EnabledSubjects.Add(SubjectItem.Key.SubjectName, SubjectItem.Key);
-
-				// Update Source FrameRate from first enabled subject with valid data. 
+				// Update Source FrameRate from first enabled subject with valid data.
 				if (LiveSubject->HasValidFrameSnapshot() && SubjectSettings->FrameRate.IsValid() && !TaggedSources.Contains(SubjectItem.Key.Source))
 				{
 					SourceSettings->BufferSettings.DetectedFrameRate = SubjectSettings->FrameRate;
@@ -393,6 +371,11 @@ ELiveLinkSubjectState FLiveLinkClient::GetSubjectState(FLiveLinkSubjectName InSu
 		return ELiveLinkSubjectState::InvalidOrDisabled;
 	}
 
+	if (LiveSubject->IsPaused())
+	{
+		return ELiveLinkSubjectState::Paused;
+	}
+
 	switch (const ETimedDataInputState InputState = LiveSubject->GetState())
 	{
 	case ETimedDataInputState::Connected: return ELiveLinkSubjectState::Connected;
@@ -1020,7 +1003,11 @@ void FLiveLinkClient::PushSubjectFrameData_Internal(FPendingSubjectFrame&& Subje
 	const bool bHasParentSubject = SourceItem->Setting->ParentSubject != FLiveLinkSubjectName();
 	if (!bHasParentSubject)
 	{
-		HandleSubjectRebroadcast(LinkSubject, SubjectFrameData.FrameData);
+		if (!LinkSubject->IsPaused())
+		{
+			// If it's paused, rebroadcast will be handled in FLiveLinkClient::BuildThisTicksSubjectSnapshot
+			HandleSubjectRebroadcast(LinkSubject, SubjectFrameData.FrameData);
+		}
 
 		if (bEnableParentSubjects)
 		{
@@ -1145,6 +1132,34 @@ void FLiveLinkClient::RemoveSubject_AnyThread(const FLiveLinkSubjectKey& InSubje
 	}
 }
 
+void FLiveLinkClient::PauseSubject_AnyThread(FLiveLinkSubjectName SubjectName)
+{
+	if (Collection)
+	{
+		if (const FLiveLinkCollectionSubjectItem* SubjectItem = Collection->FindSubject(SubjectName))
+		{
+			if (FLiveLinkSubject* LiveSubject = SubjectItem->GetLiveSubject())
+			{
+				if (!LiveSubject->IsPaused())
+				{
+					if (IsSubjectValid(SubjectName))
+					{
+						LiveSubject->PauseSubject();
+					}
+					else
+					{
+						UE_LOG(LogLiveLink, Warning, TEXT("Could not pause subject %s since it's not in a valid state."), *SubjectName.ToString());
+					}
+				}
+				else
+				{
+					LiveSubject->PauseSubject();
+				}
+			}
+		}
+	}
+}
+
 bool FLiveLinkClient::AddVirtualSubject(const FLiveLinkSubjectKey& InVirtualSubjectKey, TSubclassOf<ULiveLinkVirtualSubject> InVirtualSubjectClass)
 {
 	bool bResult = false;
diff --git a/Engine/Plugins/Animation/LiveLink/Source/LiveLink/Private/LiveLinkSubject.cpp b/Engine/Plugins/Animation/LiveLink/Source/LiveLink/Private/LiveLinkSubject.cpp
--- a/Engine/Plugins/Animation/LiveLink/Source/LiveLink/Private/LiveLinkSubject.cpp
+++ b/Engine/Plugins/Animation/LiveLink/Source/LiveLink/Private/LiveLinkSubject.cpp
@@ -79,6 +79,16 @@ void FLiveLinkSubject::Update()
 		RemoveFrames(Count);
 	}
 
+	if (bPaused)
+	{
+		// No need to process if it's paused, just keep the same frame snapshot.
+		if (HasValidFrameSnapshot())
+		{
+			FrameSnapshot.FrameData.GetBaseData()->WorldTime = FPlatformTime::Seconds();
+		}
+		return;
+	}
+
 	if (GetMode() == ELiveLinkSourceMode::EngineTime)
 	{
 		if (CachedSettings.BufferSettings.bValidEngineTimeEnabled)
@@ -477,6 +487,16 @@ void FLiveLinkSubject::PreprocessFrame(FLiveLinkFrameDataStruct& InOutFrameData)
 	}
 }
 
+bool FLiveLinkSubject::IsPaused() const
+{
+	return bPaused;
+}
+
+void FLiveLinkSubject::PauseSubject()
+{
+	bPaused = !bPaused;
+}
+
 void FLiveLinkSubject::PreprocessFrame(const FLiveLinkStaticDataStruct& InStaticData, FLiveLinkFrameDataStruct& InOutFrameData)
 {
 	for (ULiveLinkFramePreProcessor::FWorkerSharedPtr PreProcessor : FramePreProcessors)
diff --git a/Engine/Plugins/Animation/LiveLink/Source/LiveLink/Public/LiveLinkBlueprintLibrary.h b/Engine/Plugins/Animation/LiveLink/Source/LiveLink/Public/LiveLinkBlueprintLibrary.h
--- a/Engine/Plugins/Animation/LiveLink/Source/LiveLink/Public/LiveLinkBlueprintLibrary.h
+++ b/Engine/Plugins/Animation/LiveLink/Source/LiveLink/Public/LiveLinkBlueprintLibrary.h
@@ -187,6 +187,10 @@ public:
 	UFUNCTION(BlueprintCallable, Category = "LiveLink")
 	static TSubclassOf<ULiveLinkRole> GetLiveLinkSubjectRole(const FLiveLinkSubjectName SubjectName);
 
+	/** Pause / unpause a subject, freezing its evaluation to the last data it received. */
+	UFUNCTION(BlueprintCallable, Category = "LiveLink")
+	static void PauseSubject(const FLiveLinkSubjectName SubjectName);
+
 	UE_DEPRECATED(4.23, "EvaluateLiveLinkFrame with Subject Represention is deprecated, recreate the node.")
 	UFUNCTION(BlueprintCallable, CustomThunk, Category = LiveLink, meta = (DeprecatedFunction, CustomStructureParam = "OutBlueprintData", BlueprintInternalUseOnly = "true", AllowAbstract = "false"))
 	static bool EvaluateLiveLinkFrame(FLiveLinkSubjectRepresentation SubjectRepresentation, FLiveLinkBaseBlueprintData& OutBlueprintData);
diff --git a/Engine/Plugins/Animation/LiveLink/Source/LiveLink/Public/LiveLinkClient.h b/Engine/Plugins/Animation/LiveLink/Source/LiveLink/Public/LiveLinkClient.h
--- a/Engine/Plugins/Animation/LiveLink/Source/LiveLink/Public/LiveLinkClient.h
+++ b/Engine/Plugins/Animation/LiveLink/Source/LiveLink/Public/LiveLinkClient.h
@@ -61,6 +61,7 @@ public:
 	virtual bool AddVirtualSubject(const FLiveLinkSubjectKey& VirtualSubjectKey, TSubclassOf<ULiveLinkVirtualSubject> VirtualSubjectClass) override;
 	virtual void RemoveVirtualSubject(const FLiveLinkSubjectKey& VirtualSubjectKey) override;
 	virtual void RemoveSubject_AnyThread(const FLiveLinkSubjectKey& SubjectKey) override;
+	virtual void PauseSubject_AnyThread(FLiveLinkSubjectName SubjectName) override;
 	virtual void ClearSubjectsFrames_AnyThread(FLiveLinkSubjectName SubjectName) override;
 	virtual void ClearSubjectsFrames_AnyThread(const FLiveLinkSubjectKey& InSubjectKey) override;
 	virtual void ClearAllSubjectsFrames_AnyThread() override;
diff --git a/Engine/Plugins/Animation/LiveLink/Source/LiveLink/Public/LiveLinkSubject.h b/Engine/Plugins/Animation/LiveLink/Source/LiveLink/Public/LiveLinkSubject.h
--- a/Engine/Plugins/Animation/LiveLink/Source/LiveLink/Public/LiveLinkSubject.h
+++ b/Engine/Plugins/Animation/LiveLink/Source/LiveLink/Public/LiveLinkSubject.h
@@ -84,10 +84,13 @@ public:
 	UE_DEPRECATED(5.6, "Replaced with a new version of PreprocessFrame that also provides const static data.")
 	virtual void PreprocessFrame(FLiveLinkFrameDataStruct& InOutFrameData) override;
 	
-	virtual void PreprocessFrame(const FLiveLinkStaticDataStruct &InStaticData, FLiveLinkFrameDataStruct& InOutFrameData) override;
+	virtual void PreprocessFrame(const FLiveLinkStaticDataStruct& InStaticData, FLiveLinkFrameDataStruct& InOutFrameData) override;
+	virtual bool IsPaused() const override;
+	virtual void PauseSubject() override;
 
 protected:
 	virtual const FLiveLinkSubjectFrameData& GetFrameSnapshot() const override { return FrameSnapshot; }
+	friend class FLiveLinkClient;
 	//~ End ILiveLinkSubject Interface
 
 	//~Begin ITimedDataSource Interface
@@ -261,6 +264,9 @@ private:
 
 	/** If true, override static data may remap when caching settings. */
 	bool bNeedsStaticRemap = false;
+	
+	/** Flag set to indicate that a subject is currently paused, so it should keep its last snapshot. */
+	std::atomic<bool> bPaused = false;
 
 	/** Flag set to clear the override static data for a subject. */
 	std::atomic<bool> bClearOverrideStaticData = false;
diff --git a/Engine/Plugins/Animation/LiveLink/Source/LiveLinkEditor/Private/LiveLinkClientCommands.cpp b/Engine/Plugins/Animation/LiveLink/Source/LiveLinkEditor/Private/LiveLinkClientCommands.cpp
--- a/Engine/Plugins/Animation/LiveLink/Source/LiveLinkEditor/Private/LiveLinkClientCommands.cpp
+++ b/Engine/Plugins/Animation/LiveLink/Source/LiveLinkEditor/Private/LiveLinkClientCommands.cpp
@@ -14,6 +14,7 @@ void FLiveLinkClientCommands::RegisterCommands()
 	UI_COMMAND(RemoveSource, "Remove Selected Source(s)", "Remove selected live link source", EUserInterfaceActionType::Button, FInputChord());
 	UI_COMMAND(RemoveAllSources, "Remove All Sources", "Remove all live link sources", EUserInterfaceActionType::Button, FInputChord());
 	UI_COMMAND(RemoveSubject, "Remove Subject", "Remove selected live link subject", EUserInterfaceActionType::Button, FInputChord());
+	UI_COMMAND(PauseSubject, "Pause/Unpause Subject", "Pause or unpause selected live link subject", EUserInterfaceActionType::Button, FInputChord());
 }
 
 #undef LOCTEXT_NAMESPACE
diff --git a/Engine/Plugins/Animation/LiveLink/Source/LiveLinkEditor/Private/LiveLinkClientCommands.h b/Engine/Plugins/Animation/LiveLink/Source/LiveLinkEditor/Private/LiveLinkClientCommands.h
--- a/Engine/Plugins/Animation/LiveLink/Source/LiveLinkEditor/Private/LiveLinkClientCommands.h
+++ b/Engine/Plugins/Animation/LiveLink/Source/LiveLinkEditor/Private/LiveLinkClientCommands.h
@@ -12,6 +12,7 @@ public:
 	TSharedPtr<FUICommandInfo> RemoveSource;
 	TSharedPtr<FUICommandInfo> RemoveAllSources;
 	TSharedPtr<FUICommandInfo> RemoveSubject;
+	TSharedPtr<FUICommandInfo> PauseSubject;
 
 	/** Initialize commands */
 	virtual void RegisterCommands() override;
diff --git a/Engine/Plugins/Animation/LiveLink/Source/LiveLinkEditor/Private/LiveLinkClientPanelViews.cpp b/Engine/Plugins/Animation/LiveLink/Source/LiveLinkEditor/Private/LiveLinkClientPanelViews.cpp
--- a/Engine/Plugins/Animation/LiveLink/Source/LiveLinkEditor/Private/LiveLinkClientPanelViews.cpp
+++ b/Engine/Plugins/Animation/LiveLink/Source/LiveLinkEditor/Private/LiveLinkClientPanelViews.cpp
@@ -179,6 +179,15 @@ void FLiveLinkSubjectUIEntry::RemoveFromClient() const
 	Client->RemoveSubject_AnyThread(SubjectKey);
 }
 
+void FLiveLinkSubjectUIEntry::PauseSubject()
+{
+	Client->PauseSubject_AnyThread(SubjectKey.SubjectName);
+}
+
+bool FLiveLinkSubjectUIEntry::IsPaused() const
+{
+	return Client->GetSubjectState(SubjectKey.SubjectName) == ELiveLinkSubjectState::Paused;
+}
 
 class SLiveLinkClientPanelSubjectRow : public SMultiColumnTableRow<FLiveLinkSubjectUIEntryPtr>
 {
@@ -560,6 +569,8 @@ TSharedPtr<SWidget> FLiveLinkSubjectsView::OnOpenVirtualSubjectContextMenu(TShar
 	}
 	MenuBuilder.EndSection();
 
+	MenuBuilder.AddMenuEntry(FLiveLinkClientCommands::Get().PauseSubject, NAME_None, TAttribute<FText>::CreateSP(this, &FLiveLinkSubjectsView::GetPauseSubjectLabel), TAttribute<FText>::CreateSP(this, &FLiveLinkSubjectsView::GetPauseSubjectToolTip));
+
 	return MenuBuilder.MakeWidget();
 }
 
@@ -643,6 +654,24 @@ void FLiveLinkSubjectsView::RefreshSubjects()
 	}
 }
 
+bool FLiveLinkSubjectsView::CanPauseSubject() const
+{
+	TArray<FLiveLinkSubjectUIEntryPtr> Selected;
+	SubjectsTreeView->GetSelectedItems(Selected);
+
+	return Selected.Num() > 0 && Selected[0] && Selected[0]->IsSubjectValid();
+}
+
+void FLiveLinkSubjectsView::HandlePauseSubject()
+{
+	TArray<FLiveLinkSubjectUIEntryPtr> Selected;
+	SubjectsTreeView->GetSelectedItems(Selected);
+	if (Selected.Num())
+	{
+		Selected[0]->PauseSubject();
+	}
+}
+
 void FLiveLinkSubjectsView::CreateSubjectsTreeView(const TSharedPtr<FUICommandList>& InCommandList)
 {
 	SAssignNew(SubjectsTreeView, SLiveLinkSubjectsTreeView, bReadOnly)
@@ -670,5 +699,32 @@ void FLiveLinkSubjectsView::CreateSubjectsTreeView(const TSharedPtr<FUICommandLi
 		);
 }
 
+FText FLiveLinkSubjectsView::GetPauseSubjectLabel() const
+{
+	FText Label = LOCTEXT("PauseSubjectLabel", "Pause Subject");
+
+	TArray<FLiveLinkSubjectUIEntryPtr> Selected;
+	SubjectsTreeView->GetSelectedItems(Selected);
+	if (Selected.Num() && Selected[0]->IsPaused())
+	{
+		Label = LOCTEXT("UnpauseSubjectLabel", "Unpause Subject");
+	}
+
+	return Label;
+}
+
+FText FLiveLinkSubjectsView::GetPauseSubjectToolTip() const
+{
+	FText ToolTip = LOCTEXT("PauseSubjectToolTip", "Pause a subject, the last received data will be used for evaluation.");
+
+	TArray<FLiveLinkSubjectUIEntryPtr> Selected;
+	SubjectsTreeView->GetSelectedItems(Selected);
+	if (Selected.Num() && Selected[0]->IsPaused())
+	{
+		ToolTip = LOCTEXT("UnpauseSubjectToolTip", "Unpause Subject and resume operating on live data.");
+	}
+
+	return ToolTip;
+}
 
 #undef LOCTEXT_NAMESPACE /**LiveLinkClientPanel*/
diff --git a/Engine/Plugins/Animation/LiveLink/Source/LiveLinkEditor/Private/LiveLinkPanelController.cpp b/Engine/Plugins/Animation/LiveLink/Source/LiveLinkEditor/Private/LiveLinkPanelController.cpp
--- a/Engine/Plugins/Animation/LiveLink/Source/LiveLinkEditor/Private/LiveLinkPanelController.cpp
+++ b/Engine/Plugins/Animation/LiveLink/Source/LiveLinkEditor/Private/LiveLinkPanelController.cpp
@@ -164,6 +164,11 @@ void FLiveLinkPanelController::BindCommands()
 		FExecuteAction::CreateRaw(this, &FLiveLinkPanelController::HandleRemoveSubject),
 		FCanExecuteAction::CreateRaw(this, &FLiveLinkPanelController::CanRemoveSubject)
 	);
+
+	CommandList->MapAction(FLiveLinkClientCommands::Get().PauseSubject,
+		FExecuteAction::CreateRaw(this, &FLiveLinkPanelController::HandlePauseSubject),
+		FCanExecuteAction::CreateRaw(this, &FLiveLinkPanelController::CanPauseSubject)
+	);
 }
 
 void FLiveLinkPanelController::OnSourcesChangedHandler()
@@ -239,3 +244,18 @@ void FLiveLinkPanelController::HandleRemoveSubject()
 		Selected[0]->RemoveFromClient();
 	}
 }
+
+bool FLiveLinkPanelController::CanPauseSubject() const
+{
+	return SubjectsView->CanPauseSubject();
+}
+
+void FLiveLinkPanelController::HandlePauseSubject()
+{
+	TArray<FLiveLinkSubjectUIEntryPtr> Selected;
+	SubjectsView->SubjectsTreeView->GetSelectedItems(Selected);
+	if (Selected.Num() > 0 && Selected[0])
+	{
+		Selected[0]->PauseSubject();
+	}
+}
diff --git a/Engine/Plugins/Animation/LiveLink/Source/LiveLinkEditor/Public/LiveLinkClientPanelViews.h b/Engine/Plugins/Animation/LiveLink/Source/LiveLinkEditor/Public/LiveLinkClientPanelViews.h
--- a/Engine/Plugins/Animation/LiveLink/Source/LiveLinkEditor/Public/LiveLinkClientPanelViews.h
+++ b/Engine/Plugins/Animation/LiveLink/Source/LiveLinkEditor/Public/LiveLinkClientPanelViews.h
@@ -69,6 +69,10 @@ struct FLiveLinkSubjectUIEntry
 	TSubclassOf<ULiveLinkRole> GetItemRole() const;
 	// Remove the subject or source from the livelink client
 	void RemoveFromClient() const;
+	// Pause or unpause a subject if it's already paused
+	void PauseSubject();
+	// Returns whether a subject is currently paused.
+	bool IsPaused() const;
 
 private:
 	// Whether the subject is virtual
@@ -180,11 +184,18 @@ public:
 	bool CanRemoveSubject() const;
 	// Refresh the list of subjects using the livelink client.
 	void RefreshSubjects();
+	// Return whether a subject can be paused
+	bool CanPauseSubject() const;
+	// Pauses/Unpauses a subject
+	void HandlePauseSubject();
 
 private:
 	// Create the subjects tree view
 	void CreateSubjectsTreeView(const TSharedPtr<FUICommandList>& InCommandList);
 
+	FText GetPauseSubjectLabel() const;
+	FText GetPauseSubjectToolTip() const;
+
 public:
 	// Subject tree widget
 	TSharedPtr<class SLiveLinkSubjectsTreeView> SubjectsTreeView;
diff --git a/Engine/Plugins/Animation/LiveLink/Source/LiveLinkEditor/Public/LiveLinkPanelController.h b/Engine/Plugins/Animation/LiveLink/Source/LiveLinkEditor/Public/LiveLinkPanelController.h
--- a/Engine/Plugins/Animation/LiveLink/Source/LiveLinkEditor/Public/LiveLinkPanelController.h
+++ b/Engine/Plugins/Animation/LiveLink/Source/LiveLinkEditor/Public/LiveLinkPanelController.h
@@ -55,6 +55,10 @@ private:
 	bool CanRemoveSource() const;
 	// Returns whether a subject can be removed.
 	bool CanRemoveSubject() const;
+	// Returns whethera subject can be paused
+	bool CanPauseSubject() const;
+	// Pauses/Unpauses a subject.
+	void HandlePauseSubject();
 	// Handles the remove subject command.
 	void HandleRemoveSubject();
 	// Recreates the source list data behind the list view.
diff --git a/Engine/Source/Runtime/LiveLinkInterface/Private/LiveLinkVirtualSubject.cpp b/Engine/Source/Runtime/LiveLinkInterface/Private/LiveLinkVirtualSubject.cpp
--- a/Engine/Source/Runtime/LiveLinkInterface/Private/LiveLinkVirtualSubject.cpp
+++ b/Engine/Source/Runtime/LiveLinkInterface/Private/LiveLinkVirtualSubject.cpp
@@ -60,6 +60,16 @@ TArray<FLiveLinkTime> ULiveLinkVirtualSubject::GetFrameTimes() const
 	return Result;
 }
 
+bool ULiveLinkVirtualSubject::IsPaused() const
+{
+	return bPaused;
+}
+
+void ULiveLinkVirtualSubject::PauseSubject()
+{
+	bPaused = !bPaused;
+}
+
 bool ULiveLinkVirtualSubject::HasValidStaticData() const
 {
 	return CurrentFrameSnapshot.StaticData.IsValid();
diff --git a/Engine/Source/Runtime/LiveLinkInterface/Public/ILiveLinkClient.h b/Engine/Source/Runtime/LiveLinkInterface/Public/ILiveLinkClient.h
--- a/Engine/Source/Runtime/LiveLinkInterface/Public/ILiveLinkClient.h
+++ b/Engine/Source/Runtime/LiveLinkInterface/Public/ILiveLinkClient.h
@@ -32,6 +32,8 @@ enum class ELiveLinkSubjectState : uint8
 	Disconnected,
 	/** The subject is invalid or disabled */
 	InvalidOrDisabled,
+	/** The subject is currently paused. */
+	Paused,
 	/** The state of the subject is unknown. e.g. It cannot be queried */
 	Unknown
 };
@@ -127,6 +129,9 @@ public:
 	/** Clear the subject from the specific source */
 	virtual void RemoveSubject_AnyThread(const FLiveLinkSubjectKey& SubjectName) = 0;
 
+	/** Pause/Unpause a subject. It will keep it's subject snapshot until it's unpaused. */
+	virtual void PauseSubject_AnyThread(FLiveLinkSubjectName SubjectName) = 0;
+
 	/** Clear the stored frames associated with the enabled subject */
 	virtual void ClearSubjectsFrames_AnyThread(FLiveLinkSubjectName SubjectName) = 0;
 
diff --git a/Engine/Source/Runtime/LiveLinkInterface/Public/ILiveLinkSubject.h b/Engine/Source/Runtime/LiveLinkInterface/Public/ILiveLinkSubject.h
--- a/Engine/Source/Runtime/LiveLinkInterface/Public/ILiveLinkSubject.h
+++ b/Engine/Source/Runtime/LiveLinkInterface/Public/ILiveLinkSubject.h
@@ -65,6 +65,12 @@ public:
 
 	/** Apply a remapper to a frame data. Called after preprocessing.*/
 	virtual void RemapFrame(FLiveLinkSkeletonStaticData& InOutSkeletonData, FLiveLinkAnimationFrameData& InOutFrameData) {}
+	
+	/** Pause/Unpause subject. */
+	virtual bool IsPaused() const = 0;
+
+	/** Pause/Unpause subject. */
+	virtual void PauseSubject() = 0;
 
 protected:
 	virtual const FLiveLinkSubjectFrameData& GetFrameSnapshot() const = 0;
diff --git a/Engine/Source/Runtime/LiveLinkInterface/Public/LiveLinkVirtualSubject.h b/Engine/Source/Runtime/LiveLinkInterface/Public/LiveLinkVirtualSubject.h
--- a/Engine/Source/Runtime/LiveLinkInterface/Public/LiveLinkVirtualSubject.h
+++ b/Engine/Source/Runtime/LiveLinkInterface/Public/LiveLinkVirtualSubject.h
@@ -45,6 +45,8 @@ public:
 	virtual bool IsRebroadcasted() const override { return bRebroadcastSubject; }
 	virtual bool HasStaticDataBeenRebroadcasted() const override { return bHasStaticDataBeenRebroadcast; }
 	virtual void SetStaticDataAsRebroadcasted(const bool bInSent) override { bHasStaticDataBeenRebroadcast = bInSent; }
+	LIVELINKINTERFACE_API virtual bool IsPaused() const override;
+	LIVELINKINTERFACE_API virtual void PauseSubject() override;
 
 protected:
 	virtual const FLiveLinkSubjectFrameData& GetFrameSnapshot() const override { return CurrentFrameSnapshot; }
@@ -154,4 +156,6 @@ private:
 
 	/** Last evaluated frame for this subject. */
 	FLiveLinkSubjectFrameData CurrentFrameSnapshot;
+
+	std::atomic<bool> bPaused = false;
 };
