LiveLink  - Allow accessing DefaultRoleSettings and DefaultSettingsClass publicly

diff --git a/Engine/Plugins/Animation/LiveLink/Source/LiveLink/Public/LiveLinkSettings.h b/Engine/Plugins/Animation/LiveLink/Source/LiveLink/Public/LiveLinkSettings.h
--- a/Engine/Plugins/Animation/LiveLink/Source/LiveLink/Public/LiveLinkSettings.h
+++ b/Engine/Plugins/Animation/LiveLink/Source/LiveLink/Public/LiveLinkSettings.h
@@ -75,7 +75,7 @@ public:
 	virtual void PostInitProperties() override;
 	//~ End UObject Interface
 
-protected:
+public:
 	UPROPERTY(config, EditAnywhere, Category="LiveLink")
 	TArray<FLiveLinkRoleProjectSetting> DefaultRoleSettings;
 
@@ -83,7 +83,6 @@ protected:
 	UPROPERTY(config)
 	FSoftClassPath DefaultSettingsClass;
 
-public:
 	/** The interpolation class to use for new Subjects if no specific settings we set for the Subject's role. */
 	UPROPERTY(config)
 	TSubclassOf<ULiveLinkFrameInterpolationProcessor> FrameInterpolationProcessor;
