[AutoRTFM] Allow LiveLink to run properly when called from a Verse transaction.

This marks our first time encountering `FRecursiveMutex` in transactional code. Fortunately, from a functional perspective it is safe to replace this with `FTransactionallySafeCriticalSection` when within a transaction.

In a followup we can investigate a smarter implementation, but this is safest for a post-HL fix.

diff --git a/Engine/Plugins/Animation/LiveLink/Source/LiveLink/Private/LiveLinkSourceCollection.cpp b/Engine/Plugins/Animation/LiveLink/Source/LiveLink/Private/LiveLinkSourceCollection.cpp
--- a/Engine/Plugins/Animation/LiveLink/Source/LiveLink/Private/LiveLinkSourceCollection.cpp
+++ b/Engine/Plugins/Animation/LiveLink/Source/LiveLink/Private/LiveLinkSourceCollection.cpp
@@ -100,7 +100,7 @@ void FLiveLinkSourceCollection::RemoveSource(FGuid InSourceGuid)
 {
 	if (InSourceGuid != FLiveLinkSourceCollection::DefaultVirtualSubjectGuid)
 	{
-		UE::TMultiUniqueLock<UE::FRecursiveMutex> MultiLock({&SubjectsLock, &SourcesLock});
+		UE::TMultiUniqueLock<UE::FTransactionallySafeRecursiveMutex> MultiLock({&SubjectsLock, &SourcesLock});
 
 		int32 SourceIndex = Sources.IndexOfByPredicate([InSourceGuid](const FLiveLinkCollectionSourceItem& Other) { return Other.Guid == InSourceGuid; });
 		if (SourceIndex != INDEX_NONE)
@@ -339,7 +339,7 @@ void FLiveLinkSourceCollection::SetSubjectEnabled(FLiveLinkSubjectKey InSubjectK
 
 void FLiveLinkSourceCollection::RemovePendingKill()
 {
-	UE::TMultiUniqueLock<UE::FRecursiveMutex> MultiLock({&SubjectsLock, &SourcesLock});
+	UE::TMultiUniqueLock<UE::FTransactionallySafeRecursiveMutex> MultiLock({&SubjectsLock, &SourcesLock});
 
 	for (int32 SourceIndex = Sources.Num() - 1; SourceIndex >= 0; --SourceIndex)
 	{
@@ -401,7 +401,7 @@ bool FLiveLinkSourceCollection::RequestShutdown()
 
 void FLiveLinkSourceCollection::ForEachSubject(TFunctionRef<void(FLiveLinkCollectionSourceItem&, FLiveLinkCollectionSubjectItem&)> VisitorFunc)
 {
-	UE::TMultiUniqueLock<UE::FRecursiveMutex> MultiLock({&SubjectsLock, &SourcesLock});
+	UE::TMultiUniqueLock<UE::FTransactionallySafeRecursiveMutex> MultiLock({&SubjectsLock, &SourcesLock});
 
 	for (FLiveLinkCollectionSubjectItem& Subject : Subjects)
 	{
@@ -415,7 +415,7 @@ void FLiveLinkSourceCollection::ForEachSubject(TFunctionRef<void(FLiveLinkCollec
 
 void FLiveLinkSourceCollection::ForEachSubject(TFunctionRef<void(const FLiveLinkCollectionSourceItem&, const FLiveLinkCollectionSubjectItem&)> VisitorFunc) const
 {
-	UE::TMultiUniqueLock<UE::FRecursiveMutex> MultiLock({ &SubjectsLock, &SourcesLock });
+	UE::TMultiUniqueLock<UE::FTransactionallySafeRecursiveMutex> MultiLock({ &SubjectsLock, &SourcesLock });
 
 	for (const FLiveLinkCollectionSubjectItem& Subject : Subjects)
 	{
diff --git a/Engine/Plugins/Animation/LiveLink/Source/LiveLink/Public/LiveLinkSourceCollection.h b/Engine/Plugins/Animation/LiveLink/Source/LiveLink/Public/LiveLinkSourceCollection.h
--- a/Engine/Plugins/Animation/LiveLink/Source/LiveLink/Public/LiveLinkSourceCollection.h
+++ b/Engine/Plugins/Animation/LiveLink/Source/LiveLink/Public/LiveLinkSourceCollection.h
@@ -3,7 +3,7 @@
 #pragma once
 
 #include "Async/Async.h"
-#include "Async/RecursiveMutex.h"
+#include "Async/TransactionallySafeRecursiveMutex.h"
 #include "ILiveLinkClient.h"
 #include "LiveLinkSubject.h"
 #include "LiveLinkSubjectSettings.h"
@@ -202,8 +202,8 @@ private:
 	FOnLiveLinkSubjectStateChanged OnLiveLinkSubjectStateChangedDelegate;
 
 	/** Lock to stop multiple threads accessing the Subjects from the collection at the same time */
-	mutable UE::FRecursiveMutex SubjectsLock;
+	mutable UE::FTransactionallySafeRecursiveMutex SubjectsLock;
 
 	/** Lock to stop multiple threads accessing the Sources from the collection at the same time */
-	mutable UE::FRecursiveMutex SourcesLock;
+	mutable UE::FTransactionallySafeRecursiveMutex SourcesLock;
 };
diff --git a/Engine/Source/Runtime/Core/Public/Async/TransactionallySafeRecursiveMutex.h b/Engine/Source/Runtime/Core/Public/Async/TransactionallySafeRecursiveMutex.h
new file mode 100644
--- /dev/null
+++ b/Engine/Source/Runtime/Core/Public/Async/TransactionallySafeRecursiveMutex.h
@@ -0,0 +1,21 @@
+// Copyright Epic Games, Inc. All Rights Reserved.
+
+#pragma once
+
+#include "Async/RecursiveMutex.h"
+#include "AutoRTFM/AutoRTFM.h"
+
+#if UE_AUTORTFM
+#include "Misc/TransactionallySafeCriticalSection.h"
+#endif
+
+namespace UE
+{
+
+#if UE_AUTORTFM
+using FTransactionallySafeRecursiveMutex = ::FTransactionallySafeCriticalSection;
+#else
+using FTransactionallySafeRecursiveMutex = ::UE::FRecursiveMutex;
+#endif // UE_AUTORTFM
+
+} // UE
