LiveLinkHub - Fix eject button style

#rnx

diff --git a/Engine/Plugins/Animation/LiveLink/Source/LiveLinkEditor/Private/LiveLinkEditorModule.cpp b/Engine/Plugins/Animation/LiveLink/Source/LiveLinkEditor/Private/LiveLinkEditorModule.cpp
--- a/Engine/Plugins/Animation/LiveLink/Source/LiveLinkEditor/Private/LiveLinkEditorModule.cpp
+++ b/Engine/Plugins/Animation/LiveLink/Source/LiveLinkEditor/Private/LiveLinkEditorModule.cpp
@@ -160,15 +160,6 @@ public:
 
 		StyleSet->Set("EditableTextBox", EditableTextStyle);
 
-		FButtonStyle EjectButtonStyle = FAppStyle::Get().GetWidgetStyle<FButtonStyle>("Animation.PlayControlsButton");
-
-		const FCheckBoxStyle& CheckBoxStyle = FAppStyle::Get().GetWidgetStyle<FCheckBoxStyle>("ToggleButtonCheckbox");
-		EjectButtonStyle.Normal = CheckBoxStyle.CheckedImage;
-		EjectButtonStyle.Hovered = CheckBoxStyle.CheckedHoveredImage;
-
-		StyleSet->Set("LiveLinkHub.EjectButton", EjectButtonStyle);
-
-
 		FSlateStyleRegistry::RegisterSlateStyle(*StyleSet.Get());
 
 		RegisterSettings();
diff --git a/Engine/Plugins/Animation/LiveLinkHub/Source/LiveLinkHub/Private/UI/Widgets/SLiveLinkHubPlaybackWidget.cpp b/Engine/Plugins/Animation/LiveLinkHub/Source/LiveLinkHub/Private/UI/Widgets/SLiveLinkHubPlaybackWidget.cpp
--- a/Engine/Plugins/Animation/LiveLinkHub/Source/LiveLinkHub/Private/UI/Widgets/SLiveLinkHubPlaybackWidget.cpp
+++ b/Engine/Plugins/Animation/LiveLinkHub/Source/LiveLinkHub/Private/UI/Widgets/SLiveLinkHubPlaybackWidget.cpp
@@ -216,10 +216,11 @@ void SLiveLinkHubPlaybackWidget::Construct(const FArguments& InArgs)
 				]
 			]
 			+ SHorizontalBox::Slot()
+			.Padding(FMargin(12.0, 0.0, 0.0, 0.0))
 			.AutoWidth()
 			[
 				SNew(SButton)
-				.ButtonStyle(LiveLinkStyle, "LiveLinkHub.EjectButton")
+				.ButtonStyle(FAppStyle::Get(), "Animation.PlayControlsButton")
 				.ContentPadding(FMargin(2.0f))
 				.OnClicked(this, &SLiveLinkHubPlaybackWidget::OnExitPlaybackPressed)
 				.ToolTipText(LOCTEXT("EjectButtonToolTip", "Exit Playback"))
