LiveLink Remapper - Add a function to validate that the remapper is compatible with static data

diff --git a/Engine/Plugins/Animation/LiveLink/Source/LiveLink/Private/LiveLinkSubject.cpp b/Engine/Plugins/Animation/LiveLink/Source/LiveLink/Private/LiveLinkSubject.cpp
--- a/Engine/Plugins/Animation/LiveLink/Source/LiveLink/Private/LiveLinkSubject.cpp
+++ b/Engine/Plugins/Animation/LiveLink/Source/LiveLink/Private/LiveLinkSubject.cpp
@@ -1209,21 +1209,28 @@ void FLiveLinkSubject::CacheSettings(ULiveLinkSourceSettings* SourceSetting, ULi
 						RemappedStaticData.InitializeWith(StaticData);
 						NewRemapper->RemapStaticData(RemappedStaticData);
 
-						if (HasValidFrameSnapshot())
+						if (NewRemapper->IsRemapperCompatible(StaticData))
 						{
-							NewRemapper->RemapStaticData(FrameSnapshot.StaticData);
-							NewRemapper->RemapFrameData(FrameSnapshot.StaticData, FrameSnapshot.FrameData);
-						}
+							if (HasValidFrameSnapshot())
+							{
+								NewRemapper->RemapStaticData(FrameSnapshot.StaticData);
+								NewRemapper->RemapFrameData(FrameSnapshot.StaticData, FrameSnapshot.FrameData);
+							}
 
-						// Important: Because we changed the static data, we have to update the frame data to match the new static data.
-						for (FLiveLinkFrameDataStruct& FrameDataStruct : FrameData)
-						{
-							NewRemapper->RemapFrameData(RemappedStaticData, FrameDataStruct);
-						}
+							// Important: Because we changed the static data, we have to update the frame data to match the new static data.
+							for (FLiveLinkFrameDataStruct& FrameDataStruct : FrameData)
+							{
+								NewRemapper->RemapFrameData(RemappedStaticData, FrameDataStruct);
+							}
 
-						OverrideStaticData = MoveTemp(RemappedStaticData);
+							OverrideStaticData = MoveTemp(RemappedStaticData);
 
-						SetStaticDataAsRebroadcasted(false);
+							SetStaticDataAsRebroadcasted(false);
+						}
+						else
+						{
+							UE_LOG(LogLiveLink, Warning, TEXT("Remapper is not compatible with this static data."));
+						}
 					}
 					else
 					{
diff --git a/Engine/Source/Runtime/LiveLinkInterface/Public/LiveLinkSubjectRemapper.h b/Engine/Source/Runtime/LiveLinkInterface/Public/LiveLinkSubjectRemapper.h
--- a/Engine/Source/Runtime/LiveLinkInterface/Public/LiveLinkSubjectRemapper.h
+++ b/Engine/Source/Runtime/LiveLinkInterface/Public/LiveLinkSubjectRemapper.h
@@ -24,6 +24,9 @@ public:
 
 	/** Remap the frame data of a subject. */
 	LIVELINKINTERFACE_API virtual void RemapFrameData(const FLiveLinkStaticDataStruct& InStaticData, FLiveLinkFrameDataStruct& InOutFrameData) { }
+
+	/** Returns if the remapper is compatible with the static data. */
+	LIVELINKINTERFACE_API virtual bool IsRemapperCompatible(const FLiveLinkStaticDataStruct& InStaticData) { return false;  }
 };
 
 /** Class used to remap livelink subjects without having to rely on animation blueprints. */
