LiveLink - Replace green tick with Checkmark and Warning icon
Add pause icon
Tweak icon colors
Switch to alternating row colors in Subject, Sources and Clients list
Add LLHub icon to LiveLinkHub's major tab

diff --git a/Engine/Plugins/Animation/LiveLink/Source/LiveLink/Private/LiveLinkSettings.cpp b/Engine/Plugins/Animation/LiveLink/Source/LiveLink/Private/LiveLinkSettings.cpp
--- a/Engine/Plugins/Animation/LiveLink/Source/LiveLink/Private/LiveLinkSettings.cpp
+++ b/Engine/Plugins/Animation/LiveLink/Source/LiveLink/Private/LiveLinkSettings.cpp
@@ -21,8 +21,9 @@ ULiveLinkSettings::ULiveLinkSettings()
 	, MessageBusHeartbeatTimeout(2.0)
 	, MessageBusTimeBeforeRemovingInactiveSource(5.0)
 	, TimeWithoutFrameToBeConsiderAsInvalid(0.5)
-	, ValidColor(FLinearColor::Green)
-	, InvalidColor(FLinearColor::Yellow)
+	, ValidColor(FColor(31, 228, 75))
+	, InvalidColor(FColor(255, 184, 0))
+	, PausedColor(FColor(152, 152, 152))
 	, TextSizeSource(16)
 	, TextSizeSubject(12)
 {
diff --git a/Engine/Plugins/Animation/LiveLink/Source/LiveLink/Public/LiveLinkSettings.h b/Engine/Plugins/Animation/LiveLink/Source/LiveLink/Public/LiveLinkSettings.h
--- a/Engine/Plugins/Animation/LiveLink/Source/LiveLink/Public/LiveLinkSettings.h
+++ b/Engine/Plugins/Animation/LiveLink/Source/LiveLink/Public/LiveLinkSettings.h
@@ -146,6 +146,10 @@ public:
 	UPROPERTY(config, EditAnywhere, Category = "LiveLink|UI")
 	FLinearColor InvalidColor;
 
+	/** Color for Subjects that are currently paused. */
+	UPROPERTY(config, EditAnywhere, Category = "LiveLink|UI")
+	FLinearColor PausedColor;
+
 	/** Font size of Source names shown in LiveLink Debug View. */
 	UPROPERTY(config, EditAnywhere, Category = "LiveLink|UI|Debug")
 	uint8 TextSizeSource;
@@ -163,6 +167,7 @@ public:
 	double GetTimeWithoutFrameToBeConsiderAsInvalid() const { return TimeWithoutFrameToBeConsiderAsInvalid; }
 	FLinearColor GetValidColor() const { return ValidColor; }
 	FLinearColor GetInvalidColor() const { return InvalidColor; }
+	FLinearColor GetPausedColor() const { return PausedColor; }
 	float GetMessageBusPingRequestFrequency() const { return MessageBusPingRequestFrequency; }
 	float GetMessageBusHeartbeatFrequency() const { return MessageBusHeartbeatFrequency; }
 	double GetMessageBusHeartbeatTimeout() const { return MessageBusHeartbeatTimeout; }
diff --git a/Engine/Plugins/Animation/LiveLink/Source/LiveLinkEditor/Private/LiveLinkClientPanelViews.cpp b/Engine/Plugins/Animation/LiveLink/Source/LiveLinkEditor/Private/LiveLinkClientPanelViews.cpp
--- a/Engine/Plugins/Animation/LiveLink/Source/LiveLinkEditor/Private/LiveLinkClientPanelViews.cpp
+++ b/Engine/Plugins/Animation/LiveLink/Source/LiveLinkEditor/Private/LiveLinkClientPanelViews.cpp
@@ -17,6 +17,7 @@
 #include "Modules/ModuleManager.h"
 #include "PropertyEditorModule.h"
 #include "SLiveLinkDataView.h"
+#include "Styling/SlateStyleRegistry.h"
 #include "Widgets/DeclarativeSyntaxSupport.h"
 #include "Widgets/Input/SButton.h"
 #include "Widgets/Input/SCheckBox.h"
@@ -204,8 +205,14 @@ public:
 		EntryPtr = Args._Entry;
 		bReadOnly = Args._ReadOnly;
 
+		const ISlateStyle* LiveLinkStyle = FSlateStyleRegistry::FindSlateStyle("LiveLinkStyle");
+		OkayIcon = LiveLinkStyle->GetBrush("LiveLink.Subject.Okay");
+		WarningIcon = LiveLinkStyle->GetBrush("LiveLink.Subject.Warning");
+		PauseIcon = LiveLinkStyle->GetBrush("LiveLink.Subject.Pause");
+
 		SMultiColumnTableRow<FLiveLinkSubjectUIEntryPtr>::Construct(
 			FSuperRowType::FArguments()
+			.Style(&FAppStyle().GetWidgetStyle<FTableRowStyle>("TableView.AlternatingRow"))
 			.Padding(1.0f),
 			OwnerTableView
 		);
@@ -237,6 +244,7 @@ public:
 				.FillWidth(1.0f)
 				[
 					SNew(STextBlock)
+					.ColorAndOpacity(this, &SLiveLinkClientPanelSubjectRow::GetSubjectTextColor)
 					.Text(MakeAttributeSP(this, &SLiveLinkClientPanelSubjectRow::GetItemText))
 				];
 		}
@@ -244,6 +252,7 @@ public:
 		{
 			auto RoleAttribute = MakeAttributeSP(this, &SLiveLinkClientPanelSubjectRow::GetItemRole);
 			return SNew(STextBlock)
+				.ColorAndOpacity(this, &SLiveLinkClientPanelSubjectRow::GetSubjectTextColor)
 				.Text(EntryPtr->IsSubject() ? RoleAttribute : FText::GetEmpty());
 		}
 		else if (ColumnName == SubjectTreeUI::ActionsColumnName)
@@ -272,10 +281,10 @@ public:
 					.HAlign(HAlign_Center)
 					.VAlign(VAlign_Center)
 					[
-						SNew(STextBlock)
-						.Font(FAppStyle::Get().GetFontStyle("FontAwesome.8"))
+						SNew(SImage)
+						.Image(this, &SLiveLinkClientPanelSubjectRow::GetSubjectIcon)
 						.ColorAndOpacity(this, &SLiveLinkClientPanelSubjectRow::OnGetActivityColor)
-						.Text(FEditorFontGlyphs::Circle)
+						.ToolTipText(this, &SLiveLinkClientPanelSubjectRow::GetSubjectIconToolTip)
 					];
 			}
 		}
@@ -307,7 +316,14 @@ private:
 	{
 		if (EntryPtr->IsSubjectEnabled())
 		{
-			return EntryPtr->IsSubjectValid() ? GetDefault<ULiveLinkSettings>()->GetValidColor() : GetDefault<ULiveLinkSettings>()->GetInvalidColor();
+			if (!EntryPtr->IsPaused())
+			{
+				return EntryPtr->IsSubjectValid() ? GetDefault<ULiveLinkSettings>()->GetValidColor() : GetDefault<ULiveLinkSettings>()->GetInvalidColor();
+			}
+			else
+			{
+				return GetDefault<ULiveLinkSettings>()->GetPausedColor();
+			}
 		}
 
 		return FLinearColor(0.f, 0.f, 0.f, 0.f);
@@ -319,10 +335,69 @@ private:
 		return bReadOnly.Get() ? EVisibility::Collapsed : EVisibility::Visible;
 	}
 
+	/** Get the icon for a subject's status. */
+	const FSlateBrush* GetSubjectIcon() const
+	{
+		const FSlateBrush* IconBrush = nullptr;
+
+		if (EntryPtr->IsSubjectEnabled())
+		{
+			if (!EntryPtr->IsPaused())
+			{
+				IconBrush = EntryPtr->IsSubjectValid() ? OkayIcon : WarningIcon;
+			}
+			else
+			{
+				IconBrush = PauseIcon;
+			}
+		}
+		else
+		{
+			// No icon for disabled subjects, we rely on setting the text color to subdued foreground.
+		}
+
+		return IconBrush;
+	}
+
+	/** Get the tooltip for a subject's status icon. */
+	FText GetSubjectIconToolTip() const
+	{
+		FText IconToolTip;
+
+		if (EntryPtr->IsSubjectEnabled())
+		{
+			IconToolTip = EntryPtr->IsSubjectValid() ? LOCTEXT("ValidSubjectToolTip", "Subject is operating normally.") : LOCTEXT("InvalidSubjectToolTip", "Subject is invalid.");
+		}
+		else
+		{
+			IconToolTip = LOCTEXT("SubjectDisabledToolTip", "Subject is disabled.");
+		}
+
+		return IconToolTip;
+	}
+
+	/** Get the text color for a subject. */
+	FSlateColor GetSubjectTextColor() const
+	{
+		FSlateColor TextColor = FSlateColor::UseForeground();
+
+		if (EntryPtr->IsSubject() && !EntryPtr->IsSubjectEnabled())
+		{
+			TextColor = FSlateColor::UseSubduedForeground();
+		}
+
+		return TextColor;
+	}
+
 	FLiveLinkSubjectUIEntryPtr EntryPtr;
 
 	/** Returns whether the panel is in read-only mode. */
 	TAttribute<bool> bReadOnly;
+
+	//~ Status icons
+	const FSlateBrush* OkayIcon = nullptr;
+	const FSlateBrush* WarningIcon = nullptr;
+	const FSlateBrush* PauseIcon = nullptr;
 };
 
 class SLiveLinkClientPanelSourcesRow : public SMultiColumnTableRow<FLiveLinkSourceUIEntryPtr>
@@ -341,6 +416,7 @@ public:
 
 		SMultiColumnTableRow<FLiveLinkSourceUIEntryPtr>::Construct(
 			FSuperRowType::FArguments()
+			.Style(&FAppStyle().GetWidgetStyle<FTableRowStyle>("TableView.AlternatingRow"))
 			.Padding(1.0f),
 			OwnerTableView
 		);
diff --git a/Engine/Plugins/Animation/LiveLink/Source/LiveLinkEditor/Private/LiveLinkEditorModule.cpp b/Engine/Plugins/Animation/LiveLink/Source/LiveLinkEditor/Private/LiveLinkEditorModule.cpp
--- a/Engine/Plugins/Animation/LiveLink/Source/LiveLinkEditor/Private/LiveLinkEditorModule.cpp
+++ b/Engine/Plugins/Animation/LiveLink/Source/LiveLinkEditor/Private/LiveLinkEditorModule.cpp
@@ -101,6 +101,11 @@ public:
 		StyleSet->SetContentRoot(FPaths::EngineContentDir() / TEXT("Editor/Slate"));
 		StyleSet->SetCoreContentRoot(FPaths::EngineContentDir() / TEXT("Slate"));
 
+		StyleSet->Set("LiveLink.Subject.Okay", new IMAGE_PLUGIN_BRUSH_SVG(TEXT("Starship/Checkmark"), Icon16x16));
+		StyleSet->Set("LiveLink.Subject.Warning", new IMAGE_PLUGIN_BRUSH_SVG(TEXT("Starship/Warning"), Icon16x16));
+		StyleSet->Set("LiveLink.Subject.Pause", new IMAGE_PLUGIN_BRUSH_SVG(TEXT("Starship/Pause"), Icon16x16));
+		StyleSet->Set("LiveLink.Subject.Error", new IMAGE_PLUGIN_BRUSH_SVG(TEXT("Starship/Error"), Icon16x16));
+
 		StyleSet->Set("LiveLinkClient.Common.Icon", new IMAGE_PLUGIN_BRUSH(TEXT("LiveLink_40x"), Icon40x40));
 		StyleSet->Set("LiveLinkClient.Common.Icon.Small", new IMAGE_PLUGIN_BRUSH_SVG("Starship/LiveLink", Icon16x16));
 
diff --git a/Engine/Plugins/Animation/LiveLinkHub/Source/LiveLinkHub/Private/Clients/SLiveLinkHubClientsView.h b/Engine/Plugins/Animation/LiveLinkHub/Source/LiveLinkHub/Private/Clients/SLiveLinkHubClientsView.h
--- a/Engine/Plugins/Animation/LiveLinkHub/Source/LiveLinkHub/Private/Clients/SLiveLinkHubClientsView.h
+++ b/Engine/Plugins/Animation/LiveLinkHub/Source/LiveLinkHub/Private/Clients/SLiveLinkHubClientsView.h
@@ -244,6 +244,7 @@ public:
 
 		SMultiColumnTableRow<TSharedPtr<FClientTreeViewItem>>::Construct(
 			FSuperRowType::FArguments()
+			.Style(&FAppStyle().GetWidgetStyle<FTableRowStyle>("TableView.AlternatingRow"))
 			.Padding(1.0f),
 			InOwnerTableView
 		);
diff --git a/Engine/Plugins/Animation/LiveLinkHub/Source/LiveLinkHub/Private/UI/Window/LiveLinkHubWindowController.cpp b/Engine/Plugins/Animation/LiveLinkHub/Source/LiveLinkHub/Private/UI/Window/LiveLinkHubWindowController.cpp
--- a/Engine/Plugins/Animation/LiveLinkHub/Source/LiveLinkHub/Private/UI/Window/LiveLinkHubWindowController.cpp
+++ b/Engine/Plugins/Animation/LiveLinkHub/Source/LiveLinkHub/Private/UI/Window/LiveLinkHubWindowController.cpp
@@ -129,7 +129,9 @@ void FLiveLinkHubWindowController::RestoreLayout(TSharedPtr<FAssetEditorToolkit>
 			SNew(SLiveLinkHubMemoryStats)
 		];
 
-		FGlobalTabmanager::Get()->FindExistingLiveTab(StandaloneTabId)->SetTitleBarRightContent(RightMenuWidget);
+	TSharedPtr<SDockTab> MainTab = FGlobalTabmanager::Get()->FindExistingLiveTab(StandaloneTabId);
+	MainTab->SetTitleBarRightContent(RightMenuWidget);
+	MainTab->SetTabIcon(FSlateIcon("LiveLinkStyle", "LiveLinkHub.Icon.Small").GetIcon());
 }
 
 TSharedPtr<FModalWindowManager> FLiveLinkHubWindowController::InitializeSlateApplication()
