LIveLink - Split PauseSubject into two methods, PauseSubject and Unpause Subject

diff --git a/Engine/Plugins/Animation/LiveLink/Source/LiveLink/Private/LiveLinkBlueprintLibrary.cpp b/Engine/Plugins/Animation/LiveLink/Source/LiveLink/Private/LiveLinkBlueprintLibrary.cpp
--- a/Engine/Plugins/Animation/LiveLink/Source/LiveLink/Private/LiveLinkBlueprintLibrary.cpp
+++ b/Engine/Plugins/Animation/LiveLink/Source/LiveLink/Private/LiveLinkBlueprintLibrary.cpp
@@ -295,7 +295,17 @@ void ULiveLinkBlueprintLibrary::PauseSubject(const FLiveLinkSubjectName SubjectN
 	if (ModularFeatures.IsModularFeatureAvailable(ILiveLinkClient::ModularFeatureName))
 	{
 		ILiveLinkClient& LiveLinkClient = ModularFeatures.GetModularFeature<ILiveLinkClient>(ILiveLinkClient::ModularFeatureName);
-		return LiveLinkClient.PauseSubject_AnyThread(SubjectName);
+		LiveLinkClient.PauseSubject_AnyThread(SubjectName);
+	}
+}
+
+void ULiveLinkBlueprintLibrary::UnpauseSubject(const FLiveLinkSubjectName SubjectName)
+{
+	IModularFeatures& ModularFeatures = IModularFeatures::Get();
+	if (ModularFeatures.IsModularFeatureAvailable(ILiveLinkClient::ModularFeatureName))
+	{
+		ILiveLinkClient& LiveLinkClient = ModularFeatures.GetModularFeature<ILiveLinkClient>(ILiveLinkClient::ModularFeatureName);
+		LiveLinkClient.UnpauseSubject_AnyThread(SubjectName);
 	}
 }
 
diff --git a/Engine/Plugins/Animation/LiveLink/Source/LiveLink/Private/LiveLinkClient.cpp b/Engine/Plugins/Animation/LiveLink/Source/LiveLink/Private/LiveLinkClient.cpp
--- a/Engine/Plugins/Animation/LiveLink/Source/LiveLink/Private/LiveLinkClient.cpp
+++ b/Engine/Plugins/Animation/LiveLink/Source/LiveLink/Private/LiveLinkClient.cpp
@@ -1151,10 +1151,20 @@ void FLiveLinkClient::PauseSubject_AnyThread(FLiveLinkSubjectName SubjectName)
 						UE_LOG(LogLiveLink, Warning, TEXT("Could not pause subject %s since it's not in a valid state."), *SubjectName.ToString());
 					}
 				}
-				else
-				{
-					LiveSubject->PauseSubject();
-				}
+			}
+		}
+	}
+}
+
+void FLiveLinkClient::UnpauseSubject_AnyThread(FLiveLinkSubjectName SubjectName)
+{
+	if (Collection)
+	{
+		if (const FLiveLinkCollectionSubjectItem* SubjectItem = Collection->FindSubject(SubjectName))
+		{
+			if (FLiveLinkSubject* LiveSubject = SubjectItem->GetLiveSubject())
+			{
+				LiveSubject->UnpauseSubject();
 			}
 		}
 	}
diff --git a/Engine/Plugins/Animation/LiveLink/Source/LiveLink/Private/LiveLinkSubject.cpp b/Engine/Plugins/Animation/LiveLink/Source/LiveLink/Private/LiveLinkSubject.cpp
--- a/Engine/Plugins/Animation/LiveLink/Source/LiveLink/Private/LiveLinkSubject.cpp
+++ b/Engine/Plugins/Animation/LiveLink/Source/LiveLink/Private/LiveLinkSubject.cpp
@@ -494,7 +494,12 @@ bool FLiveLinkSubject::IsPaused() const
 
 void FLiveLinkSubject::PauseSubject()
 {
-	bPaused = !bPaused;
+	bPaused = true;
+}
+
+void FLiveLinkSubject::UnpauseSubject()
+{
+	bPaused = false;
 }
 
 void FLiveLinkSubject::PreprocessFrame(const FLiveLinkStaticDataStruct& InStaticData, FLiveLinkFrameDataStruct& InOutFrameData)
diff --git a/Engine/Plugins/Animation/LiveLink/Source/LiveLink/Public/LiveLinkBlueprintLibrary.h b/Engine/Plugins/Animation/LiveLink/Source/LiveLink/Public/LiveLinkBlueprintLibrary.h
--- a/Engine/Plugins/Animation/LiveLink/Source/LiveLink/Public/LiveLinkBlueprintLibrary.h
+++ b/Engine/Plugins/Animation/LiveLink/Source/LiveLink/Public/LiveLinkBlueprintLibrary.h
@@ -187,9 +187,13 @@ public:
 	UFUNCTION(BlueprintCallable, Category = "LiveLink")
 	static TSubclassOf<ULiveLinkRole> GetLiveLinkSubjectRole(const FLiveLinkSubjectName SubjectName);
 
-	/** Pause / unpause a subject, freezing its evaluation to the last data it received. */
+	/** Pause a subject, freezing its evaluation to the last data it received. */
 	UFUNCTION(BlueprintCallable, Category = "LiveLink")
 	static void PauseSubject(const FLiveLinkSubjectName SubjectName);
+	
+	/** Unpause a subject, resuming its normal operation. */
+	UFUNCTION(BlueprintCallable, Category = "LiveLink")
+	static void UnpauseSubject(const FLiveLinkSubjectName SubjectName);
 
 	UE_DEPRECATED(4.23, "EvaluateLiveLinkFrame with Subject Represention is deprecated, recreate the node.")
 	UFUNCTION(BlueprintCallable, CustomThunk, Category = LiveLink, meta = (DeprecatedFunction, CustomStructureParam = "OutBlueprintData", BlueprintInternalUseOnly = "true", AllowAbstract = "false"))
diff --git a/Engine/Plugins/Animation/LiveLink/Source/LiveLink/Public/LiveLinkClient.h b/Engine/Plugins/Animation/LiveLink/Source/LiveLink/Public/LiveLinkClient.h
--- a/Engine/Plugins/Animation/LiveLink/Source/LiveLink/Public/LiveLinkClient.h
+++ b/Engine/Plugins/Animation/LiveLink/Source/LiveLink/Public/LiveLinkClient.h
@@ -62,6 +62,7 @@ public:
 	virtual void RemoveVirtualSubject(const FLiveLinkSubjectKey& VirtualSubjectKey) override;
 	virtual void RemoveSubject_AnyThread(const FLiveLinkSubjectKey& SubjectKey) override;
 	virtual void PauseSubject_AnyThread(FLiveLinkSubjectName SubjectName) override;
+	virtual void UnpauseSubject_AnyThread(FLiveLinkSubjectName SubjectName) override;
 	virtual void ClearSubjectsFrames_AnyThread(FLiveLinkSubjectName SubjectName) override;
 	virtual void ClearSubjectsFrames_AnyThread(const FLiveLinkSubjectKey& InSubjectKey) override;
 	virtual void ClearAllSubjectsFrames_AnyThread() override;
@@ -306,3 +307,4 @@ private:
 	TOptional<FQualifiedFrameTime> CachedEngineFrameTime;
 #endif
 };
+
diff --git a/Engine/Plugins/Animation/LiveLink/Source/LiveLink/Public/LiveLinkSubject.h b/Engine/Plugins/Animation/LiveLink/Source/LiveLink/Public/LiveLinkSubject.h
--- a/Engine/Plugins/Animation/LiveLink/Source/LiveLink/Public/LiveLinkSubject.h
+++ b/Engine/Plugins/Animation/LiveLink/Source/LiveLink/Public/LiveLinkSubject.h
@@ -87,6 +87,7 @@ public:
 	virtual void PreprocessFrame(const FLiveLinkStaticDataStruct& InStaticData, FLiveLinkFrameDataStruct& InOutFrameData) override;
 	virtual bool IsPaused() const override;
 	virtual void PauseSubject() override;
+	virtual void UnpauseSubject() override;
 
 protected:
 	virtual const FLiveLinkSubjectFrameData& GetFrameSnapshot() const override { return FrameSnapshot; }
diff --git a/Engine/Plugins/Animation/LiveLink/Source/LiveLinkEditor/Private/LiveLinkClientPanelViews.cpp b/Engine/Plugins/Animation/LiveLink/Source/LiveLinkEditor/Private/LiveLinkClientPanelViews.cpp
--- a/Engine/Plugins/Animation/LiveLink/Source/LiveLinkEditor/Private/LiveLinkClientPanelViews.cpp
+++ b/Engine/Plugins/Animation/LiveLink/Source/LiveLinkEditor/Private/LiveLinkClientPanelViews.cpp
@@ -182,7 +182,14 @@ void FLiveLinkSubjectUIEntry::RemoveFromClient() const
 
 void FLiveLinkSubjectUIEntry::PauseSubject()
 {
-	Client->PauseSubject_AnyThread(SubjectKey.SubjectName);
+	if (IsPaused())
+	{
+		Client->UnpauseSubject_AnyThread(SubjectKey.SubjectName);
+	}
+	else
+	{
+		Client->PauseSubject_AnyThread(SubjectKey.SubjectName);
+	}
 }
 
 bool FLiveLinkSubjectUIEntry::IsPaused() const
@@ -779,9 +786,7 @@ FText FLiveLinkSubjectsView::GetPauseSubjectLabel() const
 {
 	FText Label = LOCTEXT("PauseSubjectLabel", "Pause Subject");
 
-	TArray<FLiveLinkSubjectUIEntryPtr> Selected;
-	SubjectsTreeView->GetSelectedItems(Selected);
-	if (Selected.Num() && Selected[0]->IsPaused())
+	if (IsSelectedSubjectPaused())
 	{
 		Label = LOCTEXT("UnpauseSubjectLabel", "Unpause Subject");
 	}
@@ -793,9 +798,7 @@ FText FLiveLinkSubjectsView::GetPauseSubjectToolTip() const
 {
 	FText ToolTip = LOCTEXT("PauseSubjectToolTip", "Pause a subject, the last received data will be used for evaluation.");
 
-	TArray<FLiveLinkSubjectUIEntryPtr> Selected;
-	SubjectsTreeView->GetSelectedItems(Selected);
-	if (Selected.Num() && Selected[0]->IsPaused())
+	if (IsSelectedSubjectPaused())
 	{
 		ToolTip = LOCTEXT("UnpauseSubjectToolTip", "Unpause Subject and resume operating on live data.");
 	}
@@ -803,4 +806,15 @@ FText FLiveLinkSubjectsView::GetPauseSubjectToolTip() const
 	return ToolTip;
 }
 
+bool FLiveLinkSubjectsView::IsSelectedSubjectPaused() const
+{
+	TArray<FLiveLinkSubjectUIEntryPtr> Selected;
+	SubjectsTreeView->GetSelectedItems(Selected);
+	if (Selected.Num())
+	{
+		return Selected[0]->IsPaused();
+	}
+	return false;
+}
+
 #undef LOCTEXT_NAMESPACE /**LiveLinkClientPanel*/
diff --git a/Engine/Plugins/Animation/LiveLink/Source/LiveLinkEditor/Public/LiveLinkClientPanelViews.h b/Engine/Plugins/Animation/LiveLink/Source/LiveLinkEditor/Public/LiveLinkClientPanelViews.h
--- a/Engine/Plugins/Animation/LiveLink/Source/LiveLinkEditor/Public/LiveLinkClientPanelViews.h
+++ b/Engine/Plugins/Animation/LiveLink/Source/LiveLinkEditor/Public/LiveLinkClientPanelViews.h
@@ -193,8 +193,12 @@ private:
 	// Create the subjects tree view
 	void CreateSubjectsTreeView(const TSharedPtr<FUICommandList>& InCommandList);
 
+	/** Get the label for the pause subject context menu entry. */
 	FText GetPauseSubjectLabel() const;
+	/** Get the tooltip for the pause subject context menu entry. */
 	FText GetPauseSubjectToolTip() const;
+	/** Returns whether the selected subject is paused. */
+	bool IsSelectedSubjectPaused() const;
 
 public:
 	// Subject tree widget
diff --git a/Engine/Source/Runtime/LiveLinkInterface/Private/LiveLinkVirtualSubject.cpp b/Engine/Source/Runtime/LiveLinkInterface/Private/LiveLinkVirtualSubject.cpp
--- a/Engine/Source/Runtime/LiveLinkInterface/Private/LiveLinkVirtualSubject.cpp
+++ b/Engine/Source/Runtime/LiveLinkInterface/Private/LiveLinkVirtualSubject.cpp
@@ -67,7 +67,12 @@ bool ULiveLinkVirtualSubject::IsPaused() const
 
 void ULiveLinkVirtualSubject::PauseSubject()
 {
-	bPaused = !bPaused;
+	bPaused = true;
+}
+
+void ULiveLinkVirtualSubject::UnpauseSubject()
+{
+	bPaused = false;
 }
 
 bool ULiveLinkVirtualSubject::HasValidStaticData() const
diff --git a/Engine/Source/Runtime/LiveLinkInterface/Public/ILiveLinkClient.h b/Engine/Source/Runtime/LiveLinkInterface/Public/ILiveLinkClient.h
--- a/Engine/Source/Runtime/LiveLinkInterface/Public/ILiveLinkClient.h
+++ b/Engine/Source/Runtime/LiveLinkInterface/Public/ILiveLinkClient.h
@@ -129,9 +129,12 @@ public:
 	/** Clear the subject from the specific source */
 	virtual void RemoveSubject_AnyThread(const FLiveLinkSubjectKey& SubjectName) = 0;
 
-	/** Pause/Unpause a subject. It will keep it's subject snapshot until it's unpaused. */
+	/** Pause a subject. It will keep its subject snapshot until it's unpaused. */
 	virtual void PauseSubject_AnyThread(FLiveLinkSubjectName SubjectName) = 0;
 
+	/** Unpause a subject, resuming its normal operation. */
+	virtual void UnpauseSubject_AnyThread(FLiveLinkSubjectName SubjectName) = 0;
+
 	/** Clear the stored frames associated with the enabled subject */
 	virtual void ClearSubjectsFrames_AnyThread(FLiveLinkSubjectName SubjectName) = 0;
 
diff --git a/Engine/Source/Runtime/LiveLinkInterface/Public/ILiveLinkSubject.h b/Engine/Source/Runtime/LiveLinkInterface/Public/ILiveLinkSubject.h
--- a/Engine/Source/Runtime/LiveLinkInterface/Public/ILiveLinkSubject.h
+++ b/Engine/Source/Runtime/LiveLinkInterface/Public/ILiveLinkSubject.h
@@ -69,9 +69,12 @@ public:
 	/** Pause/Unpause subject. */
 	virtual bool IsPaused() const = 0;
 
-	/** Pause/Unpause subject. */
+	/** Pause subject. */
 	virtual void PauseSubject() = 0;
 
+	/** Unpause subject. */
+	virtual void UnpauseSubject() = 0;
+
 protected:
 	virtual const FLiveLinkSubjectFrameData& GetFrameSnapshot() const = 0;
 
diff --git a/Engine/Source/Runtime/LiveLinkInterface/Public/LiveLinkVirtualSubject.h b/Engine/Source/Runtime/LiveLinkInterface/Public/LiveLinkVirtualSubject.h
--- a/Engine/Source/Runtime/LiveLinkInterface/Public/LiveLinkVirtualSubject.h
+++ b/Engine/Source/Runtime/LiveLinkInterface/Public/LiveLinkVirtualSubject.h
@@ -47,6 +47,7 @@ public:
 	virtual void SetStaticDataAsRebroadcasted(const bool bInSent) override { bHasStaticDataBeenRebroadcast = bInSent; }
 	LIVELINKINTERFACE_API virtual bool IsPaused() const override;
 	LIVELINKINTERFACE_API virtual void PauseSubject() override;
+	LIVELINKINTERFACE_API virtual void UnpauseSubject() override;
 
 protected:
 	virtual const FLiveLinkSubjectFrameData& GetFrameSnapshot() const override { return CurrentFrameSnapshot; }
