LiveLinkHub - UI Tweaks
Change settings icon to a single cog
Set icon and icon colors in LiveLinkCore style to allow LiveLinkDebugView to access them
Update device, subject and sources icons

Note: Using a PNG version of the sources icon since it seems like the engine isn't able to render that SVG correctly, possibly because of the mask group property used in it.

diff --git a/Engine/Plugins/Animation/LiveLink/Source/LiveLink/Private/LiveLinkDebugView.cpp b/Engine/Plugins/Animation/LiveLink/Source/LiveLink/Private/LiveLinkDebugView.cpp
--- a/Engine/Plugins/Animation/LiveLink/Source/LiveLink/Private/LiveLinkDebugView.cpp
+++ b/Engine/Plugins/Animation/LiveLink/Source/LiveLink/Private/LiveLinkDebugView.cpp
@@ -4,9 +4,11 @@
 
 #include "Framework/Views/TableViewMetadata.h"
 #include "LiveLinkClient.h"
+#include "LiveLinkModule.h"
 #include "LiveLinkSettings.h"
 
 #include "Brushes/SlateColorBrush.h"
+#include "Styling/SlateStyle.h"
 #include "Widgets/Colors/SColorBlock.h"
 #include "Widgets/Views/SListView.h"
 
@@ -33,29 +35,7 @@ public:
 	bool IsSubjectValid() const { return Client->IsSubjectValid(SubjectKey.SubjectName); }
 	bool IsSource() const { return SubjectKey.SubjectName.IsNone(); }
 	bool IsSourceValid() const { return Client->IsSourceStillValid(SubjectKey.Source); }
-	FLinearColor GetStatusColor() const
-	{
-		if (IsSource() && IsSourceValid())
-		{
-			return GetDefault<ULiveLinkSettings>()->GetValidColor();
-		}
-		else if (IsSubject())
-		{
-			if (IsSubjectEnabled())
-			{
-				if (IsSubjectValid())
-				{
-					return GetDefault<ULiveLinkSettings>()->GetValidColor();
-				}
-			}
-			else
-			{
-				return DisabledColorSubject;
-			}
-		}
-
-		return GetDefault<ULiveLinkSettings>()->GetInvalidColor();
-	}
+	bool IsPaused() const { return Client->GetSubjectState(SubjectKey.SubjectName) == ELiveLinkSubjectState::Paused; }
 
 	FText GetItemText() const
 	{
@@ -96,6 +76,14 @@ void SLiveLinkDebugView::Construct(const FArguments& Args, FLiveLinkClient* InCl
 		.SelectionMode(ESelectionMode::None)
 		.OnGenerateRow(this, &SLiveLinkDebugView::GenerateRow);
 
+	if (TSharedPtr<FSlateStyleSet> Style = ILiveLinkModule::Get().GetStyle())
+	{
+		ValidBrush = Style->GetBrush("LiveLink.Subject.Okay");
+		InvalidBrush = Style->GetBrush("LiveLink.Subject.Invalid");
+		PausedBrush = Style->GetBrush("LiveLink.Subject.Paused");
+		PausedBrush = Style->GetBrush("LiveLink.Subject.Error");
+	}
+
 	ChildSlot
 	[
 		SNew(SHorizontalBox)
@@ -136,8 +124,8 @@ TSharedRef<ITableRow> SLiveLinkDebugView::GenerateRow(FLiveLinkDebugUIEntryPtr D
 			.HAlign(HAlign_Left)
 			.VAlign(VAlign_Center)
 			[
-				SNew(SColorBlock)
-				.Color(Data.Get(), &FLiveLinkDebugUIEntry::GetStatusColor)
+				SNew(SImage)
+				.Image(this, &SLiveLinkDebugView::GetSubjectIcon, Data)
 			]
 			+ SHorizontalBox::Slot()
 			.AutoWidth()
@@ -175,3 +163,26 @@ void SLiveLinkDebugView::RefreshSourceItems()
 		}
 	}
 }
+
+const FSlateBrush* SLiveLinkDebugView::GetSubjectIcon(FLiveLinkDebugUIEntryPtr SubjectEntry) const
+{
+	const FSlateBrush* IconBrush = nullptr;
+
+	if (SubjectEntry->IsSubjectEnabled())
+	{
+		if (!SubjectEntry->IsPaused())
+		{
+			IconBrush = SubjectEntry->IsSubjectValid() ? ValidBrush : InvalidBrush;
+		}
+		else
+		{
+			IconBrush = PausedBrush;
+		}
+	}
+	else
+	{
+		IconBrush = DisabledBrush;
+	}
+
+	return IconBrush;
+}
diff --git a/Engine/Plugins/Animation/LiveLink/Source/LiveLink/Private/LiveLinkDebugView.h b/Engine/Plugins/Animation/LiveLink/Source/LiveLink/Private/LiveLinkDebugView.h
--- a/Engine/Plugins/Animation/LiveLink/Source/LiveLink/Private/LiveLinkDebugView.h
+++ b/Engine/Plugins/Animation/LiveLink/Source/LiveLink/Private/LiveLinkDebugView.h
@@ -9,6 +9,7 @@ class STableViewBase;
 template <typename ItemType> class SListView;
 
 class FLiveLinkClient;
+struct FSlateBrush;
 struct FSlateColorBrush;
 
 struct FLiveLinkDebugUIEntry;
@@ -27,6 +28,8 @@ private:
 	TSharedRef<ITableRow> GenerateRow(FLiveLinkDebugUIEntryPtr Data, const TSharedRef<STableViewBase>& OwnerTable);
 	void HandleSourcesChanged();
 	void RefreshSourceItems();
+	/** Get the subject icon according to the subject's status. */
+	const FSlateBrush* GetSubjectIcon(FLiveLinkDebugUIEntryPtr SubjectEntry) const;
 
 private:
 	FLiveLinkClient* Client;
@@ -35,4 +38,13 @@ private:
 	TSharedPtr<SListView<FLiveLinkDebugUIEntryPtr>> DebugItemView;
 	TSharedPtr<FSlateColorBrush> BackgroundBrushSource;
 	TSharedPtr<FSlateColorBrush> BackgroundBrushSubject;
+
+	/** Cached brush for valid subjects. */
+	const FSlateBrush* ValidBrush = nullptr;
+	/** Cached brush for invalid subjects. */
+	const FSlateBrush* InvalidBrush = nullptr;
+	/** Cached brush for paused subjects. */
+	const FSlateBrush* PausedBrush = nullptr;
+	/** Cached brush for disabled subjects */
+	const FSlateBrush* DisabledBrush = nullptr;
 };
diff --git a/Engine/Plugins/Animation/LiveLink/Source/LiveLink/Private/LiveLinkModule.cpp b/Engine/Plugins/Animation/LiveLink/Source/LiveLink/Private/LiveLinkModule.cpp
--- a/Engine/Plugins/Animation/LiveLink/Source/LiveLink/Private/LiveLinkModule.cpp
+++ b/Engine/Plugins/Animation/LiveLink/Source/LiveLink/Private/LiveLinkModule.cpp
@@ -20,6 +20,18 @@ LLM_DEFINE_TAG(LiveLink);
 #define LOCTEXT_NAMESPACE "LiveLinkModule"
 
 
+namespace LiveLinkModuleUtils
+{
+	FString InPluginContent(const FString& RelativePath, const ANSICHAR* Extension)
+	{
+		static FString ContentDir = IPluginManager::Get().FindPlugin(TEXT("LiveLink"))->GetContentDir();
+		return (ContentDir / RelativePath) + Extension;
+	}
+}
+
+
+#define IMAGE_PLUGIN_BRUSH_SVG( RelativePath, ... ) FSlateVectorImageBrush( LiveLinkModuleUtils::InPluginContent( RelativePath, ".svg" ), __VA_ARGS__ )
+
 FLiveLinkClient* FLiveLinkModule::LiveLinkClient_AnyThread = nullptr;
 
 FLiveLinkModule::FLiveLinkModule()
@@ -101,6 +113,22 @@ void FLiveLinkModule::CreateStyle()
 	const FVector2D Icon16x16(16.0f, 16.0f);
 
 	StyleSet->Set("LiveLinkIcon", new FSlateImageBrush((ContentDir / TEXT("LiveLink_16x")) + TEXT(".png"), Icon16x16));
+
+	const FSlateColor ValidColor = FSlateColor(FColor(31, 228, 75));
+	const FSlateColor InvalidColor = FSlateColor(FColor(255, 184, 0));
+	const FSlateColor ErrorColor = FSlateColor(FColor(239, 53, 53));
+	const FSlateColor PausedColor = FSlateColor(FColor(31, 228, 75));
+
+
+	StyleSet->Set("LiveLink.Color.Valid", ValidColor);
+	StyleSet->Set("LiveLink.Color.Invalid", InvalidColor);
+	StyleSet->Set("LiveLink.Color.Error", ErrorColor);
+	StyleSet->Set("LiveLink.Color.Paused", PausedColor);
+
+	StyleSet->Set("LiveLink.Subject.Okay", new IMAGE_PLUGIN_BRUSH_SVG(TEXT("Starship/Checkmark"), Icon16x16, ValidColor));
+	StyleSet->Set("LiveLink.Subject.Warning", new IMAGE_PLUGIN_BRUSH_SVG(TEXT("Starship/Warning"), Icon16x16, InvalidColor));
+	StyleSet->Set("LiveLink.Subject.Pause", new IMAGE_PLUGIN_BRUSH_SVG(TEXT("Starship/Pause"), Icon16x16, ErrorColor));
+	StyleSet->Set("LiveLink.Subject.Error", new IMAGE_PLUGIN_BRUSH_SVG(TEXT("Starship/Error"), Icon16x16, PausedColor));
 }
 
 void FLiveLinkModule::OnEngineLoopInitComplete()
diff --git a/Engine/Plugins/Animation/LiveLink/Source/LiveLink/Private/LiveLinkSettings.cpp b/Engine/Plugins/Animation/LiveLink/Source/LiveLink/Private/LiveLinkSettings.cpp
--- a/Engine/Plugins/Animation/LiveLink/Source/LiveLink/Private/LiveLinkSettings.cpp
+++ b/Engine/Plugins/Animation/LiveLink/Source/LiveLink/Private/LiveLinkSettings.cpp
@@ -23,7 +23,6 @@ ULiveLinkSettings::ULiveLinkSettings()
 	, TimeWithoutFrameToBeConsiderAsInvalid(0.5)
 	, ValidColor(FColor(31, 228, 75))
 	, InvalidColor(FColor(255, 184, 0))
-	, PausedColor(FColor(152, 152, 152))
 	, TextSizeSource(16)
 	, TextSizeSubject(12)
 {
diff --git a/Engine/Plugins/Animation/LiveLink/Source/LiveLink/Public/LiveLinkSettings.h b/Engine/Plugins/Animation/LiveLink/Source/LiveLink/Public/LiveLinkSettings.h
--- a/Engine/Plugins/Animation/LiveLink/Source/LiveLink/Public/LiveLinkSettings.h
+++ b/Engine/Plugins/Animation/LiveLink/Source/LiveLink/Public/LiveLinkSettings.h
@@ -139,17 +139,15 @@ public:
 	double TimeWithoutFrameToBeConsiderAsInvalid;
 
 	/** Color for active Subjects receiving data from their Source. */
+	UE_DEPRECATED(5.6, "Not used anymore in favor of using icons.")
 	UPROPERTY(config, EditAnywhere, Category = "LiveLink|UI")
 	FLinearColor ValidColor;
 
 	/** Color for Subjects that have not received data from their Source for TimeWithoutFrameToBeConsiderAsInvalid. */
+	UE_DEPRECATED(5.6, "Not used anymore in favor of using icons.")
 	UPROPERTY(config, EditAnywhere, Category = "LiveLink|UI")
 	FLinearColor InvalidColor;
 
-	/** Color for Subjects that are currently paused. */
-	UPROPERTY(config, EditAnywhere, Category = "LiveLink|UI")
-	FLinearColor PausedColor;
-
 	/** Font size of Source names shown in LiveLink Debug View. */
 	UPROPERTY(config, EditAnywhere, Category = "LiveLink|UI|Debug")
 	uint8 TextSizeSource;
@@ -165,9 +163,6 @@ public:
 	const FDirectoryPath& GetPresetSaveDir() const { return GetDefault<ULiveLinkUserSettings>()->GetPresetSaveDir(); }
 
 	double GetTimeWithoutFrameToBeConsiderAsInvalid() const { return TimeWithoutFrameToBeConsiderAsInvalid; }
-	FLinearColor GetValidColor() const { return ValidColor; }
-	FLinearColor GetInvalidColor() const { return InvalidColor; }
-	FLinearColor GetPausedColor() const { return PausedColor; }
 	float GetMessageBusPingRequestFrequency() const { return MessageBusPingRequestFrequency; }
 	float GetMessageBusHeartbeatFrequency() const { return MessageBusHeartbeatFrequency; }
 	double GetMessageBusHeartbeatTimeout() const { return MessageBusHeartbeatTimeout; }
@@ -178,4 +173,9 @@ public:
 	{
 		return GET_MEMBER_NAME_CHECKED(ULiveLinkSettings, DefaultRoleSettings);
 	}
+
+	UE_DEPRECATED(5.6, "Not used anymore in favor of using icons.")
+	FLinearColor GetValidColor() const { return FColor::Green; }
+	UE_DEPRECATED(5.6, "Not used anymore in favor of using icons.")
+	FLinearColor GetInvalidColor() const { return FColor::Red; }
 };
diff --git a/Engine/Plugins/Animation/LiveLink/Source/LiveLinkEditor/Private/LiveLinkClientPanelViews.cpp b/Engine/Plugins/Animation/LiveLink/Source/LiveLinkEditor/Private/LiveLinkClientPanelViews.cpp
--- a/Engine/Plugins/Animation/LiveLink/Source/LiveLinkEditor/Private/LiveLinkClientPanelViews.cpp
+++ b/Engine/Plugins/Animation/LiveLink/Source/LiveLinkEditor/Private/LiveLinkClientPanelViews.cpp
@@ -7,6 +7,7 @@
 #include "Framework/Commands/UICommandList.h"
 #include "IDetailsView.h"
 #include "ILiveLinkClient.h"
+#include "ILiveLinkModule.h"
 #include "Internationalization/Text.h"
 #include "LiveLinkClient.h"
 #include "LiveLinkClientCommands.h"
@@ -17,6 +18,7 @@
 #include "Modules/ModuleManager.h"
 #include "PropertyEditorModule.h"
 #include "SLiveLinkDataView.h"
+#include "Styling/SlateStyle.h"
 #include "Styling/SlateStyleRegistry.h"
 #include "Widgets/DeclarativeSyntaxSupport.h"
 #include "Widgets/Input/SButton.h"
@@ -212,10 +214,11 @@ public:
 		EntryPtr = Args._Entry;
 		bReadOnly = Args._ReadOnly;
 
-		const ISlateStyle* LiveLinkStyle = FSlateStyleRegistry::FindSlateStyle("LiveLinkStyle");
-		OkayIcon = LiveLinkStyle->GetBrush("LiveLink.Subject.Okay");
-		WarningIcon = LiveLinkStyle->GetBrush("LiveLink.Subject.Warning");
-		PauseIcon = LiveLinkStyle->GetBrush("LiveLink.Subject.Pause");
+		TSharedPtr<FSlateStyleSet> StyleSet = ILiveLinkModule::Get().GetStyle();
+
+		OkayIcon = StyleSet->GetBrush("LiveLink.Subject.Okay");
+		WarningIcon = StyleSet->GetBrush("LiveLink.Subject.Warning");
+		PauseIcon = StyleSet->GetBrush("LiveLink.Subject.Pause");
 
 		SMultiColumnTableRow<FLiveLinkSubjectUIEntryPtr>::Construct(
 			FSuperRowType::FArguments()
@@ -290,7 +293,6 @@ public:
 					[
 						SNew(SImage)
 						.Image(this, &SLiveLinkClientPanelSubjectRow::GetSubjectIcon)
-						.ColorAndOpacity(this, &SLiveLinkClientPanelSubjectRow::OnGetActivityColor)
 						.ToolTipText(this, &SLiveLinkClientPanelSubjectRow::GetSubjectIconToolTip)
 					];
 			}
@@ -319,23 +321,6 @@ private:
 		return FReply::Handled();
 	}
 
-	FSlateColor OnGetActivityColor() const
-	{
-		if (EntryPtr->IsSubjectEnabled())
-		{
-			if (!EntryPtr->IsPaused())
-			{
-				return EntryPtr->IsSubjectValid() ? GetDefault<ULiveLinkSettings>()->GetValidColor() : GetDefault<ULiveLinkSettings>()->GetInvalidColor();
-			}
-			else
-			{
-				return GetDefault<ULiveLinkSettings>()->GetPausedColor();
-			}
-		}
-
-		return FLinearColor(0.f, 0.f, 0.f, 0.f);
-	}
-
 	/** Get widget visibility according to whether or not the panel is in read-only mode. */
 	EVisibility GetVisibilityFromReadOnly() const
 	{
diff --git a/Engine/Plugins/Animation/LiveLink/Source/LiveLinkEditor/Private/LiveLinkEditorModule.cpp b/Engine/Plugins/Animation/LiveLink/Source/LiveLinkEditor/Private/LiveLinkEditorModule.cpp
--- a/Engine/Plugins/Animation/LiveLink/Source/LiveLinkEditor/Private/LiveLinkEditorModule.cpp
+++ b/Engine/Plugins/Animation/LiveLink/Source/LiveLinkEditor/Private/LiveLinkEditorModule.cpp
@@ -101,16 +101,11 @@ public:
 		StyleSet->SetContentRoot(FPaths::EngineContentDir() / TEXT("Editor/Slate"));
 		StyleSet->SetCoreContentRoot(FPaths::EngineContentDir() / TEXT("Slate"));
 
-		StyleSet->Set("LiveLink.Subject.Okay", new IMAGE_PLUGIN_BRUSH_SVG(TEXT("Starship/Checkmark"), Icon16x16));
-		StyleSet->Set("LiveLink.Subject.Warning", new IMAGE_PLUGIN_BRUSH_SVG(TEXT("Starship/Warning"), Icon16x16));
-		StyleSet->Set("LiveLink.Subject.Pause", new IMAGE_PLUGIN_BRUSH_SVG(TEXT("Starship/Pause"), Icon16x16));
-		StyleSet->Set("LiveLink.Subject.Error", new IMAGE_PLUGIN_BRUSH_SVG(TEXT("Starship/Error"), Icon16x16));
-
 		StyleSet->Set("LiveLinkClient.Common.Icon", new IMAGE_PLUGIN_BRUSH(TEXT("LiveLink_40x"), Icon40x40));
 		StyleSet->Set("LiveLinkClient.Common.Icon.Small", new IMAGE_PLUGIN_BRUSH_SVG("Starship/LiveLink", Icon16x16));
 
 		StyleSet->Set("LiveLinkHub.Clients.Icon", new IMAGE_PLUGIN_BRUSH_SVG(TEXT("Starship/Clients"), Icon16x16));
-		StyleSet->Set("LiveLinkHub.Sources.Icon", new IMAGE_PLUGIN_BRUSH_SVG(TEXT("Starship/Sources"), Icon16x16));
+		StyleSet->Set("LiveLinkHub.Sources.Icon", new IMAGE_PLUGIN_BRUSH(TEXT("Starship/Sources"), Icon16x16));
 		StyleSet->Set("LiveLinkHub.Layout.Icon", new IMAGE_PLUGIN_BRUSH_SVG(TEXT("Starship/Layout"), Icon16x16));
 		StyleSet->Set("LiveLinkHub.Playback.Icon", new IMAGE_PLUGIN_BRUSH_SVG(TEXT("Starship/Playback"), Icon16x16));
 		StyleSet->Set("LiveLinkHub.Subjects.Icon", new IMAGE_PLUGIN_BRUSH_SVG(TEXT("Starship/Subjects"), Icon16x16));
diff --git a/Engine/Plugins/Animation/LiveLinkDevice/Source/LiveLinkDevice/Private/LiveLinkDeviceModule.cpp b/Engine/Plugins/Animation/LiveLinkDevice/Source/LiveLinkDevice/Private/LiveLinkDeviceModule.cpp
--- a/Engine/Plugins/Animation/LiveLinkDevice/Source/LiveLinkDevice/Private/LiveLinkDeviceModule.cpp
+++ b/Engine/Plugins/Animation/LiveLinkDevice/Source/LiveLinkDevice/Private/LiveLinkDeviceModule.cpp
@@ -38,11 +38,13 @@ void FLiveLinkDeviceModule::StartupModule()
 	FGlobalTabmanager::Get()
 		->RegisterNomadTabSpawner(DevicesTabName,
 			FOnSpawnTab::CreateRaw(this, &FLiveLinkDeviceModule::OnSpawnDevicesTab))
+		.SetIcon(FSlateIcon("LiveLinkDeviceStyle", "LiveLinkHub.Devices.Icon"))
 		.SetDisplayName(LOCTEXT("DevicesTabDisplayName", "Devices"));
 
 	FGlobalTabmanager::Get()
 		->RegisterNomadTabSpawner(DeviceDetailsTabName,
 			FOnSpawnTab::CreateRaw(this, &FLiveLinkDeviceModule::OnSpawnDeviceDetailsTab))
+		.SetIcon(FSlateIcon(FAppStyle::GetAppStyleSetName(), "LevelEditor.Tabs.Details"))
 		.SetDisplayName(LOCTEXT("DeviceDetailsTabDisplayName", "Device Details"));
 }
 
@@ -128,7 +130,7 @@ TSharedRef<SDockTab> FLiveLinkDeviceModule::OnSpawnDevicesTab(const FSpawnTabArg
 				[
 					SNew(SImage)
 					.ColorAndOpacity(FSlateColor::UseForeground())
-					.Image(FSlateIcon("LiveLinkStyle", "LiveLinkHub.Sources.Icon").GetIcon())
+					.Image(FSlateIcon("LiveLinkDeviceStyle", "LiveLinkHub.Devices.Icon").GetIcon())
 				]
 				+ SHorizontalBox::Slot()
 				.HAlign(HAlign_Left)
diff --git a/Engine/Plugins/Animation/LiveLinkHub/Source/LiveLinkHub/Private/LiveLinkHubCreatorAppMode.h b/Engine/Plugins/Animation/LiveLinkHub/Source/LiveLinkHub/Private/LiveLinkHubCreatorAppMode.h
--- a/Engine/Plugins/Animation/LiveLinkHub/Source/LiveLinkHub/Private/LiveLinkHubCreatorAppMode.h
+++ b/Engine/Plugins/Animation/LiveLinkHub/Source/LiveLinkHub/Private/LiveLinkHubCreatorAppMode.h
@@ -45,7 +45,7 @@ public:
 		, WeakPanelController(InPanelController)
 	{
 		TabLabel = LOCTEXT("SourcesTabLabel", "Sources");
-		TabIcon = FSlateIcon(UE::LiveLinkHub::LiveLinkStyleName, TEXT("LiveLinkHub.Subjects.Icon"));
+		TabIcon = FSlateIcon(UE::LiveLinkHub::LiveLinkStyleName, TEXT("LiveLinkHub.Sources.Icon"));
 
 		bIsSingleton = true;
 	}
@@ -157,7 +157,7 @@ public:
 		, WeakPanelController(InPanelController)
 	{
 		TabLabel = LOCTEXT("SubjectsTabLabel", "Subjects");
-		TabIcon = FSlateIcon(FAppStyle::GetAppStyleSetName(), "LevelEditor.Tabs.Details");
+		TabIcon = FSlateIcon("LiveLinkStyle", "LiveLinkHub.Subjects.Icon");
 
 		bIsSingleton = true;
 	}
@@ -251,7 +251,7 @@ public:
 		: FWorkflowTabFactory(UE::LiveLinkHub::ClientsTabId, MoveTemp(InHostingApp))
 	{
 		TabLabel = LOCTEXT("ClientsTabLabel", "Clients");
-		TabIcon = FSlateIcon(FAppStyle::GetAppStyleSetName(), "LevelEditor.Tabs.Details");
+		TabIcon = FSlateIcon("LiveLinkStyle", "LiveLinkHub.Clients.Icon");
 
 		bIsSingleton = true;
 	}
diff --git a/Engine/Plugins/Animation/LiveLinkHub/Source/LiveLinkHub/Private/UI/Widgets/SLiveLinkHubSettings.h b/Engine/Plugins/Animation/LiveLinkHub/Source/LiveLinkHub/Private/UI/Widgets/SLiveLinkHubSettings.h
--- a/Engine/Plugins/Animation/LiveLinkHub/Source/LiveLinkHub/Private/UI/Widgets/SLiveLinkHubSettings.h
+++ b/Engine/Plugins/Animation/LiveLinkHub/Source/LiveLinkHub/Private/UI/Widgets/SLiveLinkHubSettings.h
@@ -67,16 +67,8 @@ public:
 					.ToolTipText(LOCTEXT("SettingsToolTip", "Settings"))
 					.ButtonContent()
 					[
-						SNew(SHorizontalBox)
-						+ SHorizontalBox::Slot()
-						.VAlign(VAlign_Center)
-						.Padding(0.0f)
-						.AutoWidth()
-						[
-							SNew(STextBlock)
-							.Font(FAppStyle::Get().GetFontStyle("FontAwesome.14"))
-							.Text(FEditorFontGlyphs::Cogs)
-						]
+						SNew(SImage)
+						.Image(FAppStyle::Get().GetBrush("FullBlueprintEditor.EditGlobalOptions"))
 					]
 					.OnGetMenuContent(this, &SLiveLinkHubSettings::OnSettingsComboBoxOpen)
 				]
