LiveLink - Fix icons not showing up in LiveLinkDebugView

#rnx

diff --git a/Engine/Plugins/Animation/LiveLink/Source/LiveLink/Private/LiveLinkDebugView.cpp b/Engine/Plugins/Animation/LiveLink/Source/LiveLink/Private/LiveLinkDebugView.cpp
--- a/Engine/Plugins/Animation/LiveLink/Source/LiveLink/Private/LiveLinkDebugView.cpp
+++ b/Engine/Plugins/Animation/LiveLink/Source/LiveLink/Private/LiveLinkDebugView.cpp
@@ -18,7 +18,6 @@ namespace
 	static const int IndentationSubject = 12;
 	static const FLinearColor BackgroundColorSource(FColor(62, 62, 62, 180));
 	static const FLinearColor BackgroundColorSubject(FColor(62, 62, 62, 120));
-	static const FLinearColor DisabledColorSubject(FColor(255, 255, 255, 120));
 }
 
 // Structure that defines a single entry in the debug UI
@@ -79,9 +78,8 @@ void SLiveLinkDebugView::Construct(const FArguments& Args, FLiveLinkClient* InCl
 	if (TSharedPtr<FSlateStyleSet> Style = ILiveLinkModule::Get().GetStyle())
 	{
 		ValidBrush = Style->GetBrush("LiveLink.Subject.Okay");
-		InvalidBrush = Style->GetBrush("LiveLink.Subject.Invalid");
+		InvalidBrush = Style->GetBrush("LiveLink.Subject.Warning");
 		PausedBrush = Style->GetBrush("LiveLink.Subject.Paused");
-		PausedBrush = Style->GetBrush("LiveLink.Subject.Error");
 	}
 
 	ChildSlot
