LiveLink - Fix style issues with valid / invalid colors

#rnx

diff --git a/Engine/Plugins/Animation/LiveLink/Source/LiveLink/Private/LiveLinkModule.cpp b/Engine/Plugins/Animation/LiveLink/Source/LiveLink/Private/LiveLinkModule.cpp
--- a/Engine/Plugins/Animation/LiveLink/Source/LiveLink/Private/LiveLinkModule.cpp
+++ b/Engine/Plugins/Animation/LiveLink/Source/LiveLink/Private/LiveLinkModule.cpp
@@ -114,10 +114,10 @@ void FLiveLinkModule::CreateStyle()
 
 	StyleSet->Set("LiveLinkIcon", new FSlateImageBrush((ContentDir / TEXT("LiveLink_16x")) + TEXT(".png"), Icon16x16));
 
-	const FSlateColor ValidColor = FSlateColor(FColor(31, 228, 75));
-	const FSlateColor InvalidColor = FSlateColor(FColor(255, 184, 0));
-	const FSlateColor ErrorColor = FSlateColor(FColor(239, 53, 53));
-	const FSlateColor PausedColor = FSlateColor(FColor(31, 228, 75));
+	const FSlateColor ValidColor = FSlateColor(FLinearColor(0.0146, 0.7874, 0.0736));
+	const FSlateColor InvalidColor = FSlateColor(FLinearColor(1.0, 0.4654, 0.0));
+	const FSlateColor ErrorColor = FSlateColor(FLinearColor(0.8524, 0.0372, 0.0372));
+	const FSlateColor PausedColor = FSlateColor(FLinearColor(0.2159, 0.2159, 0.2159));
 
 
 	StyleSet->Set("LiveLink.Color.Valid", ValidColor);
@@ -127,8 +127,8 @@ void FLiveLinkModule::CreateStyle()
 
 	StyleSet->Set("LiveLink.Subject.Okay", new IMAGE_PLUGIN_BRUSH_SVG(TEXT("Starship/Checkmark"), Icon16x16, ValidColor));
 	StyleSet->Set("LiveLink.Subject.Warning", new IMAGE_PLUGIN_BRUSH_SVG(TEXT("Starship/Warning"), Icon16x16, InvalidColor));
-	StyleSet->Set("LiveLink.Subject.Pause", new IMAGE_PLUGIN_BRUSH_SVG(TEXT("Starship/Pause"), Icon16x16, ErrorColor));
-	StyleSet->Set("LiveLink.Subject.Error", new IMAGE_PLUGIN_BRUSH_SVG(TEXT("Starship/Error"), Icon16x16, PausedColor));
+	StyleSet->Set("LiveLink.Subject.Pause", new IMAGE_PLUGIN_BRUSH_SVG(TEXT("Starship/Pause"), Icon16x16, PausedColor));
+	StyleSet->Set("LiveLink.Subject.Error", new IMAGE_PLUGIN_BRUSH_SVG(TEXT("Starship/Error"), Icon16x16, ErrorColor));
 }
 
 void FLiveLinkModule::OnEngineLoopInitComplete()
