LiveLinkHub - Allow using the base skeleton remapper to rename bones that go through the hub.

diff --git a/Engine/Plugins/Animation/LiveLink/Source/LiveLink/Public/Remapper/LiveLinkSkeletonRemapper.h b/Engine/Plugins/Animation/LiveLink/Source/LiveLink/Public/Remapper/LiveLinkSkeletonRemapper.h
--- a/Engine/Plugins/Animation/LiveLink/Source/LiveLink/Public/Remapper/LiveLinkSkeletonRemapper.h
+++ b/Engine/Plugins/Animation/LiveLink/Source/LiveLink/Public/Remapper/LiveLinkSkeletonRemapper.h
@@ -35,6 +35,14 @@ public:
 
 	virtual void RemapSkeletonStaticData(FLiveLinkSkeletonStaticData& InOutSkeletonData)
 	{
+		TArray<FName> SourceBoneNames = InOutSkeletonData.GetBoneNames();
+
+		for (int32 Index = 0; Index < SourceBoneNames.Num(); Index++)
+		{
+			SourceBoneNames[Index] = GetRemappedBoneName(SourceBoneNames[Index]);
+		}
+
+		InOutSkeletonData.SetBoneNames(SourceBoneNames);
 	}
 
 	virtual void RemapSkeletonFrameData(const FLiveLinkSkeletonStaticData& InOutSkeletonData, FLiveLinkAnimationFrameData& InOutFrameData)
@@ -54,7 +62,7 @@ public:
 	TMap<FName, FName> BoneNameMap;
 };
 
-UCLASS(Abstract)
+UCLASS(EditInlineNew)
 class LIVELINK_API ULiveLinkSkeletonRemapper : public ULiveLinkSubjectRemapper
 {
 public:
@@ -91,10 +99,26 @@ public:
 
 	virtual bool IsValidRemapper() const
 	{
-		return ReferenceSkeleton.IsValid();
+		return true;
+	}
+
+	virtual FWorkerSharedPtr GetWorker() const override
+	{
+		return Instance;
+	}
+
+	virtual FWorkerSharedPtr CreateWorker() override
+	{
+		Instance = MakeShared<FLiveLinkSkeletonRemapperWorker>();
+		Instance->BoneNameMap = BoneNameMap;
+		return Instance;
 	}
 
 public:
 	UPROPERTY(EditAnywhere, Category="Remapper", meta=(DisplayThumbnail="false"))
 	TSoftObjectPtr<USkeletalMesh> ReferenceSkeleton;
+
+private:
+	/** Instance of the remapper. */
+	TSharedPtr<FLiveLinkSkeletonRemapperWorker> Instance;
 };
diff --git a/Engine/Source/Runtime/LiveLinkInterface/Public/LiveLinkSubjectRemapper.h b/Engine/Source/Runtime/LiveLinkInterface/Public/LiveLinkSubjectRemapper.h
--- a/Engine/Source/Runtime/LiveLinkInterface/Public/LiveLinkSubjectRemapper.h
+++ b/Engine/Source/Runtime/LiveLinkInterface/Public/LiveLinkSubjectRemapper.h
@@ -26,7 +26,7 @@ public:
 	LIVELINKINTERFACE_API virtual void RemapFrameData(const FLiveLinkStaticDataStruct& InStaticData, FLiveLinkFrameDataStruct& InOutFrameData) { }
 
 	/** Returns if the remapper is compatible with the static data. */
-	LIVELINKINTERFACE_API virtual bool IsRemapperCompatible(const FLiveLinkStaticDataStruct& InStaticData) { return false;  }
+	LIVELINKINTERFACE_API virtual bool IsRemapperCompatible(const FLiveLinkStaticDataStruct& InStaticData) { return true; }
 };
 
 /** Class used to remap livelink subjects without having to rely on animation blueprints. */
