LiveLinkHub - Allow rebroadcasting evaluated LiveLink data

This should allow LiveLinkHub to transmit evaluated and interpolated data to clients if needed.
Increase default threshold before warning about high memory usage

diff --git a/Engine/Plugins/Animation/LiveLink/Source/LiveLink/Private/LiveLinkClient.cpp b/Engine/Plugins/Animation/LiveLink/Source/LiveLink/Private/LiveLinkClient.cpp
--- a/Engine/Plugins/Animation/LiveLink/Source/LiveLink/Private/LiveLinkClient.cpp
+++ b/Engine/Plugins/Animation/LiveLink/Source/LiveLink/Private/LiveLinkClient.cpp
@@ -148,6 +148,52 @@ void FLiveLinkClient::UpdateSources()
 	});
 }
 
+void FLiveLinkClient::HandleSubjectRebroadcast(ILiveLinkSubject* InSubject, FLiveLinkSubjectFrameData&& SubjectFrameData)
+{
+	// Check the rebroadcast flag and act accordingly, creating the LiveLinkProvider and/or sending the static data if needed
+	if (InSubject->IsRebroadcasted())
+	{
+		if (InSubject->GetStaticData().IsValid() && SubjectFrameData.FrameData.IsValid())
+		{
+			// Setup rebroadcast provider
+			if (!RebroadcastLiveLinkProvider.IsValid())
+			{
+				RebroadcastLiveLinkProvider = GetRebroadcastLiveLinkProvider();
+			}
+				
+			if (RebroadcastLiveLinkProvider.IsValid())
+			{
+				TSubclassOf<ULiveLinkRole> SubjectRole = InSubject->GetRole();
+
+				const FName RebroadcastName = GetRebroadcastName(InSubject->GetSubjectKey());
+
+				const FText OriginalSourceType = GetSourceType(InSubject->GetSubjectKey().Source);
+
+				TMap<FName, FString> ExtraAnnotations;
+				ExtraAnnotations.Add(FLiveLinkMessageAnnotation::OriginalSourceAnnotation, OriginalSourceType.ToString());
+
+				if (!InSubject->HasStaticDataBeenRebroadcasted())
+				{
+					RebroadcastLiveLinkProvider->UpdateSubjectStaticData(RebroadcastName, SubjectRole, MoveTemp(SubjectFrameData.StaticData), ExtraAnnotations);
+					InSubject->SetStaticDataAsRebroadcasted(true);
+					RebroadcastedSubjects.Add(InSubject->GetSubjectKey());
+				}
+
+				RebroadcastLiveLinkProvider->UpdateSubjectFrameData(RebroadcastName, MoveTemp(SubjectFrameData.FrameData), ExtraAnnotations);
+			}
+			else
+			{
+				UE_LOG(LogLiveLink, Warning, TEXT("Rebroadcaster doesn't exist, but was requested and failed"));
+			}
+		}
+	}
+	else if (InSubject->HasStaticDataBeenRebroadcasted())
+	{
+		RemoveRebroadcastedSubject(InSubject->GetSubjectKey());
+		InSubject->SetStaticDataAsRebroadcasted(false);
+	}
+}
+
 void FLiveLinkClient::BuildThisTicksSubjectSnapshot()
 {
 	check(Collection);
@@ -189,6 +235,26 @@ void FLiveLinkClient::BuildThisTicksSubjectSnapshot()
 		}
 	});
 
+	Collection->ForEachSubject([this](FLiveLinkCollectionSourceItem& SourceItem, FLiveLinkCollectionSubjectItem& SubjectItem)
+		{
+			if (FLiveLinkSubject* Subject = SubjectItem.GetLiveSubject())
+			{
+				if (SubjectItem.bEnabled && !Subject->IsPaused() && GetDefault<ULiveLinkSettings>()->bTransmitEvaluatedData)
+				{
+					FLiveLinkSubjectFrameData FrameData;
+					if (Subject->EvaluateFrame(Subject->GetRole(), FrameData))
+					{
+						const FLiveLinkFrameIdentifier FrameId = FrameData.FrameData.GetBaseData()->FrameId;
+						if (Subject->GetLastRebroadcastedFrameId() != FrameId)
+						{
+							HandleSubjectRebroadcast(SubjectItem.GetLiveSubject(), MoveTemp(FrameData));
+							Subject->SetLastRebroadcastedFrameId(FrameId);
+						}
+					}
+				}
+			}
+		});
+
 	Collection->ForEachSubject([this](FLiveLinkCollectionSourceItem& SourceItem, FLiveLinkCollectionSubjectItem& SubjectItem)
 	{
 		if (ULiveLinkVirtualSubject* VSubject = SubjectItem.GetVirtualSubject())
@@ -368,7 +434,17 @@ ELiveLinkSubjectState FLiveLinkClient::GetSubjectState(FLiveLinkSubjectName InSu
 
 	if (!LiveSubject)
 	{
-		return ELiveLinkSubjectState::InvalidOrDisabled;
+		if (SubjectItem->GetVirtualSubject() && SubjectItem->bEnabled)
+		{
+			if (SubjectItem->bEnabled)
+			{
+				return SubjectItem->GetVirtualSubject()->IsPaused() ? ELiveLinkSubjectState::Paused : ELiveLinkSubjectState::Connected;
+			}
+		}
+		else
+		{
+			return ELiveLinkSubjectState::InvalidOrDisabled;
+		}
 	}
 
 	if (LiveSubject->IsPaused())
@@ -1003,10 +1079,14 @@ void FLiveLinkClient::PushSubjectFrameData_Internal(FPendingSubjectFrame&& Subje
 	const bool bHasParentSubject = SourceItem->Setting->ParentSubject != FLiveLinkSubjectName();
 	if (!bHasParentSubject)
 	{
+		// If it's paused, rebroadcast will be handled in FLiveLinkClient::BuildThisTicksSubjectSnapshot
 		if (!LinkSubject->IsPaused())
 		{
-			// If it's paused, rebroadcast will be handled in FLiveLinkClient::BuildThisTicksSubjectSnapshot
-			HandleSubjectRebroadcast(LinkSubject, SubjectFrameData.FrameData);
+			/** Only rebroadcast here if we're transmitted non-evaluated data. */
+			if (!GetDefault<ULiveLinkSettings>()->bTransmitEvaluatedData)
+			{
+				HandleSubjectRebroadcast(LinkSubject, SubjectFrameData.FrameData);
+			}
 		}
 
 		if (bEnableParentSubjects)
@@ -1138,20 +1218,34 @@ void FLiveLinkClient::PauseSubject_AnyThread(FLiveLinkSubjectName SubjectName)
 	{
 		if (const FLiveLinkCollectionSubjectItem* SubjectItem = Collection->FindSubject(SubjectName))
 		{
+			bool bPauseResult = false;
 			if (FLiveLinkSubject* LiveSubject = SubjectItem->GetLiveSubject())
 			{
-				if (!LiveSubject->IsPaused())
+				if (IsSubjectValid(SubjectName))
 				{
-					if (IsSubjectValid(SubjectName))
-					{
-						LiveSubject->PauseSubject();
-					}
-					else
-					{
-						UE_LOG(LogLiveLink, Warning, TEXT("Could not pause subject %s since it's not in a valid state."), *SubjectName.ToString());
-					}
+					LiveSubject->PauseSubject();
+				}
+				else
+				{
+					bPauseResult = false;
 				}
 			}
+			else if (ULiveLinkVirtualSubject* VirtualSubject = SubjectItem->GetVirtualSubject())
+			{
+				if (IsSubjectValid(SubjectName))
+				{
+					VirtualSubject->PauseSubject();
+				}
+				else
+				{
+					bPauseResult = false;
+				}
+			}
+
+			if (!bPauseResult)
+			{
+				UE_LOG(LogLiveLink, Warning, TEXT("Could not pause subject %s since it's not in a valid state."), *SubjectName.ToString());
+			}
 		}
 	}
 }
@@ -1166,6 +1260,10 @@ void FLiveLinkClient::UnpauseSubject_AnyThread(FLiveLinkSubjectName SubjectName)
 			{
 				LiveSubject->UnpauseSubject();
 			}
+			else if (ULiveLinkVirtualSubject* VirtualSubject = SubjectItem->GetVirtualSubject())
+			{
+				VirtualSubject->UnpauseSubject();
+			}
 		}
 	}
 }
@@ -1328,6 +1426,11 @@ bool FLiveLinkClient::IsSubjectValid(const FLiveLinkSubjectKey& InSubjectKey) co
 	{
 		if (SubjectItem->GetSubject()->HasValidFrameSnapshot())
 		{
+			if (SubjectItem->GetVirtualSubject())
+			{
+				return true;
+			}
+
 			if (FLiveLinkSubject* LiveSubject = SubjectItem->GetLiveSubject())
 			{
 				return LiveSubject->GetState() == ETimedDataInputState::Connected;
diff --git a/Engine/Plugins/Animation/LiveLink/Source/LiveLink/Public/LiveLinkClient.h b/Engine/Plugins/Animation/LiveLink/Source/LiveLink/Public/LiveLinkClient.h
--- a/Engine/Plugins/Animation/LiveLink/Source/LiveLink/Public/LiveLinkClient.h
+++ b/Engine/Plugins/Animation/LiveLink/Source/LiveLink/Public/LiveLinkClient.h
@@ -222,6 +222,11 @@ private:
 
 	/** Process virtual subject for rebroadcast purpose */
 	void HandleSubjectRebroadcast(ILiveLinkSubject* InSubject, const FLiveLinkFrameDataStruct& InFrameData);
+	/**
+	 * Rebroadcast evaluated livelink data.
+	 * @note Calling this method will not apply preprocessors and translators since it's meant to be used to transmit evaluated data. 
+	 */
+	void HandleSubjectRebroadcast(ILiveLinkSubject* InSubject, FLiveLinkSubjectFrameData&& SubjectFrameData);
 
 	/** Called when a subject is removed. Used to remove rebroadcasted subjects */
 	void OnSubjectRemovedCallback(FLiveLinkSubjectKey InSubjectKey);
diff --git a/Engine/Plugins/Animation/LiveLink/Source/LiveLink/Public/LiveLinkSettings.h b/Engine/Plugins/Animation/LiveLink/Source/LiveLink/Public/LiveLinkSettings.h
--- a/Engine/Plugins/Animation/LiveLink/Source/LiveLink/Public/LiveLinkSettings.h
+++ b/Engine/Plugins/Animation/LiveLink/Source/LiveLink/Public/LiveLinkSettings.h
@@ -129,6 +129,13 @@ public:
 	UPROPERTY(config)
 	bool bTranslateRebroadcastFrames = false;
 
+	/**
+	 * Experimental
+	 * Enabling this will evaluate livelink data before rebroadcasting it.
+	 */
+	UPROPERTY(config, EditAnywhere, Category = "LiveLink")
+	bool bTransmitEvaluatedData = false;
+
 	/**
 	 * A source may still exist but does not send frames for a subject.
 	 * Time before considering the subject as "invalid".
diff --git a/Engine/Plugins/Animation/LiveLink/Source/LiveLink/Public/LiveLinkSubject.h b/Engine/Plugins/Animation/LiveLink/Source/LiveLink/Public/LiveLinkSubject.h
--- a/Engine/Plugins/Animation/LiveLink/Source/LiveLink/Public/LiveLinkSubject.h
+++ b/Engine/Plugins/Animation/LiveLink/Source/LiveLink/Public/LiveLinkSubject.h
@@ -140,6 +140,16 @@ public:
 	/** Clear the override static data for this subject. */
 	void ClearOverrideStaticData_AnyThread();
 
+	FLiveLinkFrameIdentifier GetLastRebroadcastedFrameId() const
+	{
+		return LastRebroadcastFrameId;
+	}
+
+	void SetLastRebroadcastedFrameId(FLiveLinkFrameIdentifier FrameId)
+	{
+		LastRebroadcastFrameId = FrameId;
+	}
+
 	/** Get the delegate triggered when the state changes. */
 	FOnStateChanged& OnStateChanged() { return OnStateChangedDelegate; }
 
@@ -277,6 +287,9 @@ private:
 
 	/** Delegate called when the state of this subject has changed. */
 	FOnStateChanged OnStateChangedDelegate;
+
+	/** Last frame ID that was rebroadcasted. Used to avoid LiveLinkClient from rebroadcasting the same frame twice when transmitting evaluated data. */
+	FLiveLinkFrameIdentifier LastRebroadcastFrameId = 0;
 	
 	/** 
 	 * Evaluation can be done on any thread so we need to protect statistic logging 
diff --git a/Engine/Plugins/Animation/LiveLinkHub/Source/LiveLinkHub/Private/Settings/LiveLinkHubSettings.h b/Engine/Plugins/Animation/LiveLinkHub/Source/LiveLinkHub/Private/Settings/LiveLinkHubSettings.h
--- a/Engine/Plugins/Animation/LiveLinkHub/Source/LiveLinkHub/Private/Settings/LiveLinkHubSettings.h
+++ b/Engine/Plugins/Animation/LiveLinkHub/Source/LiveLinkHub/Private/Settings/LiveLinkHubSettings.h
@@ -73,7 +73,7 @@ public:
 
 	/** How much RAM (in MB) the program can use before showing a warning. */
 	UPROPERTY(config, EditAnywhere, Category = "LiveLinkHub", DisplayName = "Memory Warning Threshold (MB)")
-	float ShowMemoryWarningThresholdMB = 5000.0;
+	float ShowMemoryWarningThresholdMB = 8000.0;
 
 	/**
 	 * - Experimental - If this is disabled, LiveLinkHub's LiveLink Client will tick outside of the game thread.
