LiveLink - Add status to virtual subjects
- Allow pausing virtual subjects

diff --git a/Engine/Plugins/Animation/LiveLink/Source/LiveLink/Private/LiveLinkAnimationVirtualSubject.cpp b/Engine/Plugins/Animation/LiveLink/Source/LiveLink/Private/LiveLinkAnimationVirtualSubject.cpp
--- a/Engine/Plugins/Animation/LiveLink/Source/LiveLink/Private/LiveLinkAnimationVirtualSubject.cpp
+++ b/Engine/Plugins/Animation/LiveLink/Source/LiveLink/Private/LiveLinkAnimationVirtualSubject.cpp
@@ -68,8 +68,11 @@ ULiveLinkAnimationVirtualSubject::ULiveLinkAnimationVirtualSubject()
 
 void ULiveLinkAnimationVirtualSubject::Update()
 {
-	// Invalid the snapshot
-	InvalidateFrameData();
+	if (!IsPaused())
+	{
+		// Invalid the snapshot
+		InvalidateFrameData();
+	}
 
 	UpdateTranslatorsForThisFrame();
 
@@ -84,10 +87,13 @@ void ULiveLinkAnimationVirtualSubject::Update()
 			SortSubjects();
 		}
 
-		if (BuildSubjectSnapshot(SubjectSnapshot))
+		if (!IsPaused())
 		{
-			BuildSkeleton(SubjectSnapshot);
-			BuildFrame(SubjectSnapshot);
+			if (BuildSubjectSnapshot(SubjectSnapshot))
+			{
+				BuildSkeleton(SubjectSnapshot);
+				BuildFrame(SubjectSnapshot);
+			}
 		}
 	}
 	else
diff --git a/Engine/Plugins/Animation/LiveLink/Source/LiveLink/Private/LiveLinkClient.cpp b/Engine/Plugins/Animation/LiveLink/Source/LiveLink/Private/LiveLinkClient.cpp
--- a/Engine/Plugins/Animation/LiveLink/Source/LiveLink/Private/LiveLinkClient.cpp
+++ b/Engine/Plugins/Animation/LiveLink/Source/LiveLink/Private/LiveLinkClient.cpp
@@ -436,10 +436,7 @@ ELiveLinkSubjectState FLiveLinkClient::GetSubjectState(FLiveLinkSubjectName InSu
 	{
 		if (SubjectItem->GetVirtualSubject() && SubjectItem->bEnabled)
 		{
-			if (SubjectItem->bEnabled)
-			{
-				return SubjectItem->GetVirtualSubject()->IsPaused() ? ELiveLinkSubjectState::Paused : ELiveLinkSubjectState::Connected;
-			}
+			return SubjectItem->GetVirtualSubject()->IsPaused() ? ELiveLinkSubjectState::Paused : ELiveLinkSubjectState::Connected;
 		}
 		else
 		{
@@ -1219,32 +1216,20 @@ void FLiveLinkClient::PauseSubject_AnyThread(FLiveLinkSubjectName SubjectName)
 		if (const FLiveLinkCollectionSubjectItem* SubjectItem = Collection->FindSubject(SubjectName))
 		{
 			bool bPauseResult = false;
-			if (FLiveLinkSubject* LiveSubject = SubjectItem->GetLiveSubject())
+
+			if (!IsSubjectValid(SubjectName))
 			{
-				if (IsSubjectValid(SubjectName))
-				{
-					LiveSubject->PauseSubject();
-				}
-				else
-				{
-					bPauseResult = false;
-				}
+				UE_LOG(LogLiveLink, Warning, TEXT("Could not pause subject %s since it's not in a valid state."), *SubjectName.ToString());
+				return;
 			}
-			else if (ULiveLinkVirtualSubject* VirtualSubject = SubjectItem->GetVirtualSubject())
+
+			if (FLiveLinkSubject* LiveSubject = SubjectItem->GetLiveSubject())
 			{
-				if (IsSubjectValid(SubjectName))
-				{
-					VirtualSubject->PauseSubject();
-				}
-				else
-				{
-					bPauseResult = false;
-				}
+				LiveSubject->PauseSubject();
 			}
-
-			if (!bPauseResult)
+			else if (ULiveLinkVirtualSubject* VirtualSubject = SubjectItem->GetVirtualSubject())
 			{
-				UE_LOG(LogLiveLink, Warning, TEXT("Could not pause subject %s since it's not in a valid state."), *SubjectName.ToString());
+				VirtualSubject->PauseSubject();
 			}
 		}
 	}
diff --git a/Engine/Plugins/Animation/LiveLink/Source/LiveLinkEditor/Private/LiveLinkClientPanelViews.cpp b/Engine/Plugins/Animation/LiveLink/Source/LiveLinkEditor/Private/LiveLinkClientPanelViews.cpp
--- a/Engine/Plugins/Animation/LiveLink/Source/LiveLinkEditor/Private/LiveLinkClientPanelViews.cpp
+++ b/Engine/Plugins/Animation/LiveLink/Source/LiveLinkEditor/Private/LiveLinkClientPanelViews.cpp
@@ -267,35 +267,14 @@ public:
 		}
 		else if (ColumnName == SubjectTreeUI::ActionsColumnName)
 		{
-			if (EntryPtr->IsVirtualSubject())
-			{
-				return SNew(SButton)
-					.ButtonStyle(FAppStyle::Get(), "HoverHintOnly")
-					.HAlign(HAlign_Center)
-					.VAlign(VAlign_Center)
-					.OnClicked(this, &SLiveLinkClientPanelSubjectRow::OnRemoveClicked)
-					.ToolTipText(LOCTEXT("RemoveVirtualSubject", "Remove selected live link virtual subject"))
-					.ContentPadding(0.f)
-					.ForegroundColor(FSlateColor::UseForeground())
-					.IsFocusable(false)
-					.Visibility(this, &SLiveLinkClientPanelSubjectRow::GetVisibilityFromReadOnly)
-					[
-						SNew(SImage)
-						.Image(FAppStyle::GetBrush("Icons.Delete"))
-						.ColorAndOpacity(FSlateColor::UseForeground())
-					];
-			}
-			else
-			{
-				return SNew(SBox)
-					.HAlign(HAlign_Center)
-					.VAlign(VAlign_Center)
-					[
-						SNew(SImage)
-						.Image(this, &SLiveLinkClientPanelSubjectRow::GetSubjectIcon)
-						.ToolTipText(this, &SLiveLinkClientPanelSubjectRow::GetSubjectIconToolTip)
-					];
-			}
+			return SNew(SBox)
+				.HAlign(HAlign_Center)
+				.VAlign(VAlign_Center)
+				[
+					SNew(SImage)
+					.Image(this, &SLiveLinkClientPanelSubjectRow::GetSubjectIcon)
+					.ToolTipText(this, &SLiveLinkClientPanelSubjectRow::GetSubjectIconToolTip)
+				];
 		}
 
 		return SNullWidget::NullWidget;
@@ -315,12 +294,6 @@ private:
 		return FText::GetEmpty();
 	}
 
-	FReply OnRemoveClicked()
-	{
-		EntryPtr->RemoveFromClient();
-		return FReply::Handled();
-	}
-
 	/** Get widget visibility according to whether or not the panel is in read-only mode. */
 	EVisibility GetVisibilityFromReadOnly() const
 	{
