LiveLink - Serialize and compress animation data to reduce data throughput

Animation compression will be enabled automatically for message bus sources. Clients that support compression will add an annotation to their connect message.

You can disable this on the client side by the entering the command "LiveLink.SupportCompressPayloads 0"

On the provider side, you can disable this by entering "LiveLink.CompressAnimData 1"

diff --git a/Engine/Plugins/Animation/LiveLink/Source/LiveLink/Private/LiveLinkMessageBusSource.cpp b/Engine/Plugins/Animation/LiveLink/Source/LiveLink/Private/LiveLinkMessageBusSource.cpp
--- a/Engine/Plugins/Animation/LiveLink/Source/LiveLink/Private/LiveLinkMessageBusSource.cpp
+++ b/Engine/Plugins/Animation/LiveLink/Source/LiveLink/Private/LiveLinkMessageBusSource.cpp
@@ -4,6 +4,7 @@
 
 #include "ILiveLinkClient.h"
 #include "ILiveLinkModule.h"
+#include "LiveLinkCompression.h"
 #include "LiveLinkHeartbeatEmitter.h"
 #include "LiveLinkLog.h"
 #if WITH_LIVELINK_DISCOVERY_MANAGER_THREAD
@@ -17,6 +18,7 @@
 
 #include "MessageEndpointBuilder.h"
 #include "Misc/App.h"
+#include "Roles/LiveLinkAnimationTypes.h"
 
 
 int32 GLiveLinkMessageBusSourceReconnectAfterTimeout = 1;
@@ -26,6 +28,11 @@ static FAutoConsoleVariableRef CVarLiveLinkMessageBusSourceReconnectAfterTimeout
 	TEXT("When enabled, when the connection times out, it will try to re-connect instead of removing the source."),
 	ECVF_Default);
 
+static TAutoConsoleVariable<int32> CVarLiveLinkSupportCompressPayloads(
+	TEXT("LiveLink.SupportCompressPayloads"), 1,
+	TEXT("Whether to add the annotation indicating that we support compressed animation data. Can be set to 0 to simulate that compressed payloads are not supported."),
+	ECVF_RenderThreadSafe);
+
 
 FText FLiveLinkMessageBusSource::ValidSourceStatus()
 {
@@ -153,6 +160,11 @@ void FLiveLinkMessageBusSource::InternalHandleMessage(const TSharedRef<IMessageC
 		return;
 	}
 
+	if (!bIsValid)
+	{
+		return;
+	}
+
 	UScriptStruct* MessageTypeInfo = Context->GetMessageTypeInfo().Get();
 	if (MessageTypeInfo == nullptr)
 	{
@@ -161,7 +173,9 @@ void FLiveLinkMessageBusSource::InternalHandleMessage(const TSharedRef<IMessageC
 
 	const bool bIsStaticData = MessageTypeInfo->IsChildOf(FLiveLinkBaseStaticData::StaticStruct());
 	const bool bIsFrameData = MessageTypeInfo->IsChildOf(FLiveLinkBaseFrameData::StaticStruct());
-	if (!bIsStaticData && !bIsFrameData)
+	const bool bIsSerializedData = MessageTypeInfo->IsChildOf(FLiveLinkSerializedFrameData::StaticStruct());
+
+	if (!bIsStaticData && !bIsFrameData && !bIsSerializedData)
 	{
 		return;
 	}
@@ -294,12 +308,45 @@ void FLiveLinkMessageBusSource::InitializeAndPushFrameData_AnyThread(FName Subje
 																	 const TSharedRef<IMessageContext, ESPMode::ThreadSafe>& Context,
 																	 UScriptStruct* MessageTypeInfo)
 {
-	check(MessageTypeInfo->IsChildOf(FLiveLinkBaseFrameData::StaticStruct()));
+	FLiveLinkFrameDataStruct DataStruct;
+
+	if (MessageTypeInfo && MessageTypeInfo->IsChildOf(FLiveLinkSerializedFrameData::StaticStruct()))
+	{
+		// Extract the message from the compressed serialized data.
+		const FLiveLinkSerializedFrameData* SerializedMessage = reinterpret_cast<const FLiveLinkSerializedFrameData*>(Context->GetMessage());
+		FStructOnScope Payload;
+		SerializedMessage->GetPayload(Payload);
+
+		// Special case: We want to send anim data.
+		if (Payload.GetStruct() == FLiveLinkFloatAnimationFrameData::StaticStruct())
+		{
+			const FLiveLinkFloatAnimationFrameData* FloatAnimData = reinterpret_cast<const FLiveLinkFloatAnimationFrameData*>(Payload.GetStructMemory());
+			FLiveLinkAnimationFrameData DoubleFrameData = FLiveLinkFloatAnimationFrameData::ToAnimData(*FloatAnimData);
+
+			DataStruct.InitializeWith(FLiveLinkAnimationFrameData::StaticStruct(), &DoubleFrameData);
+			DataStruct.GetBaseData()->WorldTime = DoubleFrameData.WorldTime.GetOffsettedTime();
+		}
+		else
+		{
+			check(Payload.GetStruct()->IsChildOf<FLiveLinkBaseFrameData>());
+
+			const FLiveLinkBaseFrameData* Message = reinterpret_cast<const FLiveLinkBaseFrameData*>(Payload.GetStructMemory());
+			DataStruct.InitializeWith((UScriptStruct*)Payload.GetStruct(), Message);
+			DataStruct.GetBaseData()->WorldTime = Message->WorldTime.GetOffsettedTime();
+		}
+	}
+	else if (MessageTypeInfo && MessageTypeInfo->IsChildOf(FLiveLinkBaseFrameData::StaticStruct()))
+	{
+		const FLiveLinkBaseFrameData* Message = reinterpret_cast<const FLiveLinkBaseFrameData*>(Context->GetMessage());
+		DataStruct.InitializeWith(MessageTypeInfo, Message);
+		DataStruct.GetBaseData()->WorldTime = Message->WorldTime.GetOffsettedTime();
+	}
+	else
+	{
+		static const FName NAME_InvalidFrameData = "LiveLinkMessageBusSource_InvalidFrameData";
+		FLiveLinkLog::ErrorOnce(NAME_InvalidFrameData, SubjectKey, TEXT("Invalid frame data was provided for '%s' with connection '%s'"), *SubjectName.ToString(), *GetSourceMachineName().ToString());
+	}
 
-	FLiveLinkFrameDataStruct DataStruct(MessageTypeInfo);
-	const FLiveLinkBaseFrameData* Message = reinterpret_cast<const FLiveLinkBaseFrameData*>(Context->GetMessage());
-	DataStruct.InitializeWith(MessageTypeInfo, Message);
-	DataStruct.GetBaseData()->WorldTime = Message->WorldTime.GetOffsettedTime();
 	PushClientSubjectFrameData_AnyThread(SubjectKey, MoveTemp(DataStruct));
 }
 
@@ -374,7 +421,10 @@ void FLiveLinkMessageBusSource::SendConnectMessage()
 {
 	FLiveLinkConnectMessage* ConnectMessage = FMessageEndpoint::MakeMessage<FLiveLinkConnectMessage>();
 	ConnectMessage->LiveLinkVersion = ILiveLinkClient::LIVELINK_VERSION;
-	SendMessage(ConnectMessage);
+
+	TMap<FName, FString> Annotations;
+	AddAnnotations(Annotations);
+	SendMessage(ConnectMessage, Annotations);
 	StartHeartbeatEmitter();
 	bIsValid = true;
 }
@@ -422,3 +472,13 @@ PRAGMA_DISABLE_DEPRECATION_WARNINGS
 	.WithCatchall(this, &FLiveLinkMessageBusSource::InternalHandleMessage);
 PRAGMA_ENABLE_DEPRECATION_WARNINGS
 }
+
+void FLiveLinkMessageBusSource::AddAnnotations(TMap<FName, FString>& InOutAnnotations) const
+{
+	if (CVarLiveLinkSupportCompressPayloads.GetValueOnAnyThread())
+	{
+		// The presence of this flag in the annotation will inform our provider that we support receiving compressed animation.
+		InOutAnnotations.Add({ FLiveLinkMessageAnnotation::CompressedPayloadSupport, TEXT("") });
+	}
+}
+
diff --git a/Engine/Plugins/Animation/LiveLink/Source/LiveLink/Public/LiveLinkMessageBusSource.h b/Engine/Plugins/Animation/LiveLink/Source/LiveLink/Public/LiveLinkMessageBusSource.h
--- a/Engine/Plugins/Animation/LiveLink/Source/LiveLink/Public/LiveLinkMessageBusSource.h
+++ b/Engine/Plugins/Animation/LiveLink/Source/LiveLink/Public/LiveLinkMessageBusSource.h
@@ -51,6 +51,8 @@ protected:
 	virtual const FName& GetSourceName() const;
 	// This lets child classes the opportunity to add custom message handlers to the endpoint builder
 	virtual void InitializeMessageEndpoint(FMessageEndpointBuilder& EndpointBuilder);
+	// Add compress payload support annotations.
+	void AddAnnotations(TMap<FName, FString>& InOutAnnotations) const;
 	bool IsMessageEndpointConnected() const { return ConnectionAddress.IsValid() && MessageEndpoint.IsValid() && MessageEndpoint->IsConnected(); }
 
 	// Initialize the static data and send it to the clients
@@ -78,14 +80,14 @@ protected:
 
 	// Send a message through the endpoint
 	template<typename MessageType>
-	void SendMessage(MessageType* Message)
+	void SendMessage(MessageType* Message, const TMap<FName, FString>& Annotations = {})
 	{
 		if (!Message || !IsMessageEndpointConnected())
 		{
 			return;
 		}
 
-		MessageEndpoint->Send(Message, ConnectionAddress);
+		MessageEndpoint->Send(Message, EMessageFlags::None, Annotations, nullptr, { ConnectionAddress }, FTimespan::Zero(), FDateTime::MaxValue());
 	}
 
 	// Start the heartbeat emitter for this connection
diff --git a/Engine/Plugins/Animation/LiveLinkHub/Source/LiveLinkHubMessaging/Private/LiveLinkHubMessageBusSource.cpp b/Engine/Plugins/Animation/LiveLinkHub/Source/LiveLinkHubMessaging/Private/LiveLinkHubMessageBusSource.cpp
--- a/Engine/Plugins/Animation/LiveLinkHub/Source/LiveLinkHubMessaging/Private/LiveLinkHubMessageBusSource.cpp
+++ b/Engine/Plugins/Animation/LiveLinkHub/Source/LiveLinkHubMessaging/Private/LiveLinkHubMessageBusSource.cpp
@@ -62,7 +62,10 @@ void FLiveLinkHubMessageBusSource::SendConnectMessage()
 	FLiveLinkHubConnectMessage* ConnectMessage = FMessageEndpoint::MakeMessage<FLiveLinkHubConnectMessage>();
 	ConnectMessage->ClientInfo = CreateLiveLinkClientInfo();
 
-	SendMessage(ConnectMessage);
+	TMap<FName, FString> Annotations;
+	AddAnnotations(Annotations);
+
+	SendMessage(ConnectMessage, Annotations);
 	StartHeartbeatEmitter();
 	bIsValid = true;
 }
diff --git a/Engine/Source/Runtime/LiveLinkMessageBusFramework/Private/LiveLinkProvider.cpp b/Engine/Source/Runtime/LiveLinkMessageBusFramework/Private/LiveLinkProvider.cpp
--- a/Engine/Source/Runtime/LiveLinkMessageBusFramework/Private/LiveLinkProvider.cpp
+++ b/Engine/Source/Runtime/LiveLinkMessageBusFramework/Private/LiveLinkProvider.cpp
@@ -3,10 +3,12 @@
 #include "LiveLinkProvider.h"
 #include "LiveLinkProviderImpl.h"
 
+#include "Algo/Find.h"
 #include "Algo/RemoveIf.h"
 #include "Algo/Transform.h"
 #include "HAL/PlatformProcess.h"
 #include "IMessageContext.h"
+#include "LiveLinkCompression.h"
 #include "LiveLinkMessages.h"
 #include "LiveLinkTypes.h"
 
@@ -14,6 +16,7 @@
 #include "MessageEndpoint.h"
 #include "MessageEndpointBuilder.h"
 #include "Misc/ScopeLock.h"
+#include "Roles/LiveLinkAnimationTypes.h"
 
 DEFINE_LOG_CATEGORY_STATIC(LogLiveLinkMessageBus, Warning, All);
 
@@ -26,6 +29,12 @@ FName FLiveLinkMessageAnnotation::SubjectAnnotation = TEXT("SubjectName");
 FName FLiveLinkMessageAnnotation::RoleAnnotation = TEXT("Role");
 FName FLiveLinkMessageAnnotation::OriginalSourceAnnotation = TEXT("OriginalSource");
 FName FLiveLinkMessageAnnotation::TopologyModeAnnotation = TEXT("TopologyMode");
+FName FLiveLinkMessageAnnotation::CompressedPayloadSupport = TEXT("CompressedPayloadSupport");
+
+static TAutoConsoleVariable<int32> CVarLiveLinkCompressAnimData(
+	TEXT("LiveLink.CompressAnimationData"), 1,
+	TEXT("Whether to compress animation data that goes through the provider."),
+	ECVF_RenderThreadSafe);
 
 
 namespace UE::LiveLink::Messaging
@@ -76,13 +85,15 @@ namespace UE::LiveLink::Messaging
 // Address that we have had a connection request from
 struct FTrackedAddress
 {
-	FTrackedAddress(FMessageAddress InAddress)
+	FTrackedAddress(FMessageAddress InAddress, bool bInSupportsCompressedPayloads)
 		: Address(InAddress)
 		, LastHeartbeatTime(FPlatformTime::Seconds())
+		, bSupportsCompressedPayloads(bInSupportsCompressedPayloads)
 	{}
 
 	FMessageAddress Address;
-	double			LastHeartbeatTime;
+	double LastHeartbeatTime;
+	bool bSupportsCompressedPayloads;
 };
 
 
@@ -480,7 +491,43 @@ bool FLiveLinkProvider::UpdateSubjectFrameData(const FName SubjectName, FLiveLin
 		TArray<FMessageAddress> Addresses;
 		GetFilteredAddresses(SubjectName, Addresses);
 
-		MessageEndpoint->Send(FrameData.CloneData(), const_cast<UScriptStruct*>(FrameData.GetStruct()), EMessageFlags::None, Annotations, nullptr, Addresses, FTimespan::Zero(), FDateTime::MaxValue());
+		bool bCanCompressData = CVarLiveLinkCompressAnimData.GetValueOnAnyThread() && FrameData.GetStruct()->IsChildOf(FLiveLinkAnimationFrameData::StaticStruct());
+		TMap<FMessageAddress, bool> SupportsCompression;
+
+		bool bAnyClientSupportsCompression = false;
+		for (const FMessageAddress& Address : Addresses)
+		{
+			FTrackedAddress* TrackedAddress = Algo::FindBy(ConnectedAddresses, Address, &FTrackedAddress::Address);
+			check(TrackedAddress);
+			const bool bSupportsCompression = TrackedAddress->bSupportsCompressedPayloads;
+			SupportsCompression.Add(Address, bSupportsCompression);
+
+			bAnyClientSupportsCompression = bAnyClientSupportsCompression || bSupportsCompression;
+		}
+
+		bCanCompressData = bCanCompressData && bAnyClientSupportsCompression;
+
+		FLiveLinkSerializedFrameData SerializedData;
+		if (bCanCompressData && FrameData.GetStruct()->IsChildOf(FLiveLinkAnimationFrameData::StaticStruct()))
+		{
+			// If this is animation data, serialize and compress it to reduce throughput.
+			SerializedData = SerializeAnimationData(FrameData.Cast<FLiveLinkAnimationFrameData>());
+		}
+	
+		for (const FMessageAddress& Address : Addresses)
+		{
+			const bool bClientSupportsCompression = bCanCompressData && SupportsCompression.FindChecked(Address);
+
+			if (bClientSupportsCompression)
+			{
+				FLiveLinkSerializedFrameData* SerializedFrameMessage = FMessageEndpoint::MakeMessage<FLiveLinkSerializedFrameData>(SerializedData);
+				MessageEndpoint->Send(SerializedFrameMessage, FLiveLinkSerializedFrameData::StaticStruct(), EMessageFlags::None, Annotations, nullptr, { Address }, FTimespan::Zero(), FDateTime::MaxValue());
+			}
+			else
+			{
+				MessageEndpoint->Send(FrameData.CloneData(), const_cast<UScriptStruct*>(FrameData.GetStruct()), EMessageFlags::None, Annotations, nullptr, { Address }, FTimespan::Zero(), FDateTime::MaxValue());
+			}
+		}
 	}
 
 	SetLastSubjectFrameData(SubjectName, MoveTemp(FrameData), MoveTemp(Annotations));
@@ -488,6 +535,17 @@ bool FLiveLinkProvider::UpdateSubjectFrameData(const FName SubjectName, FLiveLin
 	return true;
 }
 
+FLiveLinkSerializedFrameData FLiveLinkProvider::SerializeAnimationData(FLiveLinkAnimationFrameData* AnimData)
+{
+	FLiveLinkSerializedFrameData SerializedFrameData;
+
+	// Convert double FTransform to float FTransform to further reduce payload size.
+	FLiveLinkFloatAnimationFrameData FloatTransforms = FLiveLinkFloatAnimationFrameData::FromAnimData(*AnimData);
+	SerializedFrameData.SetTypedPayload(FloatTransforms, ELiveLinkPayloadCompressionType::Always);
+
+	return SerializedFrameData;
+}
+
 bool FLiveLinkProvider::HasConnection() const
 {
 	FScopeLock Lock(&CriticalSection);
@@ -541,7 +599,9 @@ void FLiveLinkProvider::HandleConnectMessage(const FLiveLinkConnectMessage& Mess
 
 	if (!ConnectedAddresses.ContainsByPredicate([=](const FTrackedAddress& Address) { return Address.Address == ConnectionAddress; }))
 	{
-		ConnectedAddresses.Add(FTrackedAddress(ConnectionAddress));
+		// Keep track of what clients supports compressed payloads to avoid sending them data they can't interpret.
+		const bool bSupportsCompressedPayloads = Context->GetAnnotations().Contains(FLiveLinkMessageAnnotation::CompressedPayloadSupport);
+		ConnectedAddresses.Add(FTrackedAddress(ConnectionAddress, bSupportsCompressedPayloads));
 
 		// LiveLink version 1 path
 		for (const auto& Subject : Subjects)
diff --git a/Engine/Source/Runtime/LiveLinkMessageBusFramework/Public/LiveLinkMessages.h b/Engine/Source/Runtime/LiveLinkMessageBusFramework/Public/LiveLinkMessages.h
--- a/Engine/Source/Runtime/LiveLinkMessageBusFramework/Public/LiveLinkMessages.h
+++ b/Engine/Source/Runtime/LiveLinkMessageBusFramework/Public/LiveLinkMessages.h
@@ -11,8 +11,10 @@ struct LIVELINKMESSAGEBUSFRAMEWORK_API FLiveLinkMessageAnnotation
 	static FName SubjectAnnotation;
 	static FName RoleAnnotation;
 	static FName OriginalSourceAnnotation;
-	/** Annotation to indicate what mode this instance is running in. */
+	/** Indicates what mode this instance is running in. */
 	static FName TopologyModeAnnotation;
+	/** Presense of this annotation informs providers that this client supports compressed animation payloads. */
+	static FName CompressedPayloadSupport;
 };
 
 UENUM()
diff --git a/Engine/Source/Runtime/LiveLinkMessageBusFramework/Public/LiveLinkProviderImpl.h b/Engine/Source/Runtime/LiveLinkMessageBusFramework/Public/LiveLinkProviderImpl.h
--- a/Engine/Source/Runtime/LiveLinkMessageBusFramework/Public/LiveLinkProviderImpl.h
+++ b/Engine/Source/Runtime/LiveLinkMessageBusFramework/Public/LiveLinkProviderImpl.h
@@ -80,6 +80,9 @@ private:
 	// Get the connected addresses that should receive livelink data.
 	void GetFilteredAddresses(FName SubjectName, TArray<FMessageAddress>& Addresses);
 
+	// Serialize animation frame data into a compressed payload to reduce throughput.
+	struct FLiveLinkSerializedFrameData SerializeAnimationData(struct FLiveLinkAnimationFrameData* AnimData);
+
 protected:
 	// Update connected addresses and send information to the connected source
 	void HandleConnectMessage(const FLiveLinkConnectMessage& Message, const TSharedRef<class IMessageContext, ESPMode::ThreadSafe>& Context);
