LiveLink - Fix GetLiveLinkSubjectState not working correctly for virtual subjects

diff --git a/Engine/Plugins/Animation/LiveLink/Source/LiveLink/Private/LiveLinkClient.cpp b/Engine/Plugins/Animation/LiveLink/Source/LiveLink/Private/LiveLinkClient.cpp
--- a/Engine/Plugins/Animation/LiveLink/Source/LiveLink/Private/LiveLinkClient.cpp
+++ b/Engine/Plugins/Animation/LiveLink/Source/LiveLink/Private/LiveLinkClient.cpp
@@ -434,9 +434,24 @@ ELiveLinkSubjectState FLiveLinkClient::GetSubjectState(FLiveLinkSubjectName InSu
 
 	if (!LiveSubject)
 	{
-		if (SubjectItem->GetVirtualSubject() && SubjectItem->bEnabled)
+		if (ULiveLinkVirtualSubject* VirtualSubject = SubjectItem->GetVirtualSubject())
 		{
-			return SubjectItem->GetVirtualSubject()->IsPaused() ? ELiveLinkSubjectState::Paused : ELiveLinkSubjectState::Connected;
+			if (!SubjectItem->bEnabled)
+			{
+				return ELiveLinkSubjectState::InvalidOrDisabled;
+			}
+			else if (VirtualSubject->IsPaused())
+			{
+				return ELiveLinkSubjectState::Paused;
+			}
+			else if (VirtualSubject->HasValidFrameSnapshot())
+			{
+				return ELiveLinkSubjectState::Connected;
+			}
+			else
+			{
+				return ELiveLinkSubjectState::InvalidOrDisabled;
+			}
 		}
 		else
 		{
