LiveLinkHub - Update Hub/Spoke mode icons

#rnx

diff --git a/Engine/Plugins/Animation/LiveLink/Source/LiveLinkEditor/Private/LiveLinkEditorModule.cpp b/Engine/Plugins/Animation/LiveLink/Source/LiveLinkEditor/Private/LiveLinkEditorModule.cpp
--- a/Engine/Plugins/Animation/LiveLink/Source/LiveLinkEditor/Private/LiveLinkEditorModule.cpp
+++ b/Engine/Plugins/Animation/LiveLink/Source/LiveLinkEditor/Private/LiveLinkEditorModule.cpp
@@ -118,8 +118,8 @@ public:
 		StyleSet->Set("LiveLinkHub.Icon", new IMAGE_PLUGIN_BRUSH_SVG(TEXT("Starship/LiveLinkHub_Icon"), Icon40x40));
 		StyleSet->Set("LiveLinkHub.EllipsisIcon", new IMAGE_PLUGIN_BRUSH_SVG(TEXT("Starship/Ellipsis"), Icon16x16));
 		StyleSet->Set("LiveLinkHub.TimecodeGenlock", new IMAGE_PLUGIN_BRUSH_SVG(TEXT("Starship/TimecodeGenlock"), Icon16x16));
-		StyleSet->Set("LiveLinkHub.SpokeMode", new IMAGE_PLUGIN_BRUSH_SVG(TEXT("Starship/SpokeMode"), Icon16x16));
-		StyleSet->Set("LiveLinkHub.HubMode", new IMAGE_PLUGIN_BRUSH_SVG(TEXT("Starship/HubMode"), Icon16x16));
+		StyleSet->Set("LiveLinkHub.SpokeMode", new IMAGE_PLUGIN_BRUSH_SVG(TEXT("Starship/SpokeMode"), Icon16x16, FSlateColor(FLinearColor(1.000, 0.168, 0.350))));
+		StyleSet->Set("LiveLinkHub.HubMode", new IMAGE_PLUGIN_BRUSH_SVG(TEXT("Starship/HubMode"), Icon16x16, FSlateColor(FLinearColor(0.271, 0.554, 0.067))));
 
 
 		StyleSet->Set("ClassIcon.LiveLinkPreset",                      new IMAGE_PLUGIN_BRUSH_SVG("Starship/LiveLink", Icon16x16));
