LiveLink - Trigger subject state changed delegate when pausing / unpausing a subject

diff --git a/Engine/Plugins/Animation/LiveLink/Source/LiveLink/Private/LiveLinkClient.cpp b/Engine/Plugins/Animation/LiveLink/Source/LiveLink/Private/LiveLinkClient.cpp
--- a/Engine/Plugins/Animation/LiveLink/Source/LiveLink/Private/LiveLinkClient.cpp
+++ b/Engine/Plugins/Animation/LiveLink/Source/LiveLink/Private/LiveLinkClient.cpp
@@ -1246,6 +1246,8 @@ void FLiveLinkClient::PauseSubject_AnyThread(FLiveLinkSubjectName SubjectName)
 			{
 				VirtualSubject->PauseSubject();
 			}
+
+			OnLiveLinkSubjectStateChanged().Broadcast(SubjectItem->Key, ELiveLinkSubjectState::Paused);
 		}
 	}
 }
@@ -1259,10 +1261,18 @@ void FLiveLinkClient::UnpauseSubject_AnyThread(FLiveLinkSubjectName SubjectName)
 			if (FLiveLinkSubject* LiveSubject = SubjectItem->GetLiveSubject())
 			{
 				LiveSubject->UnpauseSubject();
+				OnLiveLinkSubjectStateChanged().Broadcast(SubjectItem->Key, LiveSubject->State);
 			}
 			else if (ULiveLinkVirtualSubject* VirtualSubject = SubjectItem->GetVirtualSubject())
 			{
 				VirtualSubject->UnpauseSubject();
+				ELiveLinkSubjectState State = ELiveLinkSubjectState::Connected;
+				if (!VirtualSubject->HasValidFrameSnapshot())
+				{
+					State = ELiveLinkSubjectState::InvalidOrDisabled;
+				}
+
+				OnLiveLinkSubjectStateChanged().Broadcast(SubjectItem->Key, State);
 			}
 		}
 	}
