LiveLink - Move LiveLinkHubSubjectSettings from LiveLinkHub module to LiveLink module.

This change will allow LiveLink sources such as the metahuman video LiveLink source to extend LiveLinkHubSubjectSettings

#rnx

diff --git a/Engine/Plugins/Animation/LiveLink/Source/LiveLink/Private/LiveLinkHubSubjectSettings.cpp b/Engine/Plugins/Animation/LiveLink/Source/LiveLink/Private/LiveLinkHubSubjectSettings.cpp
new file mode 100644
--- /dev/null
+++ b/Engine/Plugins/Animation/LiveLink/Source/LiveLink/Private/LiveLinkHubSubjectSettings.cpp
@@ -0,0 +1,118 @@
+// Copyright Epic Games, Inc. All Rights Reserved.
+
+#include "LiveLinkHubSubjectSettings.h"
+
+#include "Features/IModularFeatures.h"
+#include "LiveLinkClient.h"
+#include "LiveLinkModule.h"
+#include "Misc/ConfigCacheIni.h"
+#include "Modules/ModuleManager.h"
+
+namespace UE::Private::Utils
+{
+	/** Notify connected clients that this subject's name has changed. */
+	static bool ValidateOutboundName(const FString& SubjectName, FName PreviousOutboundName, const FString& InOutboundNameCandidate)
+	{
+		if (InOutboundNameCandidate.IsEmpty() || FName(InOutboundNameCandidate) == NAME_None)
+		{
+			return false;
+		}
+
+		if (InOutboundNameCandidate == SubjectName)
+		{
+			return true;
+		}
+
+		ILiveLinkClient* LiveLinkClient = &IModularFeatures::Get().GetModularFeature<ILiveLinkClient>(ILiveLinkClient::ModularFeatureName);
+
+		// Can't rename to an existing subject, so check if it exists by querying for its role.
+		return LiveLinkClient->GetSubjectRole_AnyThread(FLiveLinkSubjectName(*InOutboundNameCandidate)).Get() == nullptr;
+	}
+}
+
+void ULiveLinkHubSubjectSettings::Initialize(FLiveLinkSubjectKey InSubjectKey)
+{
+	if (Key != InSubjectKey)
+	{
+		ILiveLinkClient* LiveLinkClient = &IModularFeatures::Get().GetModularFeature<ILiveLinkClient>(ILiveLinkClient::ModularFeatureName);
+		SubjectName = InSubjectKey.SubjectName.ToString();
+		Key = InSubjectKey;
+
+		OutboundName = SubjectName;
+
+		Source = LiveLinkClient->GetSourceType(InSubjectKey.Source).ToString();
+	}
+}
+
+#if WITH_EDITOR
+void ULiveLinkHubSubjectSettings::PreEditChange(FProperty* PropertyAboutToChange)
+{
+	Super::PreEditChange(PropertyAboutToChange);
+
+	if (PropertyAboutToChange && PropertyAboutToChange->GetFName() == GET_MEMBER_NAME_CHECKED(ULiveLinkHubSubjectSettings, OutboundName))
+	{
+		PreviousOutboundName = *OutboundName;
+	}
+}
+
+void ULiveLinkHubSubjectSettings::PostEditChangeProperty(FPropertyChangedEvent& PropertyChangedEvent)
+{
+	Super::PostEditChangeProperty(PropertyChangedEvent);
+
+	const bool bIsInLiveLinkHubApp = GConfig->GetBoolOrDefault(TEXT("LiveLink"), TEXT("bCreateLiveLinkHubInstance"), false, GEngineIni);
+
+	if (!bIsInLiveLinkHubApp)
+	{
+		// We're not using most of these settings in editor, so there's no use in calling CacheSubjectSettings or modifying the OutboundName.
+		return;
+	}
+
+	if (PropertyChangedEvent.GetPropertyName() == GET_MEMBER_NAME_CHECKED(ULiveLinkHubSubjectSettings, OutboundName))
+	{
+		if (PreviousOutboundName != *OutboundName)
+		{
+			if (UE::Private::Utils::ValidateOutboundName(SubjectName, PreviousOutboundName, OutboundName))
+			{
+				// LiveLinkModule->NotifyRename
+				FLiveLinkModule& LiveLinkModule = FModuleManager::Get().GetModuleChecked<FLiveLinkModule>("LiveLink");
+				LiveLinkModule.OnSubjectOutboundNameModified().Broadcast(Key, PreviousOutboundName.ToString(), OutboundName);
+			}
+			else
+			{
+				OutboundName = PreviousOutboundName.ToString();
+			}
+		}
+	}
+	else if (PropertyChangedEvent.GetPropertyName() == GET_MEMBER_NAME_CHECKED(ULiveLinkHubSubjectSettings, Translators)
+		|| PropertyChangedEvent.GetPropertyName() == GET_MEMBER_NAME_CHECKED(ULiveLinkHubSubjectSettings, PreProcessors)
+		|| PropertyChangedEvent.GetPropertyName() == GET_MEMBER_NAME_CHECKED(ULiveLinkHubSubjectSettings, Remapper))
+	{
+		FLiveLinkClient* LiveLinkClient = static_cast<FLiveLinkClient*>(&IModularFeatures::Get().GetModularFeature<ILiveLinkClient>(ILiveLinkClient::ModularFeatureName));
+		LiveLinkClient->CacheSubjectSettings(Key, this);
+	}
+	else if (PropertyChangedEvent.GetPropertyName() == GET_MEMBER_NAME_CHECKED(ULiveLinkHubSubjectSettings, TranslatorsProxy))
+	{
+		Translators.Reset(1);
+		if (TranslatorsProxy)
+		{
+			Translators.Add(TranslatorsProxy);
+		}
+
+		ValidateProcessors();
+
+		// Re-assign TranslatorsProxy in case the translator was denied in the validate function.
+		if (Translators.Num())
+		{
+			TranslatorsProxy = Translators[0];
+		}
+		else
+		{
+			TranslatorsProxy = nullptr;
+		}
+
+
+		FLiveLinkClient* LiveLinkClient = static_cast<FLiveLinkClient*>(&IModularFeatures::Get().GetModularFeature<ILiveLinkClient>(ILiveLinkClient::ModularFeatureName));
+		LiveLinkClient->CacheSubjectSettings(Key, this);
+	}
+}
+#endif
diff --git a/Engine/Plugins/Animation/LiveLink/Source/LiveLink/Private/LiveLinkModule.h b/Engine/Plugins/Animation/LiveLink/Source/LiveLink/Private/LiveLinkModule.h
--- a/Engine/Plugins/Animation/LiveLink/Source/LiveLink/Private/LiveLinkModule.h
+++ b/Engine/Plugins/Animation/LiveLink/Source/LiveLink/Private/LiveLinkModule.h
@@ -42,6 +42,11 @@ public:
 		return RegisteredSourceFilters;
 	}
 
+	virtual FOnSubjectOutboundNameModified& OnSubjectOutboundNameModified() override
+	{
+		return OutboundNameModifiedDelegate;
+	}
+
 	//~ End ILiveLinkModule interface
 
 private:
@@ -63,4 +68,7 @@ private:
 
 	// Filters registered externally used to filter what message bus sources should be displayed in the UI.
 	TMap<FDelegateHandle, FOnLiveLinkShouldDisplaySource> RegisteredSourceFilters;
+
+	// Delegate called by LiveLinkHubSubjectSettings to notify the hub of an outbound name change.
+	FOnSubjectOutboundNameModified OutboundNameModifiedDelegate;
 };
diff --git a/Engine/Plugins/Animation/LiveLink/Source/LiveLink/Public/ILiveLinkModule.h b/Engine/Plugins/Animation/LiveLink/Source/LiveLink/Public/ILiveLinkModule.h
--- a/Engine/Plugins/Animation/LiveLink/Source/LiveLink/Public/ILiveLinkModule.h
+++ b/Engine/Plugins/Animation/LiveLink/Source/LiveLink/Public/ILiveLinkModule.h
@@ -12,6 +12,7 @@
 class FDelegateHandle;
 class FLiveLinkHeartbeatEmitter;
 class FLiveLinkMessageBusDiscoveryManager;
+struct FLiveLinkSubjectKey;
 struct FProviderPollResult;
 class FSlateStyleSet;
 
@@ -68,6 +69,9 @@ public:
 	 */
 	virtual const TMap<FDelegateHandle, FOnLiveLinkShouldDisplaySource>& GetSourceFilters() const = 0;
 
+	DECLARE_TS_MULTICAST_DELEGATE_ThreeParams(FOnSubjectOutboundNameModified, const FLiveLinkSubjectKey& /*Key*/, const FString& /*PreviousOutboundName*/, const FString& /*OutboundName*/);
+	/** (Internal use) Delegate called by LiveLinkHubSubjectSettings to notify. */
+	virtual FOnSubjectOutboundNameModified& OnSubjectOutboundNameModified() = 0;
 public:
 
 	/** Virtual destructor. */
diff --git a/Engine/Plugins/Animation/LiveLink/Source/LiveLink/Public/LiveLinkClient.h b/Engine/Plugins/Animation/LiveLink/Source/LiveLink/Public/LiveLinkClient.h
--- a/Engine/Plugins/Animation/LiveLink/Source/LiveLink/Public/LiveLinkClient.h
+++ b/Engine/Plugins/Animation/LiveLink/Source/LiveLink/Public/LiveLinkClient.h
@@ -177,6 +177,9 @@ public:
 	/** Push subject static data with additional metadata. */
 	UE_API void PushPendingSubject_AnyThread(FPendingSubjectStatic&& PendingSubject);
 
+	/** (Internal use only) Cache subject settings for a given subject. Currently no-op outside of LiveLinkHub. */
+	virtual void CacheSubjectSettings(const FLiveLinkSubjectKey& SubjectKey, ULiveLinkSubjectSettings* Settings) const {}
+
 protected:
 	/**
 	 * Add delegates that will be triggered for all subjects.
diff --git a/Engine/Plugins/Animation/LiveLink/Source/LiveLink/Public/LiveLinkHubSubjectSettings.h b/Engine/Plugins/Animation/LiveLink/Source/LiveLink/Public/LiveLinkHubSubjectSettings.h
new file mode 100644
--- /dev/null
+++ b/Engine/Plugins/Animation/LiveLink/Source/LiveLink/Public/LiveLinkHubSubjectSettings.h
@@ -0,0 +1,77 @@
+// Copyright Epic Games, Inc. All Rights Reserved.
+
+#pragma once
+
+#include "LiveLinkSubjectSettings.h"
+
+#include "LiveLinkHubSubjectSettings.generated.h"
+
+#define UE_API LIVELINK_API
+
+class FLiveLinkHubClient;
+class ULiveLinkRole;
+
+/** Settings object for a livelinkhub subject. */
+UCLASS(MinimalAPI)
+class ULiveLinkHubSubjectSettings : public ULiveLinkSubjectSettings
+{
+	GENERATED_BODY()
+public:
+
+	//~ Begin ULiveLinkSubjectSettings interface
+	UE_API virtual void Initialize(FLiveLinkSubjectKey InSubjectKey) override;
+
+	virtual FName GetRebroadcastName() const override
+	{
+		return *OutboundName;
+	}
+	//~ End ULiveLinkSubjectSettings interface
+
+#if WITH_EDITOR
+	//~ Begin UObject interface
+	UE_API virtual void PreEditChange(FProperty* PropertyAboutToChange) override;
+	UE_API void PostEditChangeProperty(FPropertyChangedEvent& PropertyChangedEvent) override;
+	//~ End UObject interface
+
+	DECLARE_MULTICAST_DELEGATE_OneParam(FOnPropertyModified, const FName& /*PropertyName*/);
+	FOnPropertyModified OnPropertyModified() { return OnPropertyModifiedDelegate; }
+#endif
+
+	static FName GetOutboundNamePropertyName()
+	{
+		return GET_MEMBER_NAME_CHECKED(ULiveLinkHubSubjectSettings, OutboundName);
+	}
+
+	static FName GetSubjectNamePropertyName()
+	{
+		return GET_MEMBER_NAME_CHECKED(ULiveLinkHubSubjectSettings, SubjectName);
+	}
+
+private:
+	/** Name of this subject. */
+	UPROPERTY(VisibleAnywhere, Category = "LiveLink")
+	FString SubjectName;
+
+	/** Name override that will be transmitted to clients instead of the subject name. */
+	UPROPERTY(EditAnywhere, Category = "LiveLink")
+	FString OutboundName;
+
+	/** Source that contains the subject. */
+	UPROPERTY(VisibleAnywhere, Category = "LiveLink")
+	FString Source;
+
+	/** Proxy property used edit the translators. */
+	UPROPERTY(EditAnywhere, Instanced, Category = "LiveLink", DisplayName = "Translator")
+	TObjectPtr<ULiveLinkFrameTranslator> TranslatorsProxy;
+
+private:
+	/** Previous outbound name to be used for noticing clients to remove this entry from their subject list. */
+	FName PreviousOutboundName;
+
+#if WITH_EDITOR
+	/** Triggered when a property is modified. */
+	FOnPropertyModified OnPropertyModifiedDelegate;
+#endif
+};
+
+#undef UE_API
\ No newline at end of file
diff --git a/Engine/Plugins/Animation/LiveLink/Source/LiveLinkEditor/Private/LiveLinkEditorModule.cpp b/Engine/Plugins/Animation/LiveLink/Source/LiveLinkEditor/Private/LiveLinkEditorModule.cpp
--- a/Engine/Plugins/Animation/LiveLink/Source/LiveLinkEditor/Private/LiveLinkEditorModule.cpp
+++ b/Engine/Plugins/Animation/LiveLink/Source/LiveLinkEditor/Private/LiveLinkEditorModule.cpp
@@ -30,6 +30,8 @@
 #include "LiveLinkEditorPrivate.h"
 #include "LiveLinkEditorSettings.h"
 #include "LiveLinkGraphPanelPinFactory.h"
+#include "LiveLinkHubSubjectSettings.h"
+#include "LiveLinkHubSubjectSettingsDetailsCustomization.h"
 #include "LiveLinkSettings.h"
 #include "LiveLinkSourceSettingsDetailCustomization.h"
 #include "LiveLinkSubjectKeyDetailCustomization.h"
@@ -286,6 +288,7 @@ private:
 		PropertyEditorModule.RegisterCustomClassLayout(ULiveLinkSourceSettings::StaticClass()->GetFName(), FOnGetDetailCustomizationInstance::CreateStatic(&FLiveLinkSourceSettingsDetailCustomization::MakeInstance));
 		PropertyEditorModule.RegisterCustomClassLayout(ULiveLinkComponentController::StaticClass()->GetFName(), FOnGetDetailCustomizationInstance::CreateStatic(&FLiveLinkComponentDetailCustomization::MakeInstance));
 		PropertyEditorModule.RegisterCustomClassLayout(ULiveLinkControllerBase::StaticClass()->GetFName(), FOnGetDetailCustomizationInstance::CreateStatic(&FLiveLinkControllerBaseDetailCustomization::MakeInstance));
+		PropertyEditorModule.RegisterCustomClassLayout(ULiveLinkHubSubjectSettings::StaticClass()->GetFName(), FOnGetDetailCustomizationInstance::CreateStatic(&FLiveLinkHubSubjectSettingsDetailsCustomization::MakeInstance));
 
 		LiveLinkGraphPanelPinFactory = MakeShared<FLiveLinkGraphPanelPinFactory>();
 		FEdGraphUtilities::RegisterVisualPinFactory(LiveLinkGraphPanelPinFactory);
@@ -299,6 +302,7 @@ private:
 			FPropertyEditorModule* PropertyEditorModule = FModuleManager::Get().GetModulePtr<FPropertyEditorModule>("PropertyEditor");
 			if (PropertyEditorModule)
 			{
+				PropertyEditorModule->UnregisterCustomClassLayout(ULiveLinkHubSubjectSettings::StaticClass()->GetFName());
 				PropertyEditorModule->UnregisterCustomClassLayout(ULiveLinkComponentController::StaticClass()->GetFName());
 				PropertyEditorModule->UnregisterCustomClassLayout(ULiveLinkControllerBase::StaticClass()->GetFName());
 				PropertyEditorModule->UnregisterCustomClassLayout(ULiveLinkSourceSettings::StaticClass()->GetFName());
diff --git a/Engine/Plugins/Animation/LiveLink/Source/LiveLinkEditor/Private/LiveLinkHubSubjectSettingsDetailsCustomization.cpp b/Engine/Plugins/Animation/LiveLink/Source/LiveLinkEditor/Private/LiveLinkHubSubjectSettingsDetailsCustomization.cpp
new file mode 100644
--- /dev/null
+++ b/Engine/Plugins/Animation/LiveLink/Source/LiveLinkEditor/Private/LiveLinkHubSubjectSettingsDetailsCustomization.cpp
@@ -0,0 +1,84 @@
+// Copyright Epic Games, Inc. All Rights Reserved.
+
+#include "LiveLinkHubSubjectSettingsDetailsCustomization.h"
+
+#include "DetailLayoutBuilder.h"
+#include "DetailWidgetRow.h"
+#include "IDetailPropertyRow.h"
+#include "LiveLinkHubSubjectSettings.h"
+#include "Misc/ConfigCacheIni.h"
+
+
+#define LOCTEXT_NAMESPACE "LiveLinkHubSubjectSettingsDetailCustomization"
+
+void FLiveLinkHubSubjectSettingsDetailsCustomization::CustomizeDetails(IDetailLayoutBuilder& InDetailBuilder)
+{
+	//Get the current settings object being edited
+	TArray<TWeakObjectPtr<UObject>> ObjectsBeingCustomized;
+	InDetailBuilder.GetObjectsBeingCustomized(ObjectsBeingCustomized);
+	if (ObjectsBeingCustomized.Num() != 1)
+	{
+		return;
+	}
+
+	const bool bIsInLiveLinkHubApp = GConfig->GetBoolOrDefault(TEXT("LiveLink"), TEXT("bCreateLiveLinkHubInstance"), false, GEngineIni);
+
+	if (bIsInLiveLinkHubApp)
+	{
+		TSharedRef<IPropertyHandle> OutboundProperty = InDetailBuilder.GetProperty(ULiveLinkHubSubjectSettings::GetOutboundNamePropertyName());
+		if (IDetailPropertyRow* PropertyRow = InDetailBuilder.EditDefaultProperty(OutboundProperty))
+		{
+			FResetToDefaultOverride ResetOverride;
+
+			const FName SubjectNamePropertyName = ULiveLinkHubSubjectSettings::GetSubjectNamePropertyName();
+
+			PropertyRow->CustomWidget()
+				.OverrideResetToDefault(FResetToDefaultOverride::Create(
+					FIsResetToDefaultVisible::CreateLambda([SubjectNamePropertyName](TSharedPtr<IPropertyHandle> PropertyHandle)
+						{
+							TSharedPtr<IPropertyHandle> SubjectNameProperty = PropertyHandle->GetParentHandle()->GetChildHandle(SubjectNamePropertyName);
+						
+							if (SubjectNameProperty && SubjectNameProperty->IsValidHandle())
+							{
+								FText OutboundName;
+								FText SubjectName;
+								PropertyHandle->GetValueAsDisplayText(OutboundName);
+								SubjectNameProperty->GetValueAsDisplayText(SubjectName);
+								return !OutboundName.EqualTo(SubjectName);
+							}
+							return false;
+						}),
+					FResetToDefaultHandler::CreateLambda([SubjectNamePropertyName](TSharedPtr<IPropertyHandle> PropertyHandle)
+						{
+							TSharedPtr<IPropertyHandle> SubjectNameProperty = PropertyHandle->GetParentHandle()->GetChildHandle(SubjectNamePropertyName);
+
+							if (SubjectNameProperty && SubjectNameProperty->IsValidHandle())
+							{
+								FString SubjectName;
+								SubjectNameProperty->GetValueAsDisplayString(SubjectName);
+								PropertyHandle->SetValue(*SubjectName);
+							}
+						})
+				))
+				.NameContent()
+				[
+					OutboundProperty->CreatePropertyNameWidget()
+				]
+				.ValueContent()
+				[
+					OutboundProperty->CreatePropertyValueWidget()
+				];
+		}
+
+		InDetailBuilder.HideProperty(GET_MEMBER_NAME_CHECKED(ULiveLinkHubSubjectSettings, InterpolationProcessor), ULiveLinkSubjectSettings::StaticClass());
+		InDetailBuilder.HideProperty(GET_MEMBER_NAME_CHECKED(ULiveLinkHubSubjectSettings, FrameRate), ULiveLinkSubjectSettings::StaticClass());
+		InDetailBuilder.HideProperty(GET_MEMBER_NAME_CHECKED(ULiveLinkHubSubjectSettings, bRebroadcastSubject), ULiveLinkSubjectSettings::StaticClass());
+		InDetailBuilder.HideProperty(GET_MEMBER_NAME_CHECKED(ULiveLinkHubSubjectSettings, Translators), ULiveLinkSubjectSettings::StaticClass());
+	}
+	else
+	{
+		InDetailBuilder.HideProperty(ULiveLinkHubSubjectSettings::GetOutboundNamePropertyName(), ULiveLinkHubSubjectSettings::StaticClass());
+	}
+}
+
+#undef LOCTEXT_NAMESPACE
diff --git a/Engine/Plugins/Animation/LiveLink/Source/LiveLinkEditor/Private/LiveLinkHubSubjectSettingsDetailsCustomization.h b/Engine/Plugins/Animation/LiveLink/Source/LiveLinkEditor/Private/LiveLinkHubSubjectSettingsDetailsCustomization.h
new file mode 100644
--- /dev/null
+++ b/Engine/Plugins/Animation/LiveLink/Source/LiveLinkEditor/Private/LiveLinkHubSubjectSettingsDetailsCustomization.h
@@ -0,0 +1,27 @@
+// Copyright Epic Games, Inc. All Rights Reserved.
+
+#pragma once
+
+#include "IDetailCustomization.h"
+
+class ULiveLinkHubSubjectSettings;
+
+
+/**
+* Customizes a ULiveLinkHubSubjectSettings object
+*/
+class FLiveLinkHubSubjectSettingsDetailsCustomization : public IDetailCustomization
+{
+public:
+
+	static TSharedRef<IDetailCustomization> MakeInstance()
+	{
+		return MakeShared<FLiveLinkHubSubjectSettingsDetailsCustomization>();
+	}
+
+	// IDetailCustomization interface
+	virtual void CustomizeDetails(class IDetailLayoutBuilder& DetailBuilder) override;
+	// End IDetailCustomization interface
+
+private:
+};
diff --git a/Engine/Plugins/Animation/LiveLinkHub/Source/LiveLinkHub/Private/LiveLinkHub.cpp b/Engine/Plugins/Animation/LiveLinkHub/Source/LiveLinkHub/Private/LiveLinkHub.cpp
--- a/Engine/Plugins/Animation/LiveLinkHub/Source/LiveLinkHub/Private/LiveLinkHub.cpp
+++ b/Engine/Plugins/Animation/LiveLinkHub/Source/LiveLinkHub/Private/LiveLinkHub.cpp
@@ -13,6 +13,7 @@
 #include "HAL/PlatformProcess.h"
 #include "ISettingsModule.h"
 #include "ILiveLinkHubMessagingModule.h"
+#include "ILiveLinkModule.h"
 #include "LiveLinkEditorSettings.h"
 #include "LiveLinkHubClient.h"
 #include "LiveLinkHubCommands.h"
@@ -202,6 +203,9 @@ void FLiveLinkHub::InitializePostEngineInit()
 		TimeAndSyncSettings->ApplyCustomTimeStep();
 	}
 
+	ILiveLinkModule& LiveLinkModule = FModuleManager::Get().GetModuleChecked<ILiveLinkModule>("LiveLink");
+	LiveLinkModule.OnSubjectOutboundNameModified().AddSP(this, &FLiveLinkHub::OnSubjectOutboundNameModified);
+
 	GIsRunning = true;
 
 	FString SessionPath;
@@ -455,6 +459,30 @@ void FLiveLinkHub::OnModularFeatureRegistered(const FName& Type, IModularFeature
 	}
 }
 
+void FLiveLinkHub::OnSubjectOutboundNameModified(const FLiveLinkSubjectKey& SubjectKey, const FString& PreviousOutboundName, const FString& NewOutboundName) const
+{
+	LiveLinkProvider->SendClearSubjectToConnections(*PreviousOutboundName);
+
+	FLiveLinkHubModule& LiveLinkHubModule = FModuleManager::Get().GetModuleChecked<FLiveLinkHubModule>("LiveLinkHub");
+
+	// Re-send the last static data with the new name.
+	TPair<UClass*, FLiveLinkStaticDataStruct*> StaticData = LiveLinkProvider->GetLastSubjectStaticDataStruct(*PreviousOutboundName);
+	if (StaticData.Key && StaticData.Value)
+	{
+		FLiveLinkStaticDataStruct StaticDataCopy;
+		StaticDataCopy.InitializeWith(*StaticData.Value);
+
+		TMap<FName, FString> ExtraAnnotations;
+
+		const FText OriginalSourceType = LiveLinkHubClient->GetSourceType(SubjectKey.Source);
+		ExtraAnnotations.Add(FLiveLinkMessageAnnotation::OriginalSourceAnnotation, OriginalSourceType.ToString());
+		LiveLinkProvider->UpdateSubjectStaticData(*NewOutboundName, StaticData.Key, MoveTemp(StaticDataCopy), ExtraAnnotations);
+	}
+
+	// Then clear the old static data entry in the provider.
+	LiveLinkProvider->RemoveSubject(*PreviousOutboundName);
+}
+
 void FLiveLinkHub::RegisterLiveLinkHubSettings()
 {
 	if (ISettingsModule* SettingsModule = FModuleManager::GetModulePtr<ISettingsModule>("Settings"))
diff --git a/Engine/Plugins/Animation/LiveLinkHub/Source/LiveLinkHub/Private/LiveLinkHub.h b/Engine/Plugins/Animation/LiveLinkHub/Source/LiveLinkHub/Private/LiveLinkHub.h
--- a/Engine/Plugins/Animation/LiveLinkHub/Source/LiveLinkHub/Private/LiveLinkHub.h
+++ b/Engine/Plugins/Animation/LiveLinkHub/Source/LiveLinkHub/Private/LiveLinkHub.h
@@ -136,6 +136,9 @@ private:
 
 	/** Handles adding application modes that are registered after LLH has finished initialization. */
 	void OnModularFeatureRegistered(const FName& Type, class IModularFeature* ModularFeature);
+
+	/** Delegate called when a LiveLinkHubSubjectSettings modifies an outbound name. We need to notify connected clients of this change. */
+	void OnSubjectOutboundNameModified(const FLiveLinkSubjectKey& SubjectKey, const FString& PreviousOutboundName, const FString& NewOutboundName) const;
 	
 private:
 	/** Implements the logic to manage the clients tabs. */
diff --git a/Engine/Plugins/Animation/LiveLinkHub/Source/LiveLinkHub/Private/LiveLinkHubClient.h b/Engine/Plugins/Animation/LiveLinkHub/Source/LiveLinkHub/Private/LiveLinkHubClient.h
--- a/Engine/Plugins/Animation/LiveLinkHub/Source/LiveLinkHub/Private/LiveLinkHubClient.h
+++ b/Engine/Plugins/Animation/LiveLinkHub/Source/LiveLinkHub/Private/LiveLinkHubClient.h
@@ -66,7 +66,7 @@ public:
 	}
 
 	/** Cache subject settings for the subject specified by the subject key. */
-	void CacheSubjectSettings(const FLiveLinkSubjectKey& SubjectKey, ULiveLinkSubjectSettings* Settings) const;
+	virtual void CacheSubjectSettings(const FLiveLinkSubjectKey& SubjectKey, ULiveLinkSubjectSettings* Settings) const override;
 
 public:
 	//~ Begin ILiveLinkClient interface
diff --git a/Engine/Plugins/Animation/LiveLinkHub/Source/LiveLinkHub/Private/LiveLinkHubModule.cpp b/Engine/Plugins/Animation/LiveLinkHub/Source/LiveLinkHub/Private/LiveLinkHubModule.cpp
--- a/Engine/Plugins/Animation/LiveLinkHub/Source/LiveLinkHub/Private/LiveLinkHubModule.cpp
+++ b/Engine/Plugins/Animation/LiveLinkHub/Source/LiveLinkHub/Private/LiveLinkHubModule.cpp
@@ -17,7 +17,6 @@
 #include "Settings/LiveLinkHubSettings.h"
 #include "Settings/LiveLinkHubSettingsCustomization.h"
 #include "Settings/LiveLinkSettingsCustomization.h"
-#include "Subjects/LiveLinkHubSubjectSettingsDetailsCustomization.h"
 #include "Textures/SlateIcon.h"
 #include "Widgets/Text/STextBlock.h"
 
@@ -57,9 +56,7 @@ void FLiveLinkHubModule::StartupModule()
 	if (bUseSettingsDetailCustomization)
 	{
 		PropertyModule.RegisterCustomClassLayout(ULiveLinkSettings::StaticClass()->GetFName(), FOnGetDetailCustomizationInstance::CreateStatic(&FLiveLinkSettingsCustomization::MakeInstance));
-	}
-	
-	PropertyModule.RegisterCustomClassLayout(ULiveLinkHubSubjectSettings::StaticClass()->GetFName(), FOnGetDetailCustomizationInstance::CreateStatic(&FLiveLinkHubSubjectSettingsDetailsCustomization::MakeInstance));
+	}	
 }
 
 void FLiveLinkHubModule::ShutdownModule()
@@ -71,11 +68,6 @@ void FLiveLinkHubModule::ShutdownModule()
 		{
 			PropertyEditorModule->UnregisterCustomClassLayout(ULiveLinkSettings::StaticClass()->GetFName());
 		}
-		
-		if (bUseSubjectSettingsDetailsCustomization)
-		{
-			PropertyEditorModule->UnregisterCustomClassLayout(ULiveLinkHubSubjectSettings::StaticClass()->GetFName());
-		}
 	}
 
 	if (bCreateLiveLinkHubInstance)
diff --git a/Engine/Plugins/Animation/LiveLinkHub/Source/LiveLinkHub/Private/LiveLinkHubSubjectSettings.cpp b/Engine/Plugins/Animation/LiveLinkHub/Source/LiveLinkHub/Private/LiveLinkHubSubjectSettings.cpp
deleted file mode 100644
diff --git a/Engine/Plugins/Animation/LiveLinkHub/Source/LiveLinkHub/Private/Subjects/LiveLinkHubSubjectSettingsDetailsCustomization.cpp b/Engine/Plugins/Animation/LiveLinkHub/Source/LiveLinkHub/Private/Subjects/LiveLinkHubSubjectSettingsDetailsCustomization.cpp
deleted file mode 100644
diff --git a/Engine/Plugins/Animation/LiveLinkHub/Source/LiveLinkHub/Private/Subjects/LiveLinkHubSubjectSettingsDetailsCustomization.h b/Engine/Plugins/Animation/LiveLinkHub/Source/LiveLinkHub/Private/Subjects/LiveLinkHubSubjectSettingsDetailsCustomization.h
deleted file mode 100644
diff --git a/Engine/Plugins/Animation/LiveLinkHub/Source/LiveLinkHub/Public/LiveLinkHubSubjectSettings.h b/Engine/Plugins/Animation/LiveLinkHub/Source/LiveLinkHub/Public/LiveLinkHubSubjectSettings.h
deleted file mode 100644
diff --git a/Engine/Source/Programs/LiveLinkHub/Config/DefaultEngine.ini b/Engine/Source/Programs/LiveLinkHub/Config/DefaultEngine.ini
--- a/Engine/Source/Programs/LiveLinkHub/Config/DefaultEngine.ini
+++ b/Engine/Source/Programs/LiveLinkHub/Config/DefaultEngine.ini
@@ -65,4 +65,5 @@ NodeMemoryLoggingThresholdMB=5
 
 [CoreRedirects]
 +PropertyRedirects=(OldName="LiveLinkHubSettings.bAutoAddDiscoveredClients",NewName="AutoConnectClients")
++ClassRedirects=(OldName="/Script/LiveLinkHub.LiveLinkHubSubjectSettings",NewName="/Script/LiveLink.LiveLinkHubSubjectSettings")
 
diff --git a/Engine/Source/Programs/LiveLinkHub/Config/DefaultGame.ini b/Engine/Source/Programs/LiveLinkHub/Config/DefaultGame.ini
--- a/Engine/Source/Programs/LiveLinkHub/Config/DefaultGame.ini
+++ b/Engine/Source/Programs/LiveLinkHub/Config/DefaultGame.ini
@@ -10,7 +10,7 @@ bIsForExternalDistribution=true
 
 
 [/Script/LiveLink.LiveLinkSettings]
-DefaultSettingsClass=/Script/LiveLinkHub.LiveLinkHubSubjectSettings
+DefaultSettingsClass=/Script/LiveLink.LiveLinkHubSubjectSettings
 bPreProcessRebroadcastFrames=true
 bTranslateRebroadcastFrames=true
 
