LiveLink - Allow sources to provide their own tooltip for their entry in the sources list
Add tooltip for MessageBus  that displays the current data throughput

diff --git a/Engine/Plugins/Animation/LiveLink/Source/LiveLink/LiveLink.Build.cs b/Engine/Plugins/Animation/LiveLink/Source/LiveLink/LiveLink.Build.cs
--- a/Engine/Plugins/Animation/LiveLink/Source/LiveLink/LiveLink.Build.cs
+++ b/Engine/Plugins/Animation/LiveLink/Source/LiveLink/LiveLink.Build.cs
@@ -23,6 +23,7 @@ namespace UnrealBuildTool.Rules
 				"InputCore",
 				"LiveLinkAnimationCore", //Adding dependency to make sure this runtime module is pulled when packaging if LiveLink plugin is enabled. 
 				"Media",
+				"Messaging",
 				"Projects",
 				"SlateCore",
 				"Slate",
diff --git a/Engine/Plugins/Animation/LiveLink/Source/LiveLink/Private/ClientNetworkStatisticsModel.cpp b/Engine/Plugins/Animation/LiveLink/Source/LiveLink/Private/ClientNetworkStatisticsModel.cpp
new file mode 100644
--- /dev/null
+++ b/Engine/Plugins/Animation/LiveLink/Source/LiveLink/Private/ClientNetworkStatisticsModel.cpp
@@ -0,0 +1,55 @@
+// Copyright Epic Games, Inc. All Rights Reserved.
+
+#include "ClientNetworkStatisticsModel.h"
+
+#include "Features/IModularFeatures.h"
+#include "INetworkMessagingExtension.h"
+#include "Misc/ScopeExit.h"
+
+namespace UE::LiveLinkHub::Private
+{
+	static INetworkMessagingExtension* GetMessagingStatistics()
+	{
+		IModularFeatures& ModularFeatures = IModularFeatures::Get();
+		if (IsInGameThread())
+		{
+			if (ModularFeatures.IsModularFeatureAvailable(INetworkMessagingExtension::ModularFeatureName))
+			{
+				return &ModularFeatures.GetModularFeature<INetworkMessagingExtension>(INetworkMessagingExtension::ModularFeatureName);
+			}
+		}
+		else
+		{
+			ModularFeatures.LockModularFeatureList();
+			ON_SCOPE_EXIT
+			{
+				ModularFeatures.UnlockModularFeatureList();
+			};
+			
+			if (ModularFeatures.IsModularFeatureAvailable(INetworkMessagingExtension::ModularFeatureName))
+			{
+				return &ModularFeatures.GetModularFeature<INetworkMessagingExtension>(INetworkMessagingExtension::ModularFeatureName);
+			}
+		}
+		
+
+		ensureMsgf(false, TEXT("Feature %s is unavailable"), *INetworkMessagingExtension::ModularFeatureName.ToString());
+		return nullptr;
+	}
+}
+
+TOptional<FMessageTransportStatistics> UE::LiveLinkHub::FClientNetworkStatisticsModel::GetLatestNetworkStatistics(const FMessageAddress& ClientAddress)
+{
+	if (INetworkMessagingExtension* Statistics = UE::LiveLinkHub::Private::GetMessagingStatistics())
+	{
+		const FGuid NodeId = Statistics->GetNodeIdFromAddress(ClientAddress);
+		return NodeId.IsValid() ? Statistics->GetLatestNetworkStatistics(NodeId) : TOptional<FMessageTransportStatistics>();
+	}
+	return {};
+}
+
+bool UE::LiveLinkHub::FClientNetworkStatisticsModel::IsOnline(const FMessageAddress& ClientAddress)
+{
+	INetworkMessagingExtension* Statistics = UE::LiveLinkHub::Private::GetMessagingStatistics();
+	return Statistics && Statistics->GetNodeIdFromAddress(ClientAddress).IsValid();
+}
diff --git a/Engine/Plugins/Animation/LiveLink/Source/LiveLink/Private/ClientNetworkStatisticsModel.h b/Engine/Plugins/Animation/LiveLink/Source/LiveLink/Private/ClientNetworkStatisticsModel.h
new file mode 100644
--- /dev/null
+++ b/Engine/Plugins/Animation/LiveLink/Source/LiveLink/Private/ClientNetworkStatisticsModel.h
@@ -0,0 +1,15 @@
+// Copyright Epic Games, Inc. All Rights Reserved.
+
+#pragma once
+
+#include "Misc/Optional.h"
+
+struct FMessageAddress;
+struct FMessageTransportStatistics;
+
+/** Synchronizes the network statistics (the statistics are updated async). */
+namespace UE::LiveLinkHub::FClientNetworkStatisticsModel
+{
+	TOptional<FMessageTransportStatistics> GetLatestNetworkStatistics(const FMessageAddress& ClientAddress);
+	bool IsOnline(const FMessageAddress& ClientAddress);
+}
diff --git a/Engine/Plugins/Animation/LiveLink/Source/LiveLink/Private/LiveLinkClient.cpp b/Engine/Plugins/Animation/LiveLink/Source/LiveLink/Private/LiveLinkClient.cpp
--- a/Engine/Plugins/Animation/LiveLink/Source/LiveLink/Private/LiveLinkClient.cpp
+++ b/Engine/Plugins/Animation/LiveLink/Source/LiveLink/Private/LiveLinkClient.cpp
@@ -1881,6 +1881,15 @@ FText FLiveLinkClient::GetSourceStatus(FGuid InEntryGuid) const
 	return FText(NSLOCTEXT("TempLocTextLiveLink", "InvalidSourceStatus", "Invalid Source Status"));
 }
 
+FText FLiveLinkClient::GetSourceToolTip(FGuid InEntryGuid) const
+{
+	if (const FLiveLinkCollectionSourceItem* SourceItem = Collection->FindSource(InEntryGuid))
+	{
+		return SourceItem->Source->GetSourceToolTip();
+	}
+	return FText(NSLOCTEXT("TempLocTextLiveLink", "InvalidSourceToolTip", "Invalid Source ToolTip"));
+}
+
 bool FLiveLinkClient::IsSourceStillValid(FGuid InEntryGuid) const
 {
 	if (const FLiveLinkCollectionSourceItem* SourceItem = Collection->FindSource(InEntryGuid))
diff --git a/Engine/Plugins/Animation/LiveLink/Source/LiveLink/Private/LiveLinkMessageBusSource.cpp b/Engine/Plugins/Animation/LiveLink/Source/LiveLink/Private/LiveLinkMessageBusSource.cpp
--- a/Engine/Plugins/Animation/LiveLink/Source/LiveLink/Private/LiveLinkMessageBusSource.cpp
+++ b/Engine/Plugins/Animation/LiveLink/Source/LiveLink/Private/LiveLinkMessageBusSource.cpp
@@ -2,8 +2,10 @@
 
 #include "LiveLinkMessageBusSource.h"
 
+#include "ClientNetworkStatisticsModel.h"
 #include "ILiveLinkClient.h"
 #include "ILiveLinkModule.h"
+#include "INetworkMessagingExtension.h"
 #include "LiveLinkCompression.h"
 #include "LiveLinkHeartbeatEmitter.h"
 #include "LiveLinkLog.h"
@@ -101,6 +103,19 @@ void FLiveLinkMessageBusSource::Update()
 		const double HeartbeatTimeout = GetDefault<ULiveLinkSettings>()->GetMessageBusHeartbeatTimeout();
 		const double CurrentTime = FApp::GetCurrentTime();
 
+		float BytesLastSecond = 0.f;
+		if (TOptional<FMessageTransportStatistics> Statistics = UE::LiveLinkHub::FClientNetworkStatisticsModel::GetLatestNetworkStatistics(ConnectionAddress))
+		{
+			BytesLastSecond = Statistics->TotalBytesReceived - AccumulatedBytes;
+			AccumulatedBytes = Statistics->TotalBytesReceived;
+		}
+
+		if (CurrentTime - LastThroughputUpdate > 1.0)
+		{
+			LastThroughputUpdate = CurrentTime;
+			CachedThroughput = BytesLastSecond / 1'000;
+		}
+
 		bIsValid = CurrentTime - ConnectionLastActive < HeartbeatTimeout;
 		if (!bIsValid)
 		{
@@ -272,6 +287,12 @@ FText FLiveLinkMessageBusSource::GetSourceStatus() const
 	return TimeoutSourceStatus();
 }
 
+FText FLiveLinkMessageBusSource::GetSourceToolTip() const
+{
+	const FString TemplateString = FString::Printf(TEXT("Throughput: %.1f KB/s"), CachedThroughput);
+	return FText::FromString(TemplateString);
+}
+
 TSubclassOf<ULiveLinkSourceSettings> FLiveLinkMessageBusSource::GetSettingsClass() const
 {
 	return ULiveLinkMessageBusSourceSettings::StaticClass();
diff --git a/Engine/Plugins/Animation/LiveLink/Source/LiveLink/Public/LiveLinkClient.h b/Engine/Plugins/Animation/LiveLink/Source/LiveLink/Public/LiveLinkClient.h
--- a/Engine/Plugins/Animation/LiveLink/Source/LiveLink/Public/LiveLinkClient.h
+++ b/Engine/Plugins/Animation/LiveLink/Source/LiveLink/Public/LiveLinkClient.h
@@ -53,6 +53,7 @@ public:
 	UE_API virtual FLiveLinkSourcePreset GetSourcePreset(FGuid SourceGuid, UObject* DuplicatedObjectOuter) const override;
 	UE_API virtual FText GetSourceType(FGuid EntryGuid) const override;
 	UE_API virtual FText GetSourceStatus(FGuid EntryGuid) const override;
+	UE_API virtual FText GetSourceToolTip(FGuid EntryGuid) const override;
 	UE_API virtual FText GetSourceMachineName(FGuid EntryGuid) const override;
 	UE_API virtual bool IsSourceStillValid(FGuid EntryGuid) const override;
 
diff --git a/Engine/Plugins/Animation/LiveLink/Source/LiveLink/Public/LiveLinkMessageBusSource.h b/Engine/Plugins/Animation/LiveLink/Source/LiveLink/Public/LiveLinkMessageBusSource.h
--- a/Engine/Plugins/Animation/LiveLink/Source/LiveLink/Public/LiveLinkMessageBusSource.h
+++ b/Engine/Plugins/Animation/LiveLink/Source/LiveLink/Public/LiveLinkMessageBusSource.h
@@ -44,7 +44,7 @@ public:
 	virtual FText GetSourceType() const override { return SourceType; }
 	virtual FText GetSourceMachineName() const override { return SourceMachineName; }
 	UE_API virtual FText GetSourceStatus() const override;
-
+	UE_API virtual FText GetSourceToolTip() const override;
 	UE_API virtual TSubclassOf<ULiveLinkSourceSettings> GetSettingsClass() const override;
 	//~ End ILiveLinkSource interface
 
@@ -146,6 +146,15 @@ private:
 	// Offset between sender's machine engine time and receiver's machine engine time
 	double MachineTimeOffset;
 
+	//  Cached throughput (in kb/s) of data going through this messagebus source.
+	float CachedThroughput = 0.f;
+
+	// Timestamp of the last update to the cached throughput.
+	float LastThroughputUpdate = 0.f;
+
+	//  Number of bytes received since LastThroughputUpdate.
+	float AccumulatedBytes = 0.f;
+
 	// Actively discovering the source
 	bool bDiscovering;
 
diff --git a/Engine/Plugins/Animation/LiveLink/Source/LiveLinkEditor/Private/LiveLinkClientPanelViews.cpp b/Engine/Plugins/Animation/LiveLink/Source/LiveLinkEditor/Private/LiveLinkClientPanelViews.cpp
--- a/Engine/Plugins/Animation/LiveLink/Source/LiveLinkEditor/Private/LiveLinkClientPanelViews.cpp
+++ b/Engine/Plugins/Animation/LiveLink/Source/LiveLinkEditor/Private/LiveLinkClientPanelViews.cpp
@@ -103,6 +103,10 @@ FText FLiveLinkSourceUIEntry::GetDisplayName() const
 {
 	return GetSourceType();
 }
+FText FLiveLinkSourceUIEntry::GetToolTip() const
+{
+	return Client->GetSourceToolTip(EntryGuid);
+}
 
 FLiveLinkSubjectUIEntry::FLiveLinkSubjectUIEntry(const FLiveLinkSubjectKey& InSubjectKey, FLiveLinkClient* InClient, bool bInIsSource)
 	: SubjectKey(InSubjectKey)
@@ -427,6 +431,11 @@ public:
 		return SNullWidget::NullWidget;
 	}
 
+	FText GetToolTipText() const
+	{
+		return EntryPtr->GetToolTip();
+	}
+
 private:
 	FText GetMachineName() const
 	{
@@ -469,7 +478,8 @@ TSharedRef<ITableRow> FLiveLinkSourcesView::MakeSourceListViewWidget(FLiveLinkSo
 {
 	return SNew(SLiveLinkClientPanelSourcesRow, OwnerTable)
 		.Entry(Entry)
-		.ReadOnly(bReadOnly);
+		.ReadOnly(bReadOnly)
+		.ToolTipText_Lambda([Entry]() { return Entry->GetToolTip(); });
 }
 
 void FLiveLinkSourcesView::OnSourceListSelectionChanged(FLiveLinkSourceUIEntryPtr Entry, ESelectInfo::Type SelectionType) const
diff --git a/Engine/Plugins/Animation/LiveLink/Source/LiveLinkEditor/Public/LiveLinkClientPanelViews.h b/Engine/Plugins/Animation/LiveLink/Source/LiveLinkEditor/Public/LiveLinkClientPanelViews.h
--- a/Engine/Plugins/Animation/LiveLink/Source/LiveLinkEditor/Public/LiveLinkClientPanelViews.h
+++ b/Engine/Plugins/Animation/LiveLink/Source/LiveLinkEditor/Public/LiveLinkClientPanelViews.h
@@ -105,6 +105,9 @@ public:
 	void RemoveFromClient() const;
 	// Get the display name of the source
 	FText GetDisplayName() const;
+	// Get the tooltip of the source
+	FText GetToolTip() const;
+
 
 private:
 	// GUID of the source
diff --git a/Engine/Source/Runtime/LiveLinkInterface/Public/ILiveLinkClient.h b/Engine/Source/Runtime/LiveLinkInterface/Public/ILiveLinkClient.h
--- a/Engine/Source/Runtime/LiveLinkInterface/Public/ILiveLinkClient.h
+++ b/Engine/Source/Runtime/LiveLinkInterface/Public/ILiveLinkClient.h
@@ -107,6 +107,9 @@ public:
 	/** Get the status of a source */
 	virtual FText GetSourceStatus(FGuid EntryGuid) const = 0;
 
+	/** Get the tooltip of a source */
+	virtual FText GetSourceToolTip(FGuid EntryGuid) const = 0;
+
 	/** Get the machine name of the source. */
 	virtual FText GetSourceMachineName(FGuid EntryGuid) const = 0;
 
diff --git a/Engine/Source/Runtime/LiveLinkInterface/Public/ILiveLinkSource.h b/Engine/Source/Runtime/LiveLinkInterface/Public/ILiveLinkSource.h
--- a/Engine/Source/Runtime/LiveLinkInterface/Public/ILiveLinkSource.h
+++ b/Engine/Source/Runtime/LiveLinkInterface/Public/ILiveLinkSource.h
@@ -55,6 +55,11 @@ public:
 	 * Should be: "Active", "Not responding", "Connecting" or any short message that makes sense for the source. */
 	virtual FText GetSourceStatus() const = 0;
 
+	/**
+	 * For UI, what should be shown in the source list when hovering this source.
+	 */
+	virtual FText GetSourceToolTip() const { return FText::GetEmpty(); }
+
 	/**
 	 * Setting class to display and used by the Source.
 	 * An instance of that class will be constructed when the source is added to the Client. Then InitializeSettings will be called.
