LiveLinkHub - Fix deadlock in LiveLinkSourceCollection

SourcesLock was acquired on the GT before broadcasting the OnLiveLinkSourceAdded, which in turn eventually acquired the SubjectsLock. This leads to a deadlock with the LiveLink thread that's trying to acquire both locks.

#rnx

diff --git a/Engine/Plugins/Animation/LiveLink/Source/LiveLink/Private/LiveLinkSourceCollection.cpp b/Engine/Plugins/Animation/LiveLink/Source/LiveLink/Private/LiveLinkSourceCollection.cpp
--- a/Engine/Plugins/Animation/LiveLink/Source/LiveLink/Private/LiveLinkSourceCollection.cpp
+++ b/Engine/Plugins/Animation/LiveLink/Source/LiveLink/Private/LiveLinkSourceCollection.cpp
@@ -86,13 +86,17 @@ FLiveLinkSourceCollection::FLiveLinkSourceCollection()
 
 void FLiveLinkSourceCollection::AddSource(FLiveLinkCollectionSourceItem InSource)
 {
-	UE::TUniqueLock Lock(SourcesLock);
-	FLiveLinkCollectionSourceItem& SourceItem = Sources.Add_GetRef(MoveTemp(InSource));
+	FGuid SourceGuid;
+	{
+		UE::TUniqueLock Lock(SourcesLock);
+		FLiveLinkCollectionSourceItem& SourceItem = Sources.Add_GetRef(MoveTemp(InSource));
+		LiveLinkSourceCollection::SendAnalyticsSourceAdded(SourceItem.Source.Get());
+		SourceGuid = SourceItem.Guid;
+	}
 
-	OnLiveLinkSourceAdded().Broadcast(SourceItem.Guid);
+	OnLiveLinkSourceAdded().Broadcast(SourceGuid);
 	OnLiveLinkSourcesChanged().Broadcast();
 
-	LiveLinkSourceCollection::SendAnalyticsSourceAdded(SourceItem.Source.Get());
 }
 
 
