LiveLinkHub - Fix UE not receiving subjects when it connects to LiveLinkHub

diff --git a/Engine/Plugins/Animation/LiveLink/Source/LiveLink/Private/LiveLinkMessageBusSource.cpp b/Engine/Plugins/Animation/LiveLink/Source/LiveLink/Private/LiveLinkMessageBusSource.cpp
--- a/Engine/Plugins/Animation/LiveLink/Source/LiveLink/Private/LiveLinkMessageBusSource.cpp
+++ b/Engine/Plugins/Animation/LiveLink/Source/LiveLink/Private/LiveLinkMessageBusSource.cpp
@@ -175,7 +175,7 @@ void FLiveLinkMessageBusSource::InternalHandleMessage(const TSharedRef<IMessageC
 		return;
 	}
 
-	if (!bIsValid)
+	if (bIsShuttingDown)
 	{
 		return;
 	}
@@ -448,6 +448,7 @@ void FLiveLinkMessageBusSource::SendConnectMessage()
 	SendMessage(ConnectMessage, Annotations);
 	StartHeartbeatEmitter();
 	bIsValid = true;
+	bIsShuttingDown = false;
 }
 
 bool FLiveLinkMessageBusSource::RequestSourceShutdown()
@@ -473,6 +474,7 @@ bool FLiveLinkMessageBusSource::RequestSourceShutdown()
 	RoleInstances.Empty();
 
 	bInitialized = false;
+	bIsShuttingDown = true;
 
 	return true;
 }
diff --git a/Engine/Plugins/Animation/LiveLink/Source/LiveLink/Public/LiveLinkMessageBusSource.h b/Engine/Plugins/Animation/LiveLink/Source/LiveLink/Public/LiveLinkMessageBusSource.h
--- a/Engine/Plugins/Animation/LiveLink/Source/LiveLink/Public/LiveLinkMessageBusSource.h
+++ b/Engine/Plugins/Animation/LiveLink/Source/LiveLink/Public/LiveLinkMessageBusSource.h
@@ -43,6 +43,7 @@ public:
 
 	virtual FText GetSourceType() const override { return SourceType; }
 	virtual FText GetSourceMachineName() const override { return SourceMachineName; }
+
 	UE_API virtual FText GetSourceStatus() const override;
 	UE_API virtual FText GetSourceToolTip() const override;
 	UE_API virtual TSubclassOf<ULiveLinkSourceSettings> GetSettingsClass() const override;
@@ -102,7 +103,9 @@ protected:
 	FMessageAddress ConnectionAddress;
 
 	// Current Validity of Source
-	FThreadSafeBool bIsValid;
+	std::atomic<bool> bIsValid;
+	// Whether we're in the process of shutting down this source.
+	std::atomic<bool> bIsShuttingDown = false;
 
 	ILiveLinkClient* Client;
 
