LiveLinkHub - Show the interpolation processor if transmitting evaluated data is enabled.

diff --git a/Engine/Plugins/Animation/LiveLink/Source/LiveLinkEditor/Private/LiveLinkHubSubjectSettingsDetailsCustomization.cpp b/Engine/Plugins/Animation/LiveLink/Source/LiveLinkEditor/Private/LiveLinkHubSubjectSettingsDetailsCustomization.cpp
--- a/Engine/Plugins/Animation/LiveLink/Source/LiveLinkEditor/Private/LiveLinkHubSubjectSettingsDetailsCustomization.cpp
+++ b/Engine/Plugins/Animation/LiveLink/Source/LiveLinkEditor/Private/LiveLinkHubSubjectSettingsDetailsCustomization.cpp
@@ -6,6 +6,7 @@
 #include "DetailWidgetRow.h"
 #include "IDetailPropertyRow.h"
 #include "LiveLinkHubSubjectSettings.h"
+#include "LiveLinkSettings.h"
 #include "Misc/ConfigCacheIni.h"
 
 
@@ -70,7 +71,11 @@ void FLiveLinkHubSubjectSettingsDetailsCustomization::CustomizeDetails(IDetailLa
 				];
 		}
 
-		InDetailBuilder.HideProperty(GET_MEMBER_NAME_CHECKED(ULiveLinkHubSubjectSettings, InterpolationProcessor), ULiveLinkSubjectSettings::StaticClass());
+		if (!GetDefault<ULiveLinkSettings>()->bTransmitEvaluatedData)
+		{
+			InDetailBuilder.HideProperty(GET_MEMBER_NAME_CHECKED(ULiveLinkHubSubjectSettings, InterpolationProcessor), ULiveLinkSubjectSettings::StaticClass());
+		}
+
 		InDetailBuilder.HideProperty(GET_MEMBER_NAME_CHECKED(ULiveLinkHubSubjectSettings, FrameRate), ULiveLinkSubjectSettings::StaticClass());
 		InDetailBuilder.HideProperty(GET_MEMBER_NAME_CHECKED(ULiveLinkHubSubjectSettings, bRebroadcastSubject), ULiveLinkSubjectSettings::StaticClass());
 		InDetailBuilder.HideProperty(GET_MEMBER_NAME_CHECKED(ULiveLinkHubSubjectSettings, Translators), ULiveLinkSubjectSettings::StaticClass());
