LiveLinkHub - Fix virtual subjects having a long list of subjects

Subjects will now be picked in a combo box to prevent the list from taking too much space in the details panel

diff --git a/Engine/Plugins/Animation/LiveLink/Source/LiveLinkEditor/Private/LiveLinkVirtualSubjectDetailCustomization.cpp b/Engine/Plugins/Animation/LiveLink/Source/LiveLinkEditor/Private/LiveLinkVirtualSubjectDetailCustomization.cpp
--- a/Engine/Plugins/Animation/LiveLink/Source/LiveLinkEditor/Private/LiveLinkVirtualSubjectDetailCustomization.cpp
+++ b/Engine/Plugins/Animation/LiveLink/Source/LiveLinkEditor/Private/LiveLinkVirtualSubjectDetailCustomization.cpp
@@ -12,6 +12,7 @@
 #include "Roles/LiveLinkBasicRole.h"
 #include "ScopedTransaction.h"
 #include "Widgets/Input/SCheckBox.h"
+#include "Widgets/Input/SComboButton.h"
 
 #define LOCTEXT_NAMESPACE "LiveLinkVirtualSubjectDetailsCustomization"
 
@@ -73,18 +74,38 @@ void FLiveLinkVirtualSubjectDetailCustomization::CustomizeDetails(IDetailLayoutB
 		}
 	}
 
+	UpdateSelectedSubjects();
+	
+	TSharedRef<SComboButton> ComboButton = SNew(SComboButton)
+		.ButtonContent()
+		[
+			SNew(STextBlock)
+				.Text(LOCTEXT("SubjectsPickerButtonLabel", "Subjects..."))
+		]
+		.OnGetMenuContent(this, &FLiveLinkVirtualSubjectDetailCustomization::OnGetVirtualSubjectsMenu);
+
 	IDetailCategoryBuilder& SubjectPropertyGroup = DetailBuilder.EditCategory(*SubjectsPropertyHandle->GetMetaData(TEXT("Category")));
+
+	SubjectPropertyGroup.AddCustomRow(LOCTEXT("SelectedSubjectsLabel", "Selected Subjects"))
+		.NameContent()
+		[
+			SNew(STextBlock)
+				.Text(LOCTEXT("SelectedSubjectsLabel", "Selected Subjects"))
+				.Font(IDetailLayoutBuilder::GetDetailFont())
+		]
+		.ValueContent()
+		[
+			SAssignNew(SelectedSubjectsListView, SListView<FSubjectEntryPtr>)
+				.ListItemsSource(&SelectedSubjectsListItems)
+				.SelectionMode(ESelectionMode::None)
+				.OnGenerateRow(this, &FLiveLinkVirtualSubjectDetailCustomization::OnGenerateWidgetForSelectedSubjectItem)
+		];
+
 	SubjectPropertyGroup.AddCustomRow(LOCTEXT("SubjectsTitleLabel", "Subjects"))
-	.NameContent()
-	[
-		SubjectsPropertyHandle->CreatePropertyNameWidget()
-	]
-	.ValueContent()
-	[
-		SAssignNew(SubjectsListView, SListView<FSubjectEntryPtr>)
-		.ListItemsSource(&SubjectsListItems)
-		.OnGenerateRow(this, &FLiveLinkVirtualSubjectDetailCustomization::OnGenerateWidgetForSubjectItem)
-	];
+		.ValueContent()
+		[
+			ComboButton
+		];
 }
 
 int32 GetArrayPropertyIndex(TSharedPtr<IPropertyHandleArray> ArrayProperty, FName ItemToSearchFor, uint32 NumItems)
@@ -102,6 +123,28 @@ int32 GetArrayPropertyIndex(TSharedPtr<IPropertyHandleArray> ArrayProperty, FNam
 	return INDEX_NONE;
 }
 
+bool FLiveLinkVirtualSubjectDetailCustomization::IsEntrySelected(FSubjectEntryPtr EntryPtr)
+{
+	const TSharedPtr<IPropertyHandleArray> SubjectsArrayPropertyHandle = SubjectsPropertyHandle->AsArray();
+	uint32 NumItems;
+	SubjectsArrayPropertyHandle->GetNumElements(NumItems);
+
+	int32 SubjectIndex = GetArrayPropertyIndex(SubjectsArrayPropertyHandle, *EntryPtr, NumItems);
+
+	return SubjectIndex != INDEX_NONE;
+}
+
+TSharedRef<ITableRow> FLiveLinkVirtualSubjectDetailCustomization::OnGenerateWidgetForSelectedSubjectItem(FSubjectEntryPtr InItem, const TSharedRef<STableViewBase>& OwnerTable)
+{
+	return SNew(STableRow<FSubjectEntryPtr>, OwnerTable)
+		[
+			SNew(STextBlock)
+				.Text(FText::FromName(*InItem))
+				.ColorAndOpacity(this, &FLiveLinkVirtualSubjectDetailCustomization::HandleSubjectItemColor, InItem)
+				.ToolTipText(this, &FLiveLinkVirtualSubjectDetailCustomization::HandleSubjectItemToolTip, InItem)
+		];
+}
+
 TSharedRef<ITableRow> FLiveLinkVirtualSubjectDetailCustomization::OnGenerateWidgetForSubjectItem(FSubjectEntryPtr InItem, const TSharedRef<STableViewBase>& OwnerTable)
 {
 	FLiveLinkVirtualSubjectDetailCustomization* CaptureThis = this;
@@ -115,13 +158,7 @@ TSharedRef<ITableRow> FLiveLinkVirtualSubjectDetailCustomization::OnGenerateWidg
 				SNew(SCheckBox)
 				.IsChecked_Lambda([CaptureThis, InItem]()
 				{
-					const TSharedPtr<IPropertyHandleArray> SubjectsArrayPropertyHandle = CaptureThis->SubjectsPropertyHandle->AsArray();
-					uint32 NumItems;
-					SubjectsArrayPropertyHandle->GetNumElements(NumItems);
-
-					int32 RemoveIndex = GetArrayPropertyIndex(SubjectsArrayPropertyHandle, *InItem, NumItems);
-
-					return RemoveIndex != INDEX_NONE ? ECheckBoxState::Checked : ECheckBoxState::Unchecked;
+					return CaptureThis->IsEntrySelected(InItem) ? ECheckBoxState::Checked : ECheckBoxState::Unchecked;
 				})
 				.OnCheckStateChanged_Lambda([CaptureThis, InItem](const ECheckBoxState NewState)
 				{
@@ -152,6 +189,8 @@ TSharedRef<ITableRow> FLiveLinkVirtualSubjectDetailCustomization::OnGenerateWidg
 							SubjectsArrayPropertyHandle->DeleteItem(RemoveIndex);
 						}
 					}
+
+					CaptureThis->UpdateSelectedSubjects();
 				})
 			]
 			+ SHorizontalBox::Slot()
@@ -178,7 +217,7 @@ FSlateColor FLiveLinkVirtualSubjectDetailCustomization::HandleSubjectItemColor(F
 		if (!SubjectKeys.ContainsByPredicate([ThisItem](const FLiveLinkSubjectKey& Other) { return Other.SubjectName.Name == ThisItem; }))
 		{
 			Result = FLinearColor::Red;
-		}
+		} 
 	}
 
 	return Result;
@@ -201,4 +240,28 @@ FText FLiveLinkVirtualSubjectDetailCustomization::HandleSubjectItemToolTip(FSubj
 	return Result;
 }
 
+TSharedRef<SWidget> FLiveLinkVirtualSubjectDetailCustomization::OnGetVirtualSubjectsMenu()
+{
+	return SNew(SListView<FSubjectEntryPtr>)
+		.ListItemsSource(&SubjectsListItems)
+		.OnGenerateRow(this, &FLiveLinkVirtualSubjectDetailCustomization::OnGenerateWidgetForSubjectItem);
+}
+
+void FLiveLinkVirtualSubjectDetailCustomization::UpdateSelectedSubjects()
+{
+	SelectedSubjectsListItems.Reset();
+	for (FSubjectEntryPtr EntryPtr : SubjectsListItems)
+	{
+		if (IsEntrySelected(EntryPtr))
+		{
+			SelectedSubjectsListItems.Add(EntryPtr);
+		}
+	}
+
+	if (SelectedSubjectsListView)
+	{
+		SelectedSubjectsListView->RequestListRefresh();
+	}
+}
+
 #undef LOCTEXT_NAMESPACE
diff --git a/Engine/Plugins/Animation/LiveLink/Source/LiveLinkEditor/Private/LiveLinkVirtualSubjectDetailCustomization.h b/Engine/Plugins/Animation/LiveLink/Source/LiveLinkEditor/Private/LiveLinkVirtualSubjectDetailCustomization.h
--- a/Engine/Plugins/Animation/LiveLink/Source/LiveLinkEditor/Private/LiveLinkVirtualSubjectDetailCustomization.h
+++ b/Engine/Plugins/Animation/LiveLink/Source/LiveLinkEditor/Private/LiveLinkVirtualSubjectDetailCustomization.h
@@ -36,7 +36,9 @@ public:
 	// End IDetailCustomization interface
 
 private:
-	
+	// Generates a row for a selected subject.
+	TSharedRef<ITableRow> OnGenerateWidgetForSelectedSubjectItem(FSubjectEntryPtr InItem, const TSharedRef<STableViewBase>& OwnerTable);
+
 	// Creates subject tree entry widget
 	TSharedRef<ITableRow> OnGenerateWidgetForSubjectItem(FSubjectEntryPtr InItem, const TSharedRef<STableViewBase>& OwnerTable);
 
@@ -45,6 +47,16 @@ private:
 
 	FText HandleSubjectItemToolTip(FSubjectEntryPtr InItem) const;
 
+	// Generate the combo box menu to display subjects to pick.
+	TSharedRef<SWidget> OnGetVirtualSubjectsMenu();
+
+	// Returns whether an entry is currently selected.
+	bool IsEntrySelected(FSubjectEntryPtr EntryPtr);
+
+	// Update the list of selected subjects displayed above the subject picker.
+	void UpdateSelectedSubjects();
+
+private:
 	// The tile set being edited
 	TWeakObjectPtr<ULiveLinkVirtualSubject> SubjectPtr;
 
@@ -58,5 +70,6 @@ private:
 
 	// Cached data for the subject tree UI
 	TArray<FSubjectEntryPtr> SubjectsListItems;
-	TSharedPtr< SListView< FSubjectEntryPtr > > SubjectsListView;
+	TArray<FSubjectEntryPtr> SelectedSubjectsListItems;
+	TSharedPtr< SListView< FSubjectEntryPtr > > SelectedSubjectsListView;
 };
