LiveLinkHub - Minor UI fixes to the VSubject list

diff --git a/Engine/Plugins/Animation/LiveLink/Source/LiveLinkEditor/Private/LiveLinkVirtualSubjectDetailCustomization.cpp b/Engine/Plugins/Animation/LiveLink/Source/LiveLinkEditor/Private/LiveLinkVirtualSubjectDetailCustomization.cpp
--- a/Engine/Plugins/Animation/LiveLink/Source/LiveLinkEditor/Private/LiveLinkVirtualSubjectDetailCustomization.cpp
+++ b/Engine/Plugins/Animation/LiveLink/Source/LiveLinkEditor/Private/LiveLinkVirtualSubjectDetailCustomization.cpp
@@ -80,6 +80,7 @@ void FLiveLinkVirtualSubjectDetailCustomization::CustomizeDetails(IDetailLayoutB
 		.ButtonContent()
 		[
 			SNew(STextBlock)
+				.Font(IDetailLayoutBuilder::GetDetailFont())
 				.Text(LOCTEXT("SubjectsPickerButtonLabel", "Subjects..."))
 		]
 		.OnGetMenuContent(this, &FLiveLinkVirtualSubjectDetailCustomization::OnGetVirtualSubjectsMenu);
@@ -90,8 +91,8 @@ void FLiveLinkVirtualSubjectDetailCustomization::CustomizeDetails(IDetailLayoutB
 		.NameContent()
 		[
 			SNew(STextBlock)
-				.Text(LOCTEXT("SelectedSubjectsLabel", "Selected Subjects"))
 				.Font(IDetailLayoutBuilder::GetDetailFont())
+				.Text(LOCTEXT("SelectedSubjectsLabel", "Selected Subjects"))
 		]
 		.ValueContent()
 		[
@@ -136,12 +137,21 @@ bool FLiveLinkVirtualSubjectDetailCustomization::IsEntrySelected(FSubjectEntryPt
 
 TSharedRef<ITableRow> FLiveLinkVirtualSubjectDetailCustomization::OnGenerateWidgetForSelectedSubjectItem(FSubjectEntryPtr InItem, const TSharedRef<STableViewBase>& OwnerTable)
 {
+	FSlateFontInfo Font = IDetailLayoutBuilder::GetDetailFont();
+	Font.Size -= 0.5;
+
 	return SNew(STableRow<FSubjectEntryPtr>, OwnerTable)
 		[
-			SNew(STextBlock)
-				.Text(FText::FromName(*InItem))
-				.ColorAndOpacity(this, &FLiveLinkVirtualSubjectDetailCustomization::HandleSubjectItemColor, InItem)
-				.ToolTipText(this, &FLiveLinkVirtualSubjectDetailCustomization::HandleSubjectItemToolTip, InItem)
+			SNew(SHorizontalBox)
+			+ SHorizontalBox::Slot()
+			.Padding(4.0, 0.0)
+			[
+				SNew(STextBlock)
+					.Text(FText::FromName(*InItem))
+					.Font(Font)
+					.ColorAndOpacity(this, &FLiveLinkVirtualSubjectDetailCustomization::HandleSubjectItemColor, InItem)
+					.ToolTipText(this, &FLiveLinkVirtualSubjectDetailCustomization::HandleSubjectItemToolTip, InItem)
+			]
 		];
 }
 
@@ -206,13 +216,12 @@ TSharedRef<ITableRow> FLiveLinkVirtualSubjectDetailCustomization::OnGenerateWidg
 
 FSlateColor FLiveLinkVirtualSubjectDetailCustomization::HandleSubjectItemColor(FSubjectEntryPtr InItem) const
 {
-	FSlateColor Result = FLinearColor::White;
+	FSlateColor Result = FSlateColor::UseForeground();
 
 	if (ULiveLinkVirtualSubject* Subject = SubjectPtr.Get())
 	{
 		const FName ThisItem = *InItem.Get();
 
-
 		TArray<FLiveLinkSubjectKey> SubjectKeys = Client->GetSubjects(/*bIncludeDisabledSubject*/ false, /*bIncludeVirtualSubject*/ false);
 		if (!SubjectKeys.ContainsByPredicate([ThisItem](const FLiveLinkSubjectKey& Other) { return Other.SubjectName.Name == ThisItem; }))
 		{
