LiveLink - Fix crash when copy pasting LiveLink subsection into a different level sequence

LiveLink static data is now exported to text alongside the frame data when copying a livelink subsection to clipboard.

diff --git a/Engine/Plugins/Animation/LiveLink/Source/LiveLinkMovieScene/Private/MovieScene/MovieSceneLiveLinkSection.cpp b/Engine/Plugins/Animation/LiveLink/Source/LiveLinkMovieScene/Private/MovieScene/MovieSceneLiveLinkSection.cpp
--- a/Engine/Plugins/Animation/LiveLink/Source/LiveLinkMovieScene/Private/MovieScene/MovieSceneLiveLinkSection.cpp
+++ b/Engine/Plugins/Animation/LiveLink/Source/LiveLinkMovieScene/Private/MovieScene/MovieSceneLiveLinkSection.cpp
@@ -251,6 +251,46 @@ void UMovieSceneLiveLinkSection::PostEditChangeProperty(struct FPropertyChangedE
 	//Make sure the channel proxy is always in sync with subsections data. They are recreated whenever we undo / redo.
 	UpdateChannelProxy();
 }
+
+void UMovieSceneLiveLinkSection::ExportCustomProperties(FOutputDevice& Out, uint32 Indent)
+{
+	Super::ExportCustomProperties(Out, Indent);
+
+	if (StaticData)
+	{
+		FString Text;
+		StaticData->GetStruct()->ExportText(Text, StaticData->GetBaseData(), nullptr, nullptr, PPF_None, nullptr);
+		Out.Logf(TEXT("%sCustomProperties StaticDataStruct Path=%s StaticData=%s\r\n"), FCString::Spc(Indent), *StaticData->GetStruct()->GetPathName(), *Text);
+	}
+}
+
+void UMovieSceneLiveLinkSection::ImportCustomProperties(const TCHAR* SourceText, FFeedbackContext* Warn)
+{
+	if (FParse::Command(&SourceText, TEXT("StaticDataStruct")))
+	{
+		FString StructPath;
+		if (!FParse::Value(SourceText, TEXT("Path="), StructPath))
+		{
+			return;
+		}
+
+		FString StructValue;
+		if (FParse::Value(SourceText, TEXT("StaticData="), StructValue, false))
+		{
+			if (UScriptStruct* Struct = FindObject<UScriptStruct>(nullptr, *StructPath))
+			{
+				TSharedPtr<FLiveLinkStaticDataStruct> StaticStruct = MakeShared<FLiveLinkStaticDataStruct>(Struct);
+
+				Struct->ImportText(*StructValue, StaticStruct->GetBaseData(), this, PPF_None, nullptr, Struct->GetName());
+				StaticData = MoveTemp(StaticStruct);
+			}
+		}
+	}
+	else
+	{
+		Super::ImportCustomProperties(SourceText, Warn);
+	}
+}
 #endif
 
 TArray<TSubclassOf<UMovieSceneLiveLinkSection>> UMovieSceneLiveLinkSection::GetMovieSectionForRole(const TSubclassOf<ULiveLinkRole>& InRoleToSupport)
diff --git a/Engine/Plugins/Animation/LiveLink/Source/LiveLinkMovieScene/Public/MovieScene/MovieSceneLiveLinkSection.h b/Engine/Plugins/Animation/LiveLink/Source/LiveLinkMovieScene/Public/MovieScene/MovieSceneLiveLinkSection.h
--- a/Engine/Plugins/Animation/LiveLink/Source/LiveLinkMovieScene/Public/MovieScene/MovieSceneLiveLinkSection.h
+++ b/Engine/Plugins/Animation/LiveLink/Source/LiveLinkMovieScene/Public/MovieScene/MovieSceneLiveLinkSection.h
@@ -66,6 +66,8 @@ public:
 #if WITH_EDITOR
 	UE_API virtual bool Modify(bool bAlwaysMarkDirty = true) override;
 	UE_API virtual void PostEditChangeProperty(struct FPropertyChangedEvent& PropertyChangedEvent) override;
+	UE_API virtual void ExportCustomProperties(FOutputDevice& Out, uint32 Indent) override;
+	UE_API virtual void ImportCustomProperties(const TCHAR* SourceText, FFeedbackContext* Warn) override;
 #endif
 	//~ End UObject interface
 
