LiveLink - Move LiveLink visualizer utils from PCap  to LiveLink plugin

#rnx

#virtualized

diff --git a/Engine/Plugins/Animation/LiveLink/LiveLink.uplugin b/Engine/Plugins/Animation/LiveLink/LiveLink.uplugin
--- a/Engine/Plugins/Animation/LiveLink/LiveLink.uplugin
+++ b/Engine/Plugins/Animation/LiveLink/LiveLink.uplugin
@@ -11,7 +11,7 @@
 	"MarketplaceURL": "",
 	"SupportURL": "",
 	"EnabledByDefault": false,
-	"CanContainContent": false,
+	"CanContainContent": true,
 	"IsBetaVersion": false,
 	"Installed": false,
 	"SupportedPrograms": [ "LiveLinkHub" ],
diff --git a/Engine/Plugins/Animation/LiveLink/Source/LiveLink/LiveLink.Build.cs b/Engine/Plugins/Animation/LiveLink/Source/LiveLink/LiveLink.Build.cs
--- a/Engine/Plugins/Animation/LiveLink/Source/LiveLink/LiveLink.Build.cs
+++ b/Engine/Plugins/Animation/LiveLink/Source/LiveLink/LiveLink.Build.cs
@@ -19,6 +19,7 @@ namespace UnrealBuildTool.Rules
 			PrivateDependencyModuleNames.AddRange(
 			new string[]
 			{
+				"AnimationCore",
 				"HeadMountedDisplay",
 				"InputCore",
 				"LiveLinkAnimationCore", //Adding dependency to make sure this runtime module is pulled when packaging if LiveLink plugin is enabled. 
diff --git a/Engine/Plugins/Animation/LiveLink/Source/LiveLink/Private/Visualizers/LiveLinkBoneVisualizer.cpp b/Engine/Plugins/Animation/LiveLink/Source/LiveLink/Private/Visualizers/LiveLinkBoneVisualizer.cpp
new file mode 100644
--- /dev/null
+++ b/Engine/Plugins/Animation/LiveLink/Source/LiveLink/Private/Visualizers/LiveLinkBoneVisualizer.cpp
@@ -0,0 +1,149 @@
+// Copyright Epic Games, Inc. All Rights Reserved.
+
+
+#include "LiveLinkBoneVisualizer.h"
+#include "AnimationCoreLibrary.h"
+#include "SkeletalMeshAttributes.h"
+#include "Animation/SkeletalMeshActor.h"
+#include "Components/SkeletalMeshComponent.h"
+#include "Materials/MaterialInstanceDynamic.h"
+
+
+// Sets default values for this component's properties
+ULiveLinkBoneVisualizer::ULiveLinkBoneVisualizer(const FObjectInitializer& ObjectInitializer)
+{
+	PrimaryComponentTick.bCanEverTick = true;
+	bTickInEditor = true; //make this component tick in editor
+	VisualizationType = EBoneVisualType::Joint;
+	Color = FLinearColor(0.5,0.5,0.5,1.0);
+	UPrimitiveComponent::SetCollisionEnabled(ECollisionEnabled::NoCollision);
+	bAlwaysCreatePhysicsState = false;
+	SetCastShadow(false);
+}
+
+void ULiveLinkBoneVisualizer::OnRegister()
+{
+	Super::OnRegister();
+
+	ASkeletalMeshActor* OwnerActor = Cast<ASkeletalMeshActor>(GetOwner());
+	
+	if(OwnerActor)
+	{
+		SkelmeshComponent = OwnerActor->GetSkeletalMeshComponent();
+	}
+
+	ClearInstances();
+
+	if(SkelmeshComponent)
+	{
+		TArray<FTransform> BoneTransforms = GetBoneTransforms(SkelmeshComponent);
+		AddInstances(BoneTransforms, false, true, false);
+
+		if(!DynamicMaterial)
+		{
+			UMaterialInterface* BaseMaterial = GetMaterial(0);
+			if(BaseMaterial) //Check if the BaseMaterial is Valid. It will only be valid if a mesh has already been sent. 
+			{
+				DynamicMaterial = CreateDynamicMaterialInstance(0, BaseMaterial, TEXT("MaterialInstance"));
+				SetMaterial(0, DynamicMaterial);
+				DynamicMaterial->ClearParameterValues();
+				DynamicMaterial->SetVectorParameterValue(FName(TEXT("Color")), Color);
+			}
+		}
+	}
+}
+
+void ULiveLinkBoneVisualizer::OnUnregister()
+{
+	Super::OnUnregister();
+
+	ClearInstances();
+}
+
+void ULiveLinkBoneVisualizer::TickComponent(float DeltaTime, ELevelTick TickType, FActorComponentTickFunction* ThisTickFunction)
+{
+	Super::TickComponent(DeltaTime, TickType, ThisTickFunction);
+	
+	if(SkelmeshComponent)
+	{
+		switch (VisualizationType)
+		{
+		case EBoneVisualType::Bone:
+			{
+				BatchUpdateInstancesTransforms(0, GetBoneTransforms(SkelmeshComponent), true, true, false);
+			}
+			break;
+		case EBoneVisualType::Joint:
+			{
+				BatchUpdateInstancesTransforms(0, GetJointTransforms(SkelmeshComponent), true, true, false);
+			}
+
+			break;
+		default:
+			checkNoEntry();
+		}
+	}
+}
+
+TArray<FTransform> ULiveLinkBoneVisualizer::GetJointTransforms(USkeletalMeshComponent* InSkelMeshComponent) const
+{
+	TArray<FTransform> FoundTransforms;
+
+	int32 i = InSkelMeshComponent->GetNumBones();
+	FTransform BoneTransform;
+
+	for (i = 0; i < InSkelMeshComponent->GetNumBones(); i++ )
+	{
+		BoneTransform = InSkelMeshComponent->GetBoneTransform(InSkelMeshComponent->GetBoneName(i));
+		FoundTransforms.Add(BoneTransform);
+	}
+
+	return FoundTransforms;
+}
+
+TArray<FTransform> ULiveLinkBoneVisualizer::GetBoneTransforms(USkeletalMeshComponent* InSkelMeshComponent) const
+{
+	TArray<FTransform> FoundTransforms;
+
+	int32 i = InSkelMeshComponent->GetNumBones();
+	FTransform BoneTransform;
+	FTransform ParentBoneTransform;
+	double BoneLength;
+	FVector Aim = FVector( 0, 0, 1);
+
+	for (i = 0; i < InSkelMeshComponent->GetNumBones(); i++ )
+	{
+		BoneTransform = InSkelMeshComponent->GetBoneTransform(InSkelMeshComponent->GetBoneName(i));
+		ParentBoneTransform = InSkelMeshComponent->GetBoneTransform(InSkelMeshComponent->GetParentBone(InSkelMeshComponent->GetBoneName(i)));
+		BoneLength = FVector::Distance(BoneTransform.GetLocation(), ParentBoneTransform.GetLocation());
+		BoneTransform.SetScale3D(FVector(1, 1, BoneLength));
+		FQuat DiffRotation = AnimationCore::SolveAim(BoneTransform, ParentBoneTransform.GetLocation(), Aim.GetSafeNormal(), false, FVector(1, 1, 1), float(0));
+		BoneTransform.SetRotation(DiffRotation);
+		FoundTransforms.Add(BoneTransform);
+	}
+	return FoundTransforms;
+}
+
+void ULiveLinkBoneVisualizer::UpdateColor(FLinearColor NewColor)
+{
+	Color = NewColor;
+	if(DynamicMaterial)
+	{
+		DynamicMaterial->SetVectorParameterValue(FName(TEXT("Color")), Color);
+	}
+}
+
+#if WITH_EDITOR
+void ULiveLinkBoneVisualizer::PostEditChangeProperty(FPropertyChangedEvent& PropertyChangedEvent)
+{
+	Super::PostEditChangeProperty(PropertyChangedEvent);
+
+	if(PropertyChangedEvent.Property != nullptr)
+	{
+		if(PropertyChangedEvent.Property->GetFName() == GET_MEMBER_NAME_CHECKED(ULiveLinkBoneVisualizer, Color))
+		{
+			UpdateColor(Color);
+		}
+	}
+}
+#endif
diff --git a/Engine/Plugins/Animation/LiveLink/Source/LiveLink/Private/Visualizers/LiveLinkBoneVisualizer.h b/Engine/Plugins/Animation/LiveLink/Source/LiveLink/Private/Visualizers/LiveLinkBoneVisualizer.h
new file mode 100644
--- /dev/null
+++ b/Engine/Plugins/Animation/LiveLink/Source/LiveLink/Private/Visualizers/LiveLinkBoneVisualizer.h
@@ -0,0 +1,68 @@
+// Copyright Epic Games, Inc. All Rights Reserved.
+
+#pragma once
+
+#include "CoreMinimal.h"
+#include "Components/InstancedStaticMeshComponent.h"
+#include "LiveLinkBoneVisualizer.generated.h"
+
+class USkeletalMeshComponent;
+class UMaterialInstanceDynamic;
+
+UENUM(BlueprintType)
+enum class EBoneVisualType : uint8
+{
+	Joint	UMETA(DisplayName = "Joint"),
+	Bone	UMETA(DisplayName = "Bone"),
+};
+
+/** Instanced Static Mesh Component for drawing bones and joints on Skeletal Mesh. Only usable with SkeletalMesh Actors*/
+UCLASS(ClassGroup=(PerformanceCapture), meta=(BlueprintSpawnableComponent),Transient, DisplayName = "PerformanceCaptureVisualizer")
+class LIVELINK_API ULiveLinkBoneVisualizer : public UInstancedStaticMeshComponent
+{
+	GENERATED_BODY()
+
+public:
+	// Sets default values for this component's properties
+	ULiveLinkBoneVisualizer(const FObjectInitializer& ObjectInitializer);
+	
+	// Called every frame
+	virtual void TickComponent(float DeltaTime, ELevelTick TickType, FActorComponentTickFunction* ThisTickFunction) override;
+
+	/** Get all the transforms from the skelmesh. */
+	TArray<FTransform> GetJointTransforms(USkeletalMeshComponent* InSkelMeshComponent) const;
+
+	/** Get all the bones' transforms for drawing a joint (point to point mesh). */
+	TArray<FTransform> GetBoneTransforms(USkeletalMeshComponent* InSkelMeshComponent) const;
+	
+	/** Whether to draw joints along the length of the bone, or just at the pivot or scaled by the distance to the parent bone. */
+	UPROPERTY(EditAnywhere, BlueprintReadWrite, Category="PerformanceCapture|Visualization")
+	EBoneVisualType VisualizationType;
+
+	/** Color to use on the instanced static meshes. Can only be set during construction. */
+	UPROPERTY(EditAnywhere, BlueprintReadWrite, Setter="UpdateColor", Category="PerformanceCapture|Visualization")
+	FLinearColor Color;
+
+	/** Dynamic Material to use on the instanced static meshes. */
+	UPROPERTY(EditAnywhere, BlueprintReadOnly, Category="PerformanceCapture|Visualization")
+	TObjectPtr<UMaterialInstanceDynamic> DynamicMaterial;
+
+	/**
+	 * Update the color for this bone visaulizer.
+	 * @param NewColor New color to apply in the material.
+	 */
+	UFUNCTION(BlueprintCallable, Category="PerformanceCapture|Visualization")
+	void UpdateColor(FLinearColor NewColor);
+
+#if WITH_EDITOR
+	virtual void PostEditChangeProperty(FPropertyChangedEvent& PropertyChangedEvent) override;
+#endif
+
+protected:
+	virtual void OnRegister() override;
+	virtual void OnUnregister() override;
+
+private:
+	UPROPERTY()
+	TObjectPtr<USkeletalMeshComponent> SkelmeshComponent;
+};
diff --git a/Engine/Plugins/Animation/LiveLink/Source/LiveLink/Private/Visualizers/LiveLinkDataPreview.cpp b/Engine/Plugins/Animation/LiveLink/Source/LiveLink/Private/Visualizers/LiveLinkDataPreview.cpp
new file mode 100644
--- /dev/null
+++ b/Engine/Plugins/Animation/LiveLink/Source/LiveLink/Private/Visualizers/LiveLinkDataPreview.cpp
@@ -0,0 +1,210 @@
+// Copyright Epic Games, Inc. All Rights Reserved.
+
+
+#include "Visualizers/LiveLinkDataPreview.h"
+
+#include "Roles/LiveLinkAnimationRole.h"
+#include "Roles/LiveLinkCameraRole.h"
+#include "Roles/LiveLinkLocatorRole.h"
+#include "Roles/LiveLinkTransformRole.h"
+#include "ILiveLinkClient.h"
+#include "Features/IModularFeatures.h"
+#include "Components/BillboardComponent.h"
+#include "Materials/MaterialInstanceDynamic.h"
+#include "Engine/StaticMesh.h"
+#include "Engine/Texture2D.h"
+#include "Engine/World.h"
+#include "Visualizers/LiveLinkDataPreviewComponent.h"
+
+
+// Sets default values
+ALiveLinkDataPreview::ALiveLinkDataPreview()
+{
+	// Set this actor to call Tick() every frame.  You can turn this off to improve performance if you don't need it.
+	PrimaryActorTick.bCanEverTick = true;
+	bEvaluateLiveLink = true;
+
+	BoneMesh = FSoftObjectPath("/LiveLink/Visualizers/SM_Bone.SM_Bone");
+	JointMesh = FSoftObjectPath("/LiveLink/Visualizers/SM_Joint.SM_Joint");
+	AxisMesh = FSoftObjectPath("/LiveLink/Visualizers/SM_JointAxis.SM_JointAxis");
+	TransformMesh = FSoftObjectPath("/LiveLink/Visualizers/SM_TransformAxis.SM_TransformAxis");
+	LocatorMesh = FSoftObjectPath("/LiveLink/Visualizers/SM_MarkerSphere.SM_MarkerSphere");
+	CameraMesh = FSoftObjectPath("/LiveLink/Visualizers/SM_LiveLinkCamera.SM_LiveLinkCamera");
+	SpriteTexture = FSoftObjectPath("/LiveLink/Starship/Location_256.Location_256");
+
+	USceneComponent* SceneRoot = CreateDefaultSubobject<USceneComponent>("Root");
+	SetRootComponent(SceneRoot);
+	BillboardComponent = CreateDefaultSubobject<UBillboardComponent>("Sprite");
+	BillboardComponent->SetupAttachment(SceneRoot);
+	UTexture2D* SpriteTex = SpriteTexture.LoadSynchronous();
+	if(SpriteTex)
+	{
+		BillboardComponent->SetSprite(SpriteTex);
+		BillboardComponent->bIsScreenSizeScaled = true;
+		BillboardComponent->ScreenSize = 0.0006;
+		BillboardComponent->SetRelativeLocation(FVector(0,0,10));
+	}
+}
+
+// Called every frame
+void ALiveLinkDataPreview::Tick(float DeltaTime)
+{
+	Super::Tick(DeltaTime);
+
+	IModularFeatures& ModularFeatures = IModularFeatures::Get();
+	if (ModularFeatures.IsModularFeatureAvailable(ILiveLinkClient::ModularFeatureName))
+	{
+		ILiveLinkClient& LiveLinkClient = IModularFeatures::Get().GetModularFeature<ILiveLinkClient>(ILiveLinkClient::ModularFeatureName);
+	
+		TArray<FLiveLinkSubjectKey> SubjectKeys = LiveLinkClient.GetSubjects(false, false);
+
+		//If the number of Live Link subjects changes, force reinitialization. 
+		if(SubjectKeys.Num() != CachedSubjects.Num())
+		{
+			InitializeSubjects();
+		}
+		CachedSubjects = SubjectKeys;
+	}
+}
+
+void ALiveLinkDataPreview::OnConstruction(const FTransform& Transform)
+{
+	Super::OnConstruction(Transform);
+
+	InitializeSubjects();
+}
+
+#if WITH_EDITOR
+void ALiveLinkDataPreview::PostEditChangeProperty(FPropertyChangedEvent& PropertyChangedEvent)
+{
+	Super::PostEditChangeProperty(PropertyChangedEvent);
+
+	if(PropertyChangedEvent.Property != nullptr)
+	{
+		if(PropertyChangedEvent.Property->GetFName() == GET_MEMBER_NAME_CHECKED(ALiveLinkDataPreview, bEvaluateLiveLink))
+		{
+			TArray<ULiveLinkDataPreviewComponent*> DebugComponents;
+			GetComponents(DebugComponents);
+
+			for (ULiveLinkDataPreviewComponent* DebugComponent : DebugComponents)
+			{
+				DebugComponent->SetEvaluateLiveLinkData(bEvaluateLiveLink);
+			}
+		}
+	}
+}
+#endif
+
+void ALiveLinkDataPreview::SetEnableLiveLinkData(bool bNewEvaluate)
+{
+	bEvaluateLiveLink = bNewEvaluate;
+
+	//Loop through each subject's preview component and set it's bEvaluateLiveLink member.
+
+	TArray<ULiveLinkDataPreviewComponent*> DebugComponents;
+	GetComponents(DebugComponents);
+
+	for (ULiveLinkDataPreviewComponent* DebugComponent : DebugComponents)
+	{
+		DebugComponent->SetEvaluateLiveLinkData(bEvaluateLiveLink);
+	}
+}
+
+//Initialize each subject's preview component and depending on the role type, add a component with the mesh of the correct type.
+void ALiveLinkDataPreview::InitializeSubjects()
+{
+	if((GetWorld() && !GetWorld()->IsPreviewWorld())) //Do not spawn the preview components unless we are a non-preview world
+	{
+		IModularFeatures& ModularFeatures = IModularFeatures::Get();
+		if (ModularFeatures.IsModularFeatureAvailable(ILiveLinkClient::ModularFeatureName))
+		{
+			ILiveLinkClient& LiveLinkClient = IModularFeatures::Get().GetModularFeature<ILiveLinkClient>(ILiveLinkClient::ModularFeatureName);
+
+			TArray<FLiveLinkSubjectKey> SubjectKeys = LiveLinkClient.GetSubjects(false, false);
+
+			Subjects.Empty();
+
+			for (FLiveLinkSubjectKey Subjectkey : SubjectKeys)
+			{
+				FLiveLinkSubjectName Subject = Subjectkey.SubjectName;
+				Subjects.Emplace(Subject);
+			}
+
+			TArray<ULiveLinkDataPreviewComponent*> DataPreviewComponents;
+			GetComponents(DataPreviewComponents);
+
+			for (ULiveLinkDataPreviewComponent* DataPreviewComponent : DataPreviewComponents)
+			{
+				DataPreviewComponent->DestroyComponent();
+			}
+			
+			for (FLiveLinkSubjectName Subject : Subjects)
+			{
+				if(LiveLinkClient.IsSubjectEnabled(Subject))
+				{
+					ULiveLinkDataPreviewComponent* DataPreviewComponent;
+					DataPreviewComponent = Cast<ULiveLinkDataPreviewComponent>(AddComponentByClass(ULiveLinkDataPreviewComponent::StaticClass(), false, FTransform::Identity, false));
+					
+					DataPreviewComponent->SubjectName = Subject;
+					DataPreviewComponent->bEvaluateLiveLink = bEvaluateLiveLink;
+					DataPreviewComponent->bDrawLabels = bDrawLabels;
+
+					if(LiveLinkClient.DoesSubjectSupportsRole_AnyThread(Subject, ULiveLinkAnimationRole::StaticClass()))
+					{
+						if(UStaticMesh* StaticMesh = BoneMesh.LoadSynchronous())
+						{
+							DataPreviewComponent->SetStaticMesh(StaticMesh);
+							SetMaterialInstance(DataPreviewComponent);
+						}
+						DataPreviewComponent->BoneVisualType = ELiveLinkVisualBoneType::Bone;
+					}
+					else if(LiveLinkClient.DoesSubjectSupportsRole_AnyThread(Subject, ULiveLinkLocatorRole::StaticClass()))
+					{
+						if(UStaticMesh* StaticMesh = LocatorMesh.LoadSynchronous())
+						{
+							DataPreviewComponent->SetStaticMesh(StaticMesh);
+							SetMaterialInstance(DataPreviewComponent);
+						}
+					}
+					else if(LiveLinkClient.DoesSubjectSupportsRole_AnyThread(Subject, ULiveLinkCameraRole::StaticClass()))
+					{
+						
+						if(UStaticMesh* StaticMesh = CameraMesh.LoadSynchronous())
+						{
+							DataPreviewComponent->SetStaticMesh(StaticMesh);
+							SetMaterialInstance(DataPreviewComponent);
+						}
+					}
+					else if(LiveLinkClient.DoesSubjectSupportsRole_AnyThread(Subject, ULiveLinkTransformRole::StaticClass()))
+					{
+						;
+						if(UStaticMesh* StaticMesh = TransformMesh.LoadSynchronous())
+						{
+							DataPreviewComponent->SetStaticMesh(StaticMesh);
+						}
+					}
+				}
+			}
+		}
+	}
+}
+
+//Set the material instance and set a nice random color. 
+void ALiveLinkDataPreview::SetMaterialInstance(ULiveLinkDataPreviewComponent* InDataPreviewComponent)
+{
+	UMaterialInterface* DebugMaterial = InDataPreviewComponent->GetMaterial(0);
+	if(DebugMaterial)
+	{
+		UMaterialInstanceDynamic* DynamicMaterial = InDataPreviewComponent->CreateDynamicMaterialInstance(0, DebugMaterial);
+		InDataPreviewComponent->SetMaterial(0, DynamicMaterial);
+		DynamicMaterial->ClearParameterValues();
+		/*Generate a random color that is weighted toward 1 or 2 primary colors.*/
+		float Red01 = FMath::FRandRange(0.33,0.8);
+		float Green01 = FMath::FRandRange(0.33,0.8);
+		float Blue01 = FMath::FRandRange(0.33,0.8);
+		float Red02 = Red01 - Green01;
+		float Green02 = Green01 - Blue01;
+		float Blue02 = Green01 - Red01;
+		DynamicMaterial->SetVectorParameterValue(FName(TEXT("Color")), FLinearColor(FMath::RandBool() ? Red01 : Red02, FMath::RandBool() ? Green01 : Green02, FMath::RandBool() ? Blue01 : Blue02, 1.0) );
+	}
+}
diff --git a/Engine/Plugins/Animation/LiveLink/Source/LiveLink/Private/Visualizers/LiveLinkDataPreviewComponent.cpp b/Engine/Plugins/Animation/LiveLink/Source/LiveLink/Private/Visualizers/LiveLinkDataPreviewComponent.cpp
new file mode 100644
--- /dev/null
+++ b/Engine/Plugins/Animation/LiveLink/Source/LiveLink/Private/Visualizers/LiveLinkDataPreviewComponent.cpp
@@ -0,0 +1,382 @@
+// Copyright Epic Games, Inc. All Rights Reserved.
+
+
+#include "Visualizers/LiveLinkDataPreviewComponent.h"
+
+#include "AnimationCoreLibrary.h"
+#include "CanvasTypes.h"
+#include "ILiveLinkClient.h"
+#include "LiveLinkBoneVisualizer.h"
+#include "LiveLinkClient.h"
+#include "Features/IModularFeatures.h"
+#include "Roles/LiveLinkAnimationRole.h"
+#include "Roles/LiveLinkLocatorRole.h"
+#include "Roles/LiveLinkTransformRole.h"
+
+ULiveLinkDataPreviewComponent::ULiveLinkDataPreviewComponent(const FObjectInitializer& ObjectInitializer)
+	: bEvaluateLiveLink(true)
+	, bDrawLabels(false)
+	, bIsDirty(true)
+{
+	PrimaryComponentTick.bCanEverTick = true;
+	bTickInEditor = true; //make this component tick in editor
+	PrimaryComponentTick.TickGroup = ETickingGroup::TG_PostUpdateWork;
+}
+
+void ULiveLinkDataPreviewComponent::TickComponent(float DeltaTime, ELevelTick TickType, FActorComponentTickFunction* ThisTickFunction)
+{
+	Super::TickComponent(DeltaTime, TickType, ThisTickFunction);
+	//TODO: add drawing labels for bones and transforms. 
+	
+	ILiveLinkClient& LiveLinkClient = IModularFeatures::Get().GetModularFeature<ILiveLinkClient>(ILiveLinkClient::ModularFeatureName);
+
+	TSubclassOf<ULiveLinkRole> SubjectRole = LiveLinkClient.GetSubjectRole_AnyThread(SubjectName);
+	
+	if (SubjectRole && !bIsDirty)
+	{
+		//Process Animation Data. Animation Role handling.
+		if(LiveLinkClient.DoesSubjectSupportsRole_AnyThread(SubjectName, ULiveLinkAnimationRole::StaticClass()))
+		{
+			if(LiveLinkClient.IsSubjectValid(SubjectName)) //Check there is valid data on the subject before evaluate it. 
+			{
+				FLiveLinkSubjectFrameData SubjectFrameData;
+				LiveLinkClient.EvaluateFrame_AnyThread(SubjectName, LiveLinkClient.GetSubjectRole_AnyThread(SubjectName), SubjectFrameData);
+
+				FLiveLinkSkeletonStaticData* SkeletonStaticData = SubjectFrameData.StaticData.Cast<FLiveLinkSkeletonStaticData>();
+				FLiveLinkAnimationFrameData* AnimationFrameData = SubjectFrameData.FrameData.Cast<FLiveLinkAnimationFrameData>();
+				check(SkeletonStaticData);
+				check(AnimationFrameData);
+
+				CacheSkeletalAnimationData(SkeletonStaticData, AnimationFrameData);
+
+				if(GetInstanceCount()<1)
+				{
+					CreateInstances();
+				}
+
+				if(bEvaluateLiveLink)
+				{
+					switch (BoneVisualType)
+					{
+					case ELiveLinkVisualBoneType::Joint:
+						{
+							BatchUpdateInstancesTransforms(0, GetJointTransforms(), true, true, false);
+							break;
+						}
+					case ELiveLinkVisualBoneType::Bone:
+						{
+							BatchUpdateInstancesTransforms(0, GetBoneTransforms(), true, true, false);
+							break;
+						}
+					}
+				}
+			}
+		}
+		//Locator role handling
+		else if(LiveLinkClient.DoesSubjectSupportsRole_AnyThread(SubjectName, ULiveLinkLocatorRole::StaticClass()) && bEvaluateLiveLink)
+		{
+			
+			TArray<FTransform> Transforms;
+			TArray<FVector> Locators = GetLocatorPositions();
+
+			for (int32 i = 0; i < Locators.Num(); i++)
+			{
+				FVector Locator = Locators[i];
+				FTransform Transform = FTransform::Identity;
+				Transform.SetTranslation(Locator);
+				Transforms.Emplace(Transform * GetComponentTransform());
+			}
+			BatchUpdateInstancesTransforms(0, Transforms, true, true, false);
+		}
+		//Transform role handling.
+		else if(LiveLinkClient.DoesSubjectSupportsRole_AnyThread(SubjectName, ULiveLinkTransformRole::StaticClass()) && bEvaluateLiveLink)
+		{
+			BatchUpdateInstancesTransform(0,1,GetSingleTransform() * GetComponentTransform(),true,true,false);
+
+			if (bDrawLabels)
+			{
+				FVector Location = (GetSingleTransform() * GetComponentTransform()).GetTranslation();
+				float Time = 1.0;
+				UWorld* World = GetWorld();
+				DrawDebugString(World, Location, TEXT("this is some text"), GetOwner(), FColor::Red, Time, false, 1.0);
+			}
+		}
+	}
+	if (bIsDirty)
+	{
+		CreateInstances();
+	}
+}
+
+#if WITH_EDITOR
+void ULiveLinkDataPreviewComponent::PostEditChangeProperty(FPropertyChangedEvent& PropertyChangedEvent)
+{
+	Super::PostEditChangeProperty(PropertyChangedEvent);
+
+	if(PropertyChangedEvent.Property != nullptr)
+	{
+		if(PropertyChangedEvent.Property->GetFName() == GET_MEMBER_NAME_CHECKED(ULiveLinkDataPreviewComponent, SubjectName))
+		{
+			ClearInstances();
+			bIsDirty = true;
+		}
+		if(PropertyChangedEvent.Property->GetFName() == GET_MEMBER_NAME_CHECKED(ULiveLinkDataPreviewComponent, bEvaluateLiveLink))
+		{
+			ClearInstances();
+			bIsDirty = true;
+		}
+		if(PropertyChangedEvent.Property->GetFName() == GET_MEMBER_NAME_CHECKED(ULiveLinkDataPreviewComponent, BoneVisualType))
+		{
+			ClearInstances();
+			bIsDirty = true;
+		}
+		if(PropertyChangedEvent.Property->GetFName() == GET_MEMBER_NAME_CHECKED(ULiveLinkDataPreviewComponent, bDrawLabels))
+		{
+			ClearInstances();
+			bIsDirty = true;
+		}
+	}
+	
+}
+#endif
+
+void ULiveLinkDataPreviewComponent::OnRegister()
+{
+	Super::OnRegister();
+	
+	ClearInstances();
+
+	bIsDirty = true;
+}
+
+void ULiveLinkDataPreviewComponent::OnUnregister()
+{
+	Super::OnUnregister();
+
+	ClearInstances();
+	
+}
+
+void ULiveLinkDataPreviewComponent::GetTransformRootSpace(const int32 InTransformIndex, FTransform& OutTransform) const
+{
+	// Case: Root joint or invalid
+	OutTransform = FTransform::Identity;
+	
+	if (IsValidTransformIndex(InTransformIndex, CachedAnimationData))
+	{
+		TPair<bool, FTransform>& RootSpaceCache = CachedRootSpaceTransforms[InTransformIndex];
+		bool bHasValidCache = RootSpaceCache.Key;
+		// Case: Have Cached Value
+		if (bHasValidCache)
+		{
+			OutTransform = RootSpaceCache.Value;
+		}
+		// Case: Need to generate Cache
+		else
+		{
+			const TArray<int32>& BoneParents = CachedSkeletonData.GetBoneParents();
+			int32 ParentIndex = BoneParents[InTransformIndex];
+
+			const FTransform& LocalSpaceTransform = CachedAnimationData.Transforms[InTransformIndex];
+
+			FTransform ParentRootSpaceTransform;
+
+			if(InTransformIndex == 0 && CachedSkeletonData.BoneParents[0] == 0) //special case handling for virtual subjects where bone[0] has parent bone[0]
+			{
+				ParentRootSpaceTransform = FTransform::Identity;
+			}
+			else
+			{
+				GetTransformRootSpace(ParentIndex, ParentRootSpaceTransform);
+			}
+
+			OutTransform = LocalSpaceTransform * ParentRootSpaceTransform;
+
+			// Save cached results
+			RootSpaceCache.Key = true;
+			RootSpaceCache.Value = OutTransform;
+		}
+	}
+}
+
+//Add the static mesh instances to represent the Live Link data. 
+void ULiveLinkDataPreviewComponent::CreateInstances()
+{
+	ILiveLinkClient& LiveLinkClient = IModularFeatures::Get().GetModularFeature<ILiveLinkClient>(ILiveLinkClient::ModularFeatureName);
+
+	TSubclassOf<ULiveLinkRole> SubjectRole = LiveLinkClient.GetSubjectRole_AnyThread(SubjectName);
+	if (SubjectRole)
+	{
+		//Process Animation Data
+		if(LiveLinkClient.DoesSubjectSupportsRole_AnyThread(SubjectName, ULiveLinkAnimationRole::StaticClass()))
+		{
+			switch (BoneVisualType)
+			{
+			case ELiveLinkVisualBoneType::Joint:
+				{
+					AddInstances(GetJointTransforms(), false, true, false);
+					break;
+				}
+			case ELiveLinkVisualBoneType::Bone:
+				{
+					AddInstances(GetBoneTransforms(), false, true, false);
+					break;
+				}
+			}
+			UE_LOG(LogLiveLink, Display, TEXT("Created Instances %d"), GetInstanceCount());
+			UE_LOG(LogLiveLink, Display, TEXT("Created Instances %s"), bIsDirty ? TEXT("true") : TEXT("false") );
+			bIsDirty=false;
+		}
+		else if(LiveLinkClient.DoesSubjectSupportsRole_AnyThread(SubjectName, ULiveLinkLocatorRole::StaticClass()))
+		{
+			TArray<FTransform> Transforms;
+			TArray<FVector> Locators = GetLocatorPositions();
+
+			for (int32 i = 0; i < Locators.Num(); i++)
+			{
+				FVector Locator = Locators[i];
+				FTransform Transform = FTransform::Identity;
+				Transform.SetTranslation(Locator);
+				Transforms.Emplace(Transform * GetComponentTransform());
+			}
+			AddInstances(Transforms, false, true, false);
+			bIsDirty=false;
+		}
+		else if(LiveLinkClient.DoesSubjectSupportsRole_AnyThread(SubjectName, ULiveLinkTransformRole::StaticClass()))
+		{
+			AddInstance(GetSingleTransform() * GetComponentTransform(),true);
+			bIsDirty=false;
+		}
+	}
+}
+
+//Loop through a bone's parent's to get it's transform in actor space. Live Link bone transforms are expressed in parent bone space so we need to do this to convert to actor space. 
+TArray<FTransform> ULiveLinkDataPreviewComponent::GetBoneTransforms() const
+{
+	TArray<FTransform> JointTransforms;
+		
+		FTransform BoneTransform;
+		FTransform ParentBoneTransform;
+		double BoneLength;
+		FVector Aim = FVector( 0, 0, 1);
+
+		const TArray<FTransform> BoneTransforms = CachedAnimationData.Transforms;
+		const TArray<int32> BoneParents = CachedSkeletonData.BoneParents;
+
+					
+		for (int32 i=0; i < CachedSkeletonData.BoneNames.Num(); i++)
+		{
+			GetTransformRootSpace(i, BoneTransform);
+			BoneTransform = BoneTransform * GetComponentTransform();
+			GetTransformRootSpace(CachedSkeletonData.BoneParents[i],ParentBoneTransform);
+			ParentBoneTransform = ParentBoneTransform * GetComponentTransform();
+			
+			BoneLength = FVector::Distance(BoneTransform.GetTranslation(), ParentBoneTransform.GetTranslation());
+			BoneTransform.SetScale3D(FVector(1, 1, BoneLength));
+			FQuat DiffRotation = AnimationCore::SolveAim(BoneTransform, ParentBoneTransform.GetLocation(), Aim.GetSafeNormal(), false, FVector(1, 1, 1), float(0));
+			BoneTransform.SetRotation(DiffRotation);
+			JointTransforms.Add(BoneTransform);
+		}
+	return JointTransforms;
+}
+
+//Get bone transforms from the cached skeleton data.
+TArray<FTransform> ULiveLinkDataPreviewComponent::GetJointTransforms() const
+{
+	TArray<FTransform> BoneTransforms;
+	
+				
+		int32 i = 0;
+		FTransform BoneTransform;
+
+		for (i=0 ; i < CachedSkeletonData.BoneNames.Num() ; i++)
+		{
+			GetTransformRootSpace(i, BoneTransform);
+			BoneTransform = BoneTransform * GetComponentTransform();
+			BoneTransforms.Add(BoneTransform);
+		}
+	return BoneTransforms;
+}
+
+//Get the location of each locator in the locator role and apply it to the instanced static meshes.
+TArray<FVector> ULiveLinkDataPreviewComponent::GetLocatorPositions() const
+{
+	ILiveLinkClient& LiveLinkClient = IModularFeatures::Get().GetModularFeature<ILiveLinkClient>(ILiveLinkClient::ModularFeatureName);
+	
+	TArray<FVector> MarkerPostions;
+
+	TSubclassOf<ULiveLinkRole> SubjectRole = LiveLinkClient.GetSubjectRole_AnyThread(SubjectName);
+	if (SubjectRole)
+	{
+		if(LiveLinkClient.DoesSubjectSupportsRole_AnyThread(SubjectName, ULiveLinkLocatorRole::StaticClass()))
+		{
+			if(LiveLinkClient.IsSubjectValid(SubjectName))
+			{
+				FLiveLinkSubjectFrameData SubjectFrameData;
+				LiveLinkClient.EvaluateFrame_AnyThread(SubjectName, LiveLinkClient.GetSubjectRole_AnyThread(SubjectName), SubjectFrameData);
+							
+				const FLiveLinkLocatorStaticData* MarkerStaticData = SubjectFrameData.StaticData.Cast<FLiveLinkLocatorStaticData>();
+				const FLiveLinkLocatorFrameData* MarkerFrameData = SubjectFrameData.FrameData.Cast<FLiveLinkLocatorFrameData>();
+				check(MarkerStaticData);
+				check(MarkerFrameData);
+
+				if (MarkerFrameData->Locators.Num() > 0)
+				{
+					TArray<FVector> MarkerPositions = MarkerFrameData->Locators;
+					return MarkerFrameData->Locators;
+				}
+			}
+
+		}
+	}
+	return MarkerPostions;
+}
+
+//Get a tranform role's data and applying it to the instance static mesh in actor space.
+FTransform ULiveLinkDataPreviewComponent::GetSingleTransform() const
+{
+	ILiveLinkClient& LiveLinkClient = IModularFeatures::Get().GetModularFeature<ILiveLinkClient>(ILiveLinkClient::ModularFeatureName);
+
+	FTransform Transform;
+	
+	if(LiveLinkClient.DoesSubjectSupportsRole_AnyThread(SubjectName, ULiveLinkTransformRole::StaticClass()))
+	{
+		if(LiveLinkClient.IsSubjectValid(SubjectName))
+		{
+			FLiveLinkSubjectFrameData SubjectFrameData;
+			LiveLinkClient.EvaluateFrame_AnyThread(SubjectName, LiveLinkClient.GetSubjectRole_AnyThread(SubjectName), SubjectFrameData);
+			const FLiveLinkTransformFrameData* TransformFrameData = SubjectFrameData.FrameData.Cast<FLiveLinkTransformFrameData>();
+			Transform = TransformFrameData->Transform;
+		}
+
+	}
+
+	return Transform;
+}
+
+//Cached a frame's Live Link Animation role bone data for use later in the frame, updating in the instance static mesh components. 
+void ULiveLinkDataPreviewComponent::CacheSkeletalAnimationData(const FLiveLinkSkeletonStaticData* InStaticData, const FLiveLinkAnimationFrameData* InFrameData)
+{
+	{
+		CachedSkeletonData = *InStaticData;
+		CachedAnimationData = *InFrameData;
+
+		const int32 NumTransforms = InFrameData->Transforms.Num();
+		check(InStaticData->BoneNames.Num() == NumTransforms);
+		check(InStaticData->BoneParents.Num() == NumTransforms);
+		check(InStaticData->PropertyNames.Num() == InFrameData->PropertyValues.Num());
+		CachedRootSpaceTransforms.SetNum(NumTransforms);
+		CachedChildTransformIndices.SetNum(NumTransforms);
+		for (int32 i = 0; i < NumTransforms; ++i)
+		{
+			CachedRootSpaceTransforms[i].Key = false;
+			CachedChildTransformIndices[i].Key = false;
+		}
+	}
+}
+
+//Check a transform is valid, based on the given index. 
+bool ULiveLinkDataPreviewComponent::IsValidTransformIndex(int32 InTransformIndex, FLiveLinkAnimationFrameData InAnimData)
+{
+	return (InTransformIndex >= 0) && (InTransformIndex < InAnimData.Transforms.Num());
+}
diff --git a/Engine/Plugins/Animation/LiveLink/Source/LiveLink/Private/Visualizers/LiveLinkMarkerVisualizer.cpp b/Engine/Plugins/Animation/LiveLink/Source/LiveLink/Private/Visualizers/LiveLinkMarkerVisualizer.cpp
new file mode 100644
--- /dev/null
+++ b/Engine/Plugins/Animation/LiveLink/Source/LiveLink/Private/Visualizers/LiveLinkMarkerVisualizer.cpp
@@ -0,0 +1,134 @@
+// Copyright Epic Games, Inc. All Rights Reserved.
+
+
+#include "Visualizers/LiveLinkMarkerVisualizer.h"
+
+#include "Materials/MaterialInstanceDynamic.h"
+#include "Roles/LiveLinkAnimationRole.h"
+#include "ILiveLinkClient.h"
+#include "Features/IModularFeatures.h"
+#include "LiveLinkClient.h"
+#include "Roles/LiveLinkAnimationRole.h"
+#include "Roles/LiveLinkTransformRole.h"
+
+ULiveLinkMarkerVisualizer::ULiveLinkMarkerVisualizer(const FObjectInitializer& ObjectInitializer)
+{
+	PrimaryComponentTick.bCanEverTick = true;
+	bTickInEditor = true; //make this component tick in editor
+	PrimaryComponentTick.TickGroup = ETickingGroup::TG_PostUpdateWork;
+	bDrawLabels = false;
+	XAxisSign = EAxisSign::Positive;
+	YAxisSign = EAxisSign::Positive;
+	bEvaluateLiveLink = true;
+	
+}
+
+ULiveLinkMarkerVisualizer::~ULiveLinkMarkerVisualizer()
+{
+	//Placeholder 
+}
+
+void ULiveLinkMarkerVisualizer::TickComponent(float DeltaTime, ELevelTick TickType, FActorComponentTickFunction* ThisTickFunction)
+{
+	Super::TickComponent(DeltaTime, TickType, ThisTickFunction);
+
+	//If there are no instances, add them. 
+	if(GetInstanceCount() == 0)
+	{
+		AddInstances(GetMarkerTransforms(), true,true,false);
+	}
+	
+	if(bEvaluateLiveLink)
+	{
+		BatchUpdateInstancesTransforms(0, GetMarkerTransforms(), true,true,false);
+	}
+	
+	if(bDrawLabels)
+	{
+		DrawLabels(MarkerLabels);
+	}
+}
+
+TArray<FTransform> ULiveLinkMarkerVisualizer::GetMarkerTransforms()
+{
+	FLiveLinkSubjectFrameData SubjectFrameData;
+	
+	ILiveLinkClient& LiveLinkClient = IModularFeatures::Get().GetModularFeature<ILiveLinkClient>(ILiveLinkClient::ModularFeatureName);
+
+	LiveLinkClient.EvaluateFrame_AnyThread(LiveLinkSubject, LiveLinkClient.GetSubjectRole_AnyThread(LiveLinkSubject), SubjectFrameData);
+
+	const FLiveLinkBaseFrameData* MarkerData = SubjectFrameData.FrameData.GetBaseData();
+	const FLiveLinkBaseStaticData* MarkerNames = SubjectFrameData.StaticData.GetBaseData();
+
+	//TODO : special handling for world attached skeletal mesh actor (animation role)
+
+	if(MarkerData)
+	{
+		//Early out if the two array's aren't the same length
+		if(MarkerNames->PropertyNames.Num() != MarkerData->PropertyValues.Num())
+		{
+			UE_LOG(LogLiveLink, Warning, TEXT("LiveLnk data not valid - length of property values does not match length of property names"))
+			return MarkerLocations;
+		}
+
+		MarkerLocations.Empty();
+	
+		const TArray<float>& LocationValues = MarkerData->PropertyValues;
+		float XScale = GetAxisSign(XAxisSign);
+		float YScale = GetAxisSign(YAxisSign);
+		FTransform MarkerTransform;
+		AActor* OwnerActor = GetOwner();
+		AActor* AttachedParent = GetOwner()->GetAttachParentActor();
+
+		if(LiveLinkClient.GetSubjectRole_AnyThread(LiveLinkSubject) == ULiveLinkAnimationRole::StaticClass())
+		{
+			for ( int i=0; i < MarkerData->PropertyValues.Num();  i = i + 3)
+			{
+				MarkerTransform = FTransform(FVector(LocationValues[i], LocationValues[i+1], LocationValues[i+2])) * FTransform(FRotator(0,0,0), FVector(0,0,0),FVector(XScale,YScale,1));
+				MarkerTransform = MarkerTransform * OwnerActor->GetActorTransform();
+				//MarkerLocations.Emplace(MarkerTransform);
+				MarkerLocations.Add(MarkerTransform);
+			}
+		}
+		else //If not an animation role assume a transform or a basic role
+		{
+			for ( int i=0; i < MarkerData->PropertyValues.Num();  i = i + 3)
+			{
+				if(AttachedParent!=nullptr)
+				{
+					MarkerTransform = FTransform(FVector(LocationValues[i], LocationValues[i+1], LocationValues[i+2])) * FTransform(FRotator(0,0,0), FVector(0,0,0),FVector(XScale,YScale,1));
+					MarkerTransform = MarkerTransform * AttachedParent->GetActorTransform();
+					MarkerLocations.Add(MarkerTransform);
+				}
+				else
+				{
+					MarkerTransform = FTransform(FVector(LocationValues[i], LocationValues[i+1], LocationValues[i+2])) * FTransform(FRotator(0,0,0), FVector(0,0,0),FVector(XScale,YScale,1));
+					MarkerLocations.Add(MarkerTransform);
+				}
+			}
+		}
+		
+	}
+	return MarkerLocations;
+}
+
+void ULiveLinkMarkerVisualizer::DrawLabels(TArray<FName> Labels)
+{
+	//ToDo
+}
+
+void ULiveLinkMarkerVisualizer::OnRegister()
+{
+	Super::OnRegister();
+
+	//ClearInstances();
+	
+	AddInstances(GetMarkerTransforms(), true,true,false);
+}
+
+void ULiveLinkMarkerVisualizer::OnUnregister()
+{
+	Super::OnUnregister();
+
+	ClearInstances();
+}
diff --git a/Engine/Plugins/Animation/LiveLink/Source/LiveLink/Private/Visualizers/LiveLinkMarkerVisualizer.h b/Engine/Plugins/Animation/LiveLink/Source/LiveLink/Private/Visualizers/LiveLinkMarkerVisualizer.h
new file mode 100644
--- /dev/null
+++ b/Engine/Plugins/Animation/LiveLink/Source/LiveLink/Private/Visualizers/LiveLinkMarkerVisualizer.h
@@ -0,0 +1,88 @@
+// Copyright Epic Games, Inc. All Rights Reserved.
+
+#pragma once
+
+#include "CoreMinimal.h"
+#include "Components/InstancedStaticMeshComponent.h"
+#include "LiveLinkTypes.h"
+#include "ILiveLinkClient.h"
+#include "LiveLinkMarkerVisualizer.generated.h"
+
+class UMaterialInstanceDynamic;
+
+/**
+ * Sign of the axis, to handle data coming from non-left handed systems.
+ */
+UENUM(BlueprintType)
+enum class EAxisSign : uint8
+{
+	Positive = 0	UMETA(DisplayName = "Positive"),
+	Negative = 1	UMETA(DisplayName = "Negative")
+};
+
+/**
+ * An Instanced Static Mesh Component to represent Motion Capture marker data locations
+ */
+UCLASS(ClassGroup=(PerformanceCapture), meta=(BlueprintSpawnableComponent),Transient, DisplayName = "MocapMarkerVisualizer")
+class LIVELINK_API ULiveLinkMarkerVisualizer : public UInstancedStaticMeshComponent
+{
+	GENERATED_BODY()
+
+public:
+	//Sets default values for this component's properties
+	ULiveLinkMarkerVisualizer(const FObjectInitializer& ObjectInitializer);
+
+	//Destructor
+	~ULiveLinkMarkerVisualizer();
+
+	//Called every frame
+	virtual void TickComponent(float DeltaTime, ELevelTick TickType, FActorComponentTickFunction* ThisTickFunction) override;
+
+	/** Get marker transforms defined in mocap tracking space */
+	TArray<FTransform> GetMarkerTransforms();
+
+	static void DrawLabels(TArray<FName> Labels);
+
+	/** Dynamic material to use on the instance static meshes. */
+	UPROPERTY(EditAnywhere, BlueprintReadOnly, Category="PerformanceCapture|Visualization")
+	TObjectPtr<UMaterialInstanceDynamic> DynamicMaterial;
+
+	/** Live Link subject to drive the instanced static mesh transforms. */
+	UPROPERTY(BlueprintReadWrite, EditAnywhere, Category="PerformanceCapture|Visualization")
+	FLiveLinkSubjectName LiveLinkSubject;
+
+	/** Bool to control evaluation of Live Link data. */
+	UPROPERTY(BlueprintReadWrite, EditAnywhere, Category="PerformanceCapture|Visualization")
+	bool bEvaluateLiveLink;
+
+	//To Do - add text labels to drawn markers
+	UPROPERTY(BlueprintReadWrite, EditAnywhere, Category="PerformanceCapture|Visualization")
+	bool bDrawLabels;
+
+	/** Sign for the x axis of the marker. */
+	UPROPERTY(BlueprintReadWrite, EditAnywhere, Category="PerformanceCapture|Visualization")
+	EAxisSign XAxisSign;
+
+	/** Sign for the y axis of the marker. */
+	UPROPERTY(BlueprintReadWrite, EditAnywhere, Category="PerformanceCapture|Visualization")
+	EAxisSign YAxisSign;
+
+	/** Array of Transforms for the markers.  */
+	UPROPERTY(BlueprintReadWrite, EditAnywhere, Category="PerformanceCapture|Visualization")
+	TArray<FTransform> MarkerLocations;
+
+protected:
+	
+	virtual void OnRegister() override;
+	virtual void OnUnregister() override;
+
+private:
+	
+	TArray<FName> MarkerLabels;
+
+	static float GetAxisSign(EAxisSign AxisSign)
+	{
+		switch (AxisSign)
+		{case EAxisSign::Positive:return 1.0f; case EAxisSign::Negative:return -1.0f; default: return 1.0f;	}
+	}
+};
diff --git a/Engine/Plugins/Animation/LiveLink/Source/LiveLink/Public/Visualizers/LiveLinkDataPreview.h b/Engine/Plugins/Animation/LiveLink/Source/LiveLink/Public/Visualizers/LiveLinkDataPreview.h
new file mode 100644
--- /dev/null
+++ b/Engine/Plugins/Animation/LiveLink/Source/LiveLink/Public/Visualizers/LiveLinkDataPreview.h
@@ -0,0 +1,91 @@
+// Copyright Epic Games, Inc. All Rights Reserved.
+
+#pragma once
+
+#include "CoreMinimal.h"
+#include "Components/BillboardComponent.h"
+#include "GameFramework/Actor.h"
+#include "LiveLinkTypes.h"
+#include "Visualizers/LiveLinkDataPreviewComponent.h"
+
+#include "LiveLinkDataPreview.generated.h"
+
+/**
+ * Actor for visualizing LiveLink data in the viewport/level editor
+ */
+
+UCLASS(Blueprintable, BlueprintType, DisplayName = "Live Link Data Preview Actor", Category = "LiveLink", MinimalApi)
+class ALiveLinkDataPreview : public AActor
+{
+	GENERATED_BODY()
+
+public:
+	// Sets default values for this actor's properties
+	LIVELINK_API ALiveLinkDataPreview();
+
+	/**
+	 * Stop Start animation evaluation.
+	 * @param bNewEvaluate Enable Live Link Evaluation
+	 */
+	UFUNCTION(CallInEditor, BlueprintCallable, DisplayName = "Evaluate Live Link", Category = LiveLink)
+	LIVELINK_API void SetEnableLiveLinkData(bool bNewEvaluate);
+
+public:
+	/** The list of Live Link subjects this actor will draw. */
+	UPROPERTY(EditAnywhere, BlueprintReadWrite, Category = LiveLink, meta = (ShowOnlyInnerProperties = "true"))
+	TArray<FLiveLinkSubjectName> Subjects;
+
+	TSoftObjectPtr<UStaticMesh> BoneMesh;
+	TSoftObjectPtr<UStaticMesh> JointMesh;
+	TSoftObjectPtr<UStaticMesh> AxisMesh;
+	TSoftObjectPtr<UStaticMesh> TransformMesh;
+	TSoftObjectPtr<UStaticMesh> CameraMesh;
+	TSoftObjectPtr<UStaticMesh> LocatorMesh;
+	TSoftObjectPtr<UTexture2D> SpriteTexture;
+
+	UPROPERTY(EditAnywhere, BlueprintReadWrite, Category = LiveLink)
+	bool bDrawLabels;
+
+	// Enable Live Link evaluation
+	UPROPERTY(EditAnywhere, Category = LiveLink)
+	bool bEvaluateLiveLink;
+
+	TArray<ULiveLinkDataPreviewComponent> Visualizers;
+
+	/** The billboard component showing the actor's root position. */
+	UPROPERTY(EditAnywhere, BlueprintReadWrite, Category = Components)
+	TObjectPtr<UBillboardComponent> BillboardComponent;
+
+protected:
+	/**
+	 * Initialize all the subjects.
+	 */
+	UFUNCTION(BlueprintCallable, DisplayName = "Initialize All Subjects", Category = LiveLink)
+	void InitializeSubjects();
+
+	/**
+	 * Set the material instance on a Live Link data preview component.
+	 * @param InDataPreviewComponent Preview component instanced static mesh.
+	 */
+	void SetMaterialInstance(ULiveLinkDataPreviewComponent* InDataPreviewComponent);
+
+	// Called every frame
+	virtual void Tick(float DeltaTime) override;
+
+	//Actor Interface
+	virtual void OnConstruction(const FTransform& Transform) override;
+
+	//PostEnd Actor Interface
+
+#if WITH_EDITOR
+	virtual void PostEditChangeProperty(FPropertyChangedEvent& PropertyChangedEvent) override;
+
+	virtual bool ShouldTickIfViewportsOnly() const override
+	{
+		return true;
+	}
+#endif
+
+private:
+	TArray<FLiveLinkSubjectKey> CachedSubjects;
+};
\ No newline at end of file
diff --git a/Engine/Plugins/Animation/LiveLink/Source/LiveLink/Public/Visualizers/LiveLinkDataPreviewComponent.h b/Engine/Plugins/Animation/LiveLink/Source/LiveLink/Public/Visualizers/LiveLinkDataPreviewComponent.h
new file mode 100644
--- /dev/null
+++ b/Engine/Plugins/Animation/LiveLink/Source/LiveLink/Public/Visualizers/LiveLinkDataPreviewComponent.h
@@ -0,0 +1,102 @@
+// Copyright Epic Games, Inc. All Rights Reserved.
+
+#pragma once
+
+#include "CoreMinimal.h"
+#include "LiveLinkTypes.h"
+#include "Components/InstancedStaticMeshComponent.h"
+#include "Roles/LiveLinkAnimationTypes.h"
+#include "Roles/LiveLinkLocatorTypes.h"
+#include "Roles/LiveLinkTransformTypes.h"
+#include "LiveLinkDataPreviewComponent.generated.h"
+
+class FCanvas;
+class UFont;
+class FPrimitiveDrawInterface;
+class FSceneView;
+
+UENUM(BlueprintType)
+enum class ELiveLinkVisualBoneType : uint8
+{
+	Joint = 0		UMETA(DisplayName = "Joint", Toolip = "For use drawing the location of a joint"),
+	Bone = 1		UMETA(DisplayName = "Bone", Toolip = "For use drawing a bone pointing at the next child in a chain")
+};
+
+/**
+ * An instance static mesh component for drawing LiveLink subject data in-level. Support drawing Transforms, Locators, Skeletons and Cameras
+ */
+UCLASS(ClassGroup=(LiveLink), Transient, DisplayName = "Live Link Data Preview", MinimalApi)
+class ULiveLinkDataPreviewComponent : public UInstancedStaticMeshComponent
+{
+	GENERATED_BODY()
+
+public:
+	// Sets default values for this component's properties
+	LIVELINK_API ULiveLinkDataPreviewComponent(const FObjectInitializer& ObjectInitializer);
+
+	// Called every frame
+	LIVELINK_API virtual void TickComponent(float DeltaTime, ELevelTick TickType, FActorComponentTickFunction* ThisTickFunction) override;
+
+	/** The Live Link subject this component will preview. */
+	UPROPERTY(EditAnywhere, BlueprintReadWrite, Category = LiveLink)
+	FLiveLinkSubjectName SubjectName;
+
+	/** Bool to control animation evaluation. */
+	UPROPERTY(EditAnywhere, BlueprintReadOnly, Category=LiveLink)
+	bool bEvaluateLiveLink;
+
+	/** Bool to control visibility of labels. */
+	UPROPERTY(EditAnywhere, BlueprintReadOnly, Category="LiveLink|Visualization")
+	bool bDrawLabels;
+
+	/** Type of bone visualization - joint or bone.  */
+	UPROPERTY(BlueprintReadWrite, EditAnywhere, Category="LiveLink|Visualization")
+	ELiveLinkVisualBoneType BoneVisualType = ELiveLinkVisualBoneType::Bone;
+
+	/**
+	 * Stop/Start animation updates. 
+	 * @param bNewEvaluateLiveLink New evaluation state.
+	 */
+	UFUNCTION(BlueprintCallable, Category="LiveLink|Visualization")
+	LIVELINK_API void SetEvaluateLiveLinkData(bool bNewEvaluateLiveLink) { bEvaluateLiveLink = bNewEvaluateLiveLink; }
+
+	/**
+	 * Show or Hide Labels
+	 * @param bNewDrawLabel New lable visiblity..
+	 */
+	//TODO: not yet implemented, waiting on slate immediate mode
+	UFUNCTION(BlueprintCallable, Category="LiveLink|Visualization")
+	LIVELINK_API void SetDrawLabels(bool bNewDrawLabel) { bDrawLabels = bNewDrawLabel; }
+	
+	bool bIsDirty;
+	
+#if WITH_EDITOR
+	LIVELINK_API virtual void PostEditChangeProperty(FPropertyChangedEvent& PropertyChangedEvent) override;
+#endif
+	
+protected:
+	virtual void OnRegister() override;
+	virtual void OnUnregister() override;
+	
+	void CreateInstances();
+	
+	TArray<FTransform> GetJointTransforms() const;
+	TArray<FTransform> GetBoneTransforms() const;
+	TArray<FVector> GetLocatorPositions() const;
+	FTransform GetSingleTransform() const;;
+
+private:
+	TArray<FName> MarkerLabels;
+
+	void CacheSkeletalAnimationData(const FLiveLinkSkeletonStaticData* InStaticData, const FLiveLinkAnimationFrameData* InFrameData);
+	
+	FLiveLinkSkeletonStaticData CachedSkeletonData;
+	FLiveLinkAnimationFrameData CachedAnimationData;
+
+	void GetTransformRootSpace(const int32 InTransformIndex, FTransform& OutTransform) const;
+
+	mutable TArray<TPair<bool, FTransform>> CachedRootSpaceTransforms;
+	mutable TArray<TPair<bool, TArray<int32>>> CachedChildTransformIndices;
+	static bool IsValidTransformIndex(int32 InTransformIndex, FLiveLinkAnimationFrameData AnimData);
+	
+};
diff --git a/Engine/Plugins/Animation/LiveLink/Source/LiveLinkEditor/LiveLinkEditor.Build.cs b/Engine/Plugins/Animation/LiveLink/Source/LiveLinkEditor/LiveLinkEditor.Build.cs
--- a/Engine/Plugins/Animation/LiveLink/Source/LiveLinkEditor/LiveLinkEditor.Build.cs
+++ b/Engine/Plugins/Animation/LiveLink/Source/LiveLinkEditor/LiveLinkEditor.Build.cs
@@ -39,6 +39,7 @@ namespace UnrealBuildTool.Rules
 					"MessageLog",
 					"MovieScene",
 					"Persona",
+					"PlacementMode",
 					"Projects",
 					"PropertyEditor",
 					"Settings",
diff --git a/Engine/Plugins/Animation/LiveLink/Source/LiveLinkEditor/Private/LiveLinkEditorModule.cpp b/Engine/Plugins/Animation/LiveLink/Source/LiveLinkEditor/Private/LiveLinkEditorModule.cpp
--- a/Engine/Plugins/Animation/LiveLink/Source/LiveLinkEditor/Private/LiveLinkEditorModule.cpp
+++ b/Engine/Plugins/Animation/LiveLink/Source/LiveLinkEditor/Private/LiveLinkEditorModule.cpp
@@ -7,6 +7,7 @@
 #include "Features/IModularFeatures.h"
 #include "Brushes/SlateNoResource.h"
 #include "ISettingsModule.h"
+#include "IPlacementModeModule.h"
 #include "LevelEditor.h"
 #include "LiveLinkSourceSettings.h"
 #include "PropertyEditorModule.h"
@@ -41,6 +42,8 @@
 #include "LiveLinkVirtualSubjectBoneAttachment.h"
 #include "Kismet2/KismetEditorUtilities.h"
 #include "VirtualSubjects/LiveLinkBlueprintVirtualSubject.h"
+#include "Visualizers/LiveLinkDataPreview.h"
+
 
 
 /**
@@ -138,6 +141,8 @@ public:
 		StyleSet->Set("LiveLinkClient.Common.RemoveSource", new IMAGE_PLUGIN_BRUSH(TEXT("icon_RemoveSource_40x"), Icon40x40));
 		StyleSet->Set("LiveLinkClient.Common.RemoveAllSources", new IMAGE_PLUGIN_BRUSH(TEXT("icon_RemoveSource_40x"), Icon40x40));
 
+		StyleSet->Set("PlacementBrowser.Icons.VirtualProduction", new IMAGE_PLUGIN_BRUSH_SVG(TEXT("Starship/VirtualProduction"), Icon16x16));
+
 		StyleSet->Set("LiveLinkController.WarningIcon", new IMAGE_BRUSH_SVG("Starship/Common/alert-circle", Icon16x16, FStyleColors::Warning));
 		
 		FButtonStyle Button = FButtonStyle()
@@ -170,6 +175,7 @@ public:
 		RegisterSettings();
 		RegisterCustomizations();
 		RegisterAutoGeneratedDefaultEvents();
+		RegisterPlacementModeItems();
 	}
 
 	void ModulesChangesCallback(FName ModuleName, EModuleChangeReason ReasonForChange)
@@ -184,6 +190,11 @@ public:
 	{
 		LLM_SCOPE_BYTAG(LiveLink_LiveLinkEditor);
 
+		if (!IsEngineExitRequested() && GEditor && UObjectInitialized())
+		{
+			UnregisterPlacementModeItems();
+		}
+
 		UnregisterCustomizations();
 
 		UnregisterTabSpawner();
@@ -322,6 +333,90 @@ private:
 		FKismetEditorUtilities::RegisterAutoGeneratedDefaultEvent(this, ULiveLinkBlueprintVirtualSubject::StaticClass(), "OnUpdate");
 	}
 
+	// Register the livelink data preview actor.
+	void RegisterPlacementModeItems()
+	{
+		auto RegisterPlaceActors = [&]() -> void
+			{
+				if (!GEditor)
+				{
+					return;
+				}
+
+				const FPlacementCategoryInfo* Info = GetVirtualProductionCategoryRegisteredInfo();
+
+				if (!Info)
+				{
+					UE_LOG(LogLiveLinkEditor, Warning, TEXT("Could not find or create VirtualProduction Place Actor Category"));
+					return;
+				}
+
+				/*Register the Live Link Data Prevew Actor*/
+				PlaceActors.Add(IPlacementModeModule::Get().RegisterPlaceableItem(Info->UniqueHandle, MakeShared<FPlaceableItem>(
+					*ALiveLinkDataPreview::StaticClass(),
+					FAssetData(ALiveLinkDataPreview::StaticClass()),
+					NAME_None,
+					NAME_None,
+					TOptional<FLinearColor>(),
+					TOptional<int32>(),
+					NSLOCTEXT("PlacementMode", "LiveLinkDataPreview", "Live Link Data Preview")
+				)));
+			};
+
+		if (FApp::CanEverRender())
+		{
+			if (GEngine && GEngine->IsInitialized())
+			{
+				RegisterPlaceActors();
+			}
+			else
+			{
+				FCoreDelegates::OnPostEngineInit.AddLambda(RegisterPlaceActors);
+			}
+		}
+	}
+
+	// Unregister the data preview actor placement item.
+	void UnregisterPlacementModeItems()
+	{
+		IPlacementModeModule& PlacementModeModule = IPlacementModeModule::Get();
+
+		for (TOptional<FPlacementModeID>& PlaceActor : PlaceActors)
+		{
+			if (PlaceActor.IsSet())
+			{
+				PlacementModeModule.UnregisterPlaceableItem(*PlaceActor);
+			}
+		}
+
+		PlaceActors.Empty();
+	}
+
+	// Get the virtual production actor category info.
+	const FPlacementCategoryInfo* GetVirtualProductionCategoryRegisteredInfo() const
+	{
+		IPlacementModeModule& PlacementModeModule = IPlacementModeModule::Get();
+
+		if (const FPlacementCategoryInfo* RegisteredInfo = PlacementModeModule.GetRegisteredPlacementCategory(FLevelEditorOutlinerBuiltInCategories::VirtualProduction()))
+		{
+			return RegisteredInfo;
+		}
+		else
+		{
+			FPlacementCategoryInfo Info(
+				LOCTEXT("VirtualProductionCategoryName", "Virtual Production"),
+				FSlateIcon(StyleSet->GetStyleSetName(), "PlacementBrowser.Icons.VirtualProduction"),
+				FLevelEditorOutlinerBuiltInCategories::VirtualProduction(),
+				TEXT("PMVirtualProduction"),
+				25 // Determines where the category shows up in the list with respect to the others.
+			);
+			Info.ShortDisplayName = LOCTEXT("VirtualProductionShortCategoryName", "VP");
+			IPlacementModeModule::Get().RegisterPlacementCategory(Info);
+
+			return PlacementModeModule.GetRegisteredPlacementCategory(FLevelEditorOutlinerBuiltInCategories::VirtualProduction());
+		}
+	}
+
 private:
 	FDelegateHandle LevelEditorTabManagerChangedHandle;
 	FDelegateHandle ModulesChangedHandle;
@@ -330,6 +425,9 @@ private:
 
 	// Track if we have registered
 	bool bHasRegisteredTabSpawners;
+
+	// Registered livelink data preview actor.
+	TArray<TOptional<FPlacementModeID>> PlaceActors;
 };
 
 TSharedPtr<FSlateStyleSet> FLiveLinkEditorModule::StyleSet;
diff --git a/Engine/Plugins/VirtualProduction/PerformanceCaptureWorkflow/Config/BasePerformanceCaptureWorkflow.ini b/Engine/Plugins/VirtualProduction/PerformanceCaptureWorkflow/Config/BasePerformanceCaptureWorkflow.ini
--- a/Engine/Plugins/VirtualProduction/PerformanceCaptureWorkflow/Config/BasePerformanceCaptureWorkflow.ini
+++ b/Engine/Plugins/VirtualProduction/PerformanceCaptureWorkflow/Config/BasePerformanceCaptureWorkflow.ini
@@ -7,8 +7,40 @@ ViewModelClass=/PerformanceCaptureWorkflow/Core/VM_PCapWorkflow.VM_PCapWorkflow_
 MocapManagerUI=/PerformanceCaptureWorkflow/Core/EW_MocapManager.EW_MocapManager
 DefaultSessionTemplate=/PerformanceCaptureWorkflow/Core/DefaultSessionTemplate.DefaultSessionTemplate
 DatabaseHelperClass=/PerformanceCaptureWorkflow/Database/EU_InternalDatabaseHelper.EU_InternalDatabaseHelper_C
+
 [CoreRedirects]
 +ClassRedirects=(OldName="/Script/PerformanceCaptureWorkflow.PCap_DataAsset",NewName="/Script/PerformanceCaptureWorkflow.PCapDataAsset")
 +ClassRedirects=(OldName="/Script/PerformanceCaptureWorkflow.PCap_CharacterDataAsset",NewName="/Script/PerformanceCaptureWorkflow.PCapCharacterDataAsset")
 +ClassRedirects=(OldName="/Script/PerformanceCaptureWorkflow.PCap_PropDataAsset",NewName="/Script/PerformanceCaptureWorkflow.PCapPropDataAsset")
-+ClassRedirects=(OldName="/Script/PerformanceCaptureWorkflow.PCap_PerformerDataAsset",NewName="/Script/PerformanceCaptureWorkflow.PCapPerformerDataAsset")
\ No newline at end of file
++ClassRedirects=(OldName="/Script/PerformanceCaptureWorkflow.PCap_PerformerDataAsset",NewName="/Script/PerformanceCaptureWorkflow.PCapPerformerDataAsset")
+
++ClassRedirects=(OldName="/Script/PerformanceCaptureWorkflow.PCapBoneVisualiser",NewName="/Script/LiveLink.LiveLinkBoneVisualizer")
++ClassRedirects=(OldName="/Script/PerformanceCaptureWorkflow.PCapMarkerVisualizer",NewName="/Script/LiveLink.LiveLinkMarkerVisualizer")
++ClassRedirects=(OldName="/Script/PerformanceCaptureWorkflow.LiveLinkDataPreview",NewName="/Script/LiveLink.LiveLinkDataPreview")
++ClassRedirects=(OldName="/Script/PerformanceCaptureWorkflow.LiveLinkDataPreviewComponent",NewName="/Script/LiveLink.LiveLinkDataPreviewComponent")
+
++PackageRedirects=(OldName="/PerformanceCaptureWorkflow/Visualizers/AxisVisualizer", NewName="/LiveLink/Visualizers/AxisVisualizer")
++PackageRedirects=(OldName="/PerformanceCaptureWorkflow/Visualizers/BoneVisualizer", NewName="/LiveLink/Visualizers/BoneVisualizer")
++PackageRedirects=(OldName="/PerformanceCaptureWorkflow/Visualizers/JointVisualizer", NewName="/LiveLink/Visualizers/JointVisualizer")
+
++PackageRedirects=(OldName="/PerformanceCaptureWorkflow/Visualizers/MI_JointAxisVisualizer_Blue", NewName="/LiveLink/Visualizers/MI_JointAxisVisualizer_Blue")
++PackageRedirects=(OldName="/PerformanceCaptureWorkflow/Visualizers/MI_JointAxisVisualizer_Green", NewName="/LiveLink/Visualizers/MI_JointAxisVisualizer_Green")
++PackageRedirects=(OldName="/PerformanceCaptureWorkflow/Visualizers/MI_JointAxisVisualizer_Red", NewName="/LiveLink/Visualizers/MI_JointAxisVisualizer_Red")
++PackageRedirects=(OldName="/PerformanceCaptureWorkflow/Visualizers/MI_JointAxisVisualizer_Yellow", NewName="/LiveLink/Visualizers/MI_JointAxisVisualizer_Yellow")
++PackageRedirects=(OldName="/PerformanceCaptureWorkflow/Visualizers/MI_MarkerAxisVisualizer_Blue", NewName="/LiveLink/Visualizers/MI_MarkerAxisVisualizer_Blue")
++PackageRedirects=(OldName="/PerformanceCaptureWorkflow/Visualizers/MI_MarkerAxisVisualizer_Green", NewName="/LiveLink/Visualizers/MI_MarkerAxisVisualizer_Green")
++PackageRedirects=(OldName="/PerformanceCaptureWorkflow/Visualizers/MI_MarkerAxisVisualizer_Red", NewName="/LiveLink/Visualizers/MI_MarkerAxisVisualizer_Red")
++PackageRedirects=(OldName="/PerformanceCaptureWorkflow/Visualizers/MI_MarkerAxisVisualizer_Yellow", NewName="/LiveLink/Visualizers/MI_MarkerAxisVisualizer_Yellow")
++PackageRedirects=(OldName="/PerformanceCaptureWorkflow/Visualizers/M_BoneVisualizer", NewName="/LiveLink/Visualizers/M_BoneVisualizer")
++PackageRedirects=(OldName="/PerformanceCaptureWorkflow/Visualizers/M_BoneVisualizerOverlay", NewName="/LiveLink/Visualizers/M_BoneVisualizerOverlay")
++PackageRedirects=(OldName="/PerformanceCaptureWorkflow/Visualizers/M_JointAxisVisualizer", NewName="/LiveLink/Visualizers/M_JointAxisVisualizer")
++PackageRedirects=(OldName="/PerformanceCaptureWorkflow/Visualizers/M_MarkerAxisVisualizer", NewName="/LiveLink/Visualizers/M_MarkerAxisVisualizer")
++PackageRedirects=(OldName="/PerformanceCaptureWorkflow/Visualizers/M_MarkerAxisVisualizer_NoDepth", NewName="/LiveLink/Visualizers/M_MarkerAxisVisualizer_NoDepth")
++PackageRedirects=(OldName="/PerformanceCaptureWorkflow/Visualizers/SM_Bone", NewName="/LiveLink/Visualizers/SM_Bone")
++PackageRedirects=(OldName="/PerformanceCaptureWorkflow/Visualizers/SM_Joint", NewName="/LiveLink/Visualizers/SM_Joint")
++PackageRedirects=(OldName="/PerformanceCaptureWorkflow/Visualizers/SM_JointAxis", NewName="/LiveLink/Visualizers/SM_JointAxis")
++PackageRedirects=(OldName="/PerformanceCaptureWorkflow/Visualizers/SM_LiveLinkCamera", NewName="/LiveLink/Visualizers/SM_LiveLinkCamera")
++PackageRedirects=(OldName="/PerformanceCaptureWorkflow/Visualizers/SM_MarkerAxis", NewName="/LiveLink/Visualizers/SM_MarkerAxis")
++PackageRedirects=(OldName="/PerformanceCaptureWorkflow/Visualizers/SM_MarkerBox", NewName="/LiveLink/Visualizers/SM_MarkerBox")
++PackageRedirects=(OldName="/PerformanceCaptureWorkflow/Visualizers/SM_MarkerSphere", NewName="/LiveLink/Visualizers/SM_MarkerSphere")
++PackageRedirects=(OldName="/PerformanceCaptureWorkflow/Visualizers/SM_TransformAxis", NewName="/LiveLink/Visualizers/SM_TransformAxis")
\ No newline at end of file
diff --git a/Engine/Plugins/VirtualProduction/PerformanceCaptureWorkflow/Source/PerformanceCaptureWorkflow/Private/PerformanceCapture.cpp b/Engine/Plugins/VirtualProduction/PerformanceCaptureWorkflow/Source/PerformanceCaptureWorkflow/Private/PerformanceCapture.cpp
--- a/Engine/Plugins/VirtualProduction/PerformanceCaptureWorkflow/Source/PerformanceCaptureWorkflow/Private/PerformanceCapture.cpp
+++ b/Engine/Plugins/VirtualProduction/PerformanceCaptureWorkflow/Source/PerformanceCaptureWorkflow/Private/PerformanceCapture.cpp
@@ -14,7 +14,6 @@
 #include "Widgets/Layout/SBox.h"
 #include "Widgets/Text/STextBlock.h"
 #include "ToolMenus.h"
-#include "Visualizers/LiveLinkDataPreview.h"
 
 static const FName PerformanceCapturePanelTabName("PerformanceCaptureTab");
 
@@ -49,7 +48,7 @@ void FPerformanceCaptureModule::StartupModule()
 		.SetMenuType(ETabSpawnerMenuType::Hidden)
 		.SetIcon(FSlateIcon(FPerformanceCaptureStyle::GetStyleSetName(), "PerformanceCapture.MocapManagerTabIcon", "PerformanceCapture.MocapManagerTabIcon.Small")
 		);
-	RegisterPlacementModeItems();
+
 }
 
 void FPerformanceCaptureModule::ShutdownModule()
@@ -60,10 +59,6 @@ void FPerformanceCaptureModule::ShutdownModule()
 		SettingsModule->UnregisterSettings("Project", "Plugins", "PerformanceCapture");
 	}
 
-	if (!IsEngineExitRequested() && GEditor && UObjectInitialized())
-	{
-		UnregisterPlacementModeItems();
-	}
 	//Clean up nomad tab spawner
 	FGlobalTabmanager::Get()->UnregisterNomadTabSpawner(PerformanceCapturePanelTabName);
 	
@@ -144,86 +139,6 @@ void FPerformanceCaptureModule::RegisterMenus()
 	}
 }
 
-void FPerformanceCaptureModule::UnregisterPlacementModeItems()
-{
-	IPlacementModeModule& PlacementModeModule = IPlacementModeModule::Get();
-
-	for (TOptional<FPlacementModeID>& PlaceActor : PlaceActors)
-	{
-		if (PlaceActor.IsSet())
-		{
-			PlacementModeModule.UnregisterPlaceableItem(*PlaceActor);
-		}
-	}
-
-	PlaceActors.Empty();
-}
-
-const FPlacementCategoryInfo* FPerformanceCaptureModule::GetVirtualProductionCategoryRegisteredInfo() const
-{
-	IPlacementModeModule& PlacementModeModule = IPlacementModeModule::Get();
-
-	if (const FPlacementCategoryInfo* RegisteredInfo = PlacementModeModule.GetRegisteredPlacementCategory(FLevelEditorOutlinerBuiltInCategories::VirtualProduction()))
-	{
-		return RegisteredInfo;
-	}
-	else
-	{
-		FPlacementCategoryInfo Info(
-			LOCTEXT("VirtualProductionCategoryName", "Virtual Production"),
-			FSlateIcon(FPerformanceCaptureStyle::Get().GetStyleSetName(), "PlacementBrowser.Icons.VirtualProduction"),
-			FLevelEditorOutlinerBuiltInCategories::VirtualProduction(),
-			TEXT("PMVirtualProduction"),
-			25 // Determines where the category shows up in the list with respect to the others.
-		);
-		Info.ShortDisplayName = LOCTEXT("VirtualProductionShortCategoryName", "VP");
-		IPlacementModeModule::Get().RegisterPlacementCategory(Info);
-
-		return PlacementModeModule.GetRegisteredPlacementCategory(FLevelEditorOutlinerBuiltInCategories::VirtualProduction());
-	}
-}
-
-void FPerformanceCaptureModule::RegisterPlacementModeItems()
-{
-	auto RegisterPlaceActors = [&]() -> void
-	{
-		if (!GEditor)
-		{
-			return;
-		}
-
-		const FPlacementCategoryInfo* Info = GetVirtualProductionCategoryRegisteredInfo();
-
-		if (!Info)
-		{
-			UE_LOG(LogPCap, Warning, TEXT("Could not find or create VirtualProduction Place Actor Category"));
-			return;
-		}
-
-		/*Register the Live Link Data Prevew Actor*/
-		PlaceActors.Add(IPlacementModeModule::Get().RegisterPlaceableItem(Info->UniqueHandle, MakeShared<FPlaceableItem>(
-			*ALiveLinkDataPreview::StaticClass(),
-			FAssetData(ALiveLinkDataPreview::StaticClass()),
-			NAME_None,
-			NAME_None,
-			TOptional<FLinearColor>(),
-			TOptional<int32>(),
-			NSLOCTEXT("PlacementMode", "LiveLinkDataPreview", "Live Link Data Preview")
-		)));
-	};
-
-	if (FApp::CanEverRender())
-	{
-		if (GEngine && GEngine->IsInitialized())
-		{
-			RegisterPlaceActors();
-		}
-		else
-		{
-			PostEngineInitHandle = FCoreDelegates::OnPostEngineInit.AddLambda(RegisterPlaceActors);
-		}
-	}
-}
 
 #undef LOCTEXT_NAMESPACE
 	
diff --git a/Engine/Plugins/VirtualProduction/PerformanceCaptureWorkflow/Source/PerformanceCaptureWorkflow/Private/PerformanceCaptureStyle.cpp b/Engine/Plugins/VirtualProduction/PerformanceCaptureWorkflow/Source/PerformanceCaptureWorkflow/Private/PerformanceCaptureStyle.cpp
--- a/Engine/Plugins/VirtualProduction/PerformanceCaptureWorkflow/Source/PerformanceCaptureWorkflow/Private/PerformanceCaptureStyle.cpp
+++ b/Engine/Plugins/VirtualProduction/PerformanceCaptureWorkflow/Source/PerformanceCaptureWorkflow/Private/PerformanceCaptureStyle.cpp
@@ -52,7 +52,6 @@ TSharedRef< FSlateStyleSet > FPerformanceCaptureStyle::Create()
 
 	Style->Set("PerformanceCapture.MocapManagerTabIcon.Small", new IMAGE_BRUSH_SVG(TEXT("ButtonIcon"), Icon16x16));
 	
-	Style->Set("PlacementBrowser.Icons.VirtualProduction", new IMAGE_BRUSH_SVG(TEXT("VirtualProduction"), Icon16x16));
 
 	Style->Set("ClassThumbnail.PCapDataTable", new IMAGE_BRUSH_SVG(TEXT("DataTableThumbnail"), Icon64x64));
 
diff --git a/Engine/Plugins/VirtualProduction/PerformanceCaptureWorkflow/Source/PerformanceCaptureWorkflow/Private/Visualizers/LiveLinkDataPreview.cpp b/Engine/Plugins/VirtualProduction/PerformanceCaptureWorkflow/Source/PerformanceCaptureWorkflow/Private/Visualizers/LiveLinkDataPreview.cpp
deleted file mode 100644
diff --git a/Engine/Plugins/VirtualProduction/PerformanceCaptureWorkflow/Source/PerformanceCaptureWorkflow/Private/Visualizers/LiveLinkDataPreview.h b/Engine/Plugins/VirtualProduction/PerformanceCaptureWorkflow/Source/PerformanceCaptureWorkflow/Private/Visualizers/LiveLinkDataPreview.h
deleted file mode 100644
diff --git a/Engine/Plugins/VirtualProduction/PerformanceCaptureWorkflow/Source/PerformanceCaptureWorkflow/Private/Visualizers/LiveLinkDataPreviewComponent.cpp b/Engine/Plugins/VirtualProduction/PerformanceCaptureWorkflow/Source/PerformanceCaptureWorkflow/Private/Visualizers/LiveLinkDataPreviewComponent.cpp
deleted file mode 100644
diff --git a/Engine/Plugins/VirtualProduction/PerformanceCaptureWorkflow/Source/PerformanceCaptureWorkflow/Private/Visualizers/LiveLinkDataPreviewComponent.h b/Engine/Plugins/VirtualProduction/PerformanceCaptureWorkflow/Source/PerformanceCaptureWorkflow/Private/Visualizers/LiveLinkDataPreviewComponent.h
deleted file mode 100644
diff --git a/Engine/Plugins/VirtualProduction/PerformanceCaptureWorkflow/Source/PerformanceCaptureWorkflow/Private/Visualizers/PCapBoneVisualizer.cpp b/Engine/Plugins/VirtualProduction/PerformanceCaptureWorkflow/Source/PerformanceCaptureWorkflow/Private/Visualizers/PCapBoneVisualizer.cpp
deleted file mode 100644
diff --git a/Engine/Plugins/VirtualProduction/PerformanceCaptureWorkflow/Source/PerformanceCaptureWorkflow/Private/Visualizers/PCapBoneVisualizer.h b/Engine/Plugins/VirtualProduction/PerformanceCaptureWorkflow/Source/PerformanceCaptureWorkflow/Private/Visualizers/PCapBoneVisualizer.h
deleted file mode 100644
diff --git a/Engine/Plugins/VirtualProduction/PerformanceCaptureWorkflow/Source/PerformanceCaptureWorkflow/Private/Visualizers/PCapMarkerVisualizer.cpp b/Engine/Plugins/VirtualProduction/PerformanceCaptureWorkflow/Source/PerformanceCaptureWorkflow/Private/Visualizers/PCapMarkerVisualizer.cpp
deleted file mode 100644
diff --git a/Engine/Plugins/VirtualProduction/PerformanceCaptureWorkflow/Source/PerformanceCaptureWorkflow/Private/Visualizers/PCapMarkerVisualizer.h b/Engine/Plugins/VirtualProduction/PerformanceCaptureWorkflow/Source/PerformanceCaptureWorkflow/Private/Visualizers/PCapMarkerVisualizer.h
deleted file mode 100644
diff --git a/Engine/Plugins/VirtualProduction/PerformanceCaptureWorkflow/Source/PerformanceCaptureWorkflow/Public/PerformanceCapture.h b/Engine/Plugins/VirtualProduction/PerformanceCaptureWorkflow/Source/PerformanceCaptureWorkflow/Public/PerformanceCapture.h
--- a/Engine/Plugins/VirtualProduction/PerformanceCaptureWorkflow/Source/PerformanceCaptureWorkflow/Public/PerformanceCapture.h
+++ b/Engine/Plugins/VirtualProduction/PerformanceCaptureWorkflow/Source/PerformanceCaptureWorkflow/Public/PerformanceCapture.h
@@ -24,18 +24,6 @@ private:
 
 	void RegisterMenus();
 
-	/** Register items to show up in the Place Actors panel. */
-	void RegisterPlacementModeItems();
-
-	/** Unregister items in Place Actors panel */
-	void UnregisterPlacementModeItems();
-
-	/** Gathers the Info on the Virtual Production Place Actors Category */
-	const FPlacementCategoryInfo* GetVirtualProductionCategoryRegisteredInfo() const;
-
-	TArray<TOptional<FPlacementModeID>> PlaceActors;
-	FDelegateHandle PostEngineInitHandle;
-
 	TSharedRef<class SDockTab> OnSpawnMocapManager(const class FSpawnTabArgs& SpawnTabArgs);
 	
 	TSharedPtr<class FUICommandList> PluginCommands;
