LiveLink - Convert "Add Source" menu to pop up a modal window instead of a dropdown men

diff --git a/Engine/Plugins/Animation/LiveLink/Source/LiveLinkEditor/Private/LiveLinkClientPanelToolbar.cpp b/Engine/Plugins/Animation/LiveLink/Source/LiveLinkEditor/Private/LiveLinkClientPanelToolbar.cpp
--- a/Engine/Plugins/Animation/LiveLink/Source/LiveLinkEditor/Private/LiveLinkClientPanelToolbar.cpp
+++ b/Engine/Plugins/Animation/LiveLink/Source/LiveLinkEditor/Private/LiveLinkClientPanelToolbar.cpp
@@ -49,7 +49,6 @@
 
 #define LOCTEXT_NAMESPACE "LiveLinkClientPanel"
 
-
 /** Dialog to create a new virtual subject */
 class SVirtualSubjectCreateDialog : public SCompoundWidget
 {
@@ -577,11 +576,15 @@ TSharedRef<SWidget> SLiveLinkClientPanelToolbar::OnGenerateSourceMenu()
 
 			if (MenuType == ULiveLinkSourceFactory::EMenuType::SubPanel)
 			{
-				MenuBuilder.AddSubMenu(
+				MenuBuilder.AddMenuEntry(
 					FactoryInstance->GetSourceDisplayName(),
 					FactoryInstance->GetSourceTooltip(),
-					FNewMenuDelegate::CreateSP(this, &SLiveLinkClientPanelToolbar::RetrieveFactorySourcePanel, FactoryIndex),
-					false);
+					FSlateIcon(),
+					FUIAction(
+						FExecuteAction::CreateSP(this, &SLiveLinkClientPanelToolbar::OpenCreateMenuWindow, FactoryIndex)
+					),
+					NAME_None,
+					EUserInterfaceActionType::Button);
 			}
 			else if (MenuType == ULiveLinkSourceFactory::EMenuType::MenuEntry)
 			{
@@ -638,7 +641,7 @@ TSharedRef<SWidget> SLiveLinkClientPanelToolbar::OnGenerateSourceMenu()
 	return MenuBuilder.MakeWidget();
 }
 
-void SLiveLinkClientPanelToolbar::RetrieveFactorySourcePanel(FMenuBuilder& MenuBuilder, int32 FactoryIndex)
+void SLiveLinkClientPanelToolbar::OpenCreateMenuWindow(int32 FactoryIndex)
 {
 	if (Factories.IsValidIndex(FactoryIndex))
 	{
@@ -647,9 +650,18 @@ void SLiveLinkClientPanelToolbar::RetrieveFactorySourcePanel(FMenuBuilder& MenuB
 			TSharedPtr<SWidget> Widget = FactoryInstance->BuildCreationPanel(ULiveLinkSourceFactory::FOnLiveLinkSourceCreated::CreateSP(this, &SLiveLinkClientPanelToolbar::OnSourceCreated, TSubclassOf<ULiveLinkSourceFactory>(FactoryInstance->GetClass())));
 			if (Widget.IsValid())
 			{
-				MenuBuilder.AddWidget(Widget.ToSharedRef()
-					, FText()
-					, true);
+				FText CreateText = FText::Format(LOCTEXT("CreateSourceLabel", "Create {0} connection"), FactoryInstance->GetSourceDisplayName());
+				SourceCreationWindow = SNew(SWindow)
+					.Title(CreateText)
+					.ClientSize(FVector2D(400, 150))
+					.SizingRule(ESizingRule::Autosized)
+					.SupportsMinimize(false)
+					.SupportsMaximize(false)
+					[
+						Widget.ToSharedRef()
+					];
+
+				GEditor->EditorAddModalWindow(SourceCreationWindow.ToSharedRef());
 			}
 		}
 	}
@@ -680,8 +692,13 @@ void SLiveLinkClientPanelToolbar::OnSourceCreated(TSharedPtr<ILiveLinkSource> Ne
 			}
 		}
 	}
-	FSlateApplication::Get().DismissAllMenus();
-}
+
+	if (SourceCreationWindow)
+	{
+		SourceCreationWindow->RequestDestroyWindow();
+		SourceCreationWindow.Reset();
+	}
+} 
 
 void SLiveLinkClientPanelToolbar::PopulateVirtualSubjectSourceCreationMenu(FMenuBuilder& InMenuBuilder)
 {
diff --git a/Engine/Plugins/Animation/LiveLink/Source/LiveLinkEditor/Public/LiveLinkClientPanelToolbar.h b/Engine/Plugins/Animation/LiveLink/Source/LiveLinkEditor/Public/LiveLinkClientPanelToolbar.h
--- a/Engine/Plugins/Animation/LiveLink/Source/LiveLinkEditor/Public/LiveLinkClientPanelToolbar.h
+++ b/Engine/Plugins/Animation/LiveLink/Source/LiveLinkEditor/Public/LiveLinkClientPanelToolbar.h
@@ -50,7 +50,8 @@ class LIVELINKEDITOR_API SLiveLinkClientPanelToolbar : public SCompoundWidget, p
 
 private:
 	TSharedRef<SWidget> OnGenerateSourceMenu();
-	void RetrieveFactorySourcePanel(FMenuBuilder& MenuBuilder, int32 FactoryIndex);
+	/** Open a modal window that hosts the source factory widget. */
+	void OpenCreateMenuWindow(int32 FactoryIndex);
 	void ExecuteCreateSource(int32 FactoryIndex);
 	void OnSourceCreated(TSharedPtr<ILiveLinkSource> NewSource, FString ConnectionString, TSubclassOf<ULiveLinkSourceFactory> Factory);
 
@@ -71,6 +72,8 @@ private:
 private:
 	FLiveLinkClient* Client;
 
+	/** Source creation modal window. */
+	TSharedPtr<SWindow> SourceCreationWindow;
 	TWeakPtr<IMenu> AddSubjectMenu;
 	TWeakPtr<IMenu> VirtualSubjectMenu;
 	TWeakPtr<STextEntryPopup> VirtualSubjectPopup;
diff --git a/Engine/Plugins/VirtualProduction/LiveLinkInputDevice/Source/LiveLinkInputDevice/Public/LiveLinkInputDeviceSourceFactory.h b/Engine/Plugins/VirtualProduction/LiveLinkInputDevice/Source/LiveLinkInputDevice/Public/LiveLinkInputDeviceSourceFactory.h
--- a/Engine/Plugins/VirtualProduction/LiveLinkInputDevice/Source/LiveLinkInputDevice/Public/LiveLinkInputDeviceSourceFactory.h
+++ b/Engine/Plugins/VirtualProduction/LiveLinkInputDevice/Source/LiveLinkInputDevice/Public/LiveLinkInputDeviceSourceFactory.h
@@ -15,7 +15,7 @@ public:
 	virtual FText GetSourceDisplayName() const override;
 	virtual FText GetSourceTooltip() const override;
 
-	virtual EMenuType GetMenuType() const override { return EMenuType::SubPanel; }
+	virtual EMenuType GetMenuType() const override { return EMenuType::MenuEntry; }
 	virtual TSharedPtr<SWidget> BuildCreationPanel(FOnLiveLinkSourceCreated OnLiveLinkSourceCreated) const override;
 	virtual TSharedPtr<ILiveLinkSource> CreateSource(const FString& ConnectionString) const override;
 
