Updated headers using UnrealCodeFixup to make sure dllstorage is on methods/staticvars instead of types. Part 1/n

diff --git a/Engine/Plugins/Animation/LiveLink/Source/LiveLinkGraphNode/Private/K2Node_UpdateVirtualSubjectDataTyped.h b/Engine/Plugins/Animation/LiveLink/Source/LiveLinkGraphNode/Private/K2Node_UpdateVirtualSubjectDataTyped.h
--- a/Engine/Plugins/Animation/LiveLink/Source/LiveLinkGraphNode/Private/K2Node_UpdateVirtualSubjectDataTyped.h
+++ b/Engine/Plugins/Animation/LiveLink/Source/LiveLinkGraphNode/Private/K2Node_UpdateVirtualSubjectDataTyped.h
@@ -5,53 +5,57 @@
 #include "K2Node_UpdateVirtualSubjectDataBase.h"
 #include "K2Node_UpdateVirtualSubjectDataTyped.generated.h"
 
-UCLASS()
-class LIVELINKGRAPHNODE_API UK2Node_UpdateVirtualSubjectStaticData : public UK2Node_UpdateVirtualSubjectDataBase
+#define UE_API LIVELINKGRAPHNODE_API
+
+UCLASS(MinimalAPI)
+class UK2Node_UpdateVirtualSubjectStaticData : public UK2Node_UpdateVirtualSubjectDataBase
 {
 	GENERATED_BODY()
 
 public:
 
 	//~ Begin UEdGraphNode Interface.
-	virtual FText GetNodeTitle(ENodeTitleType::Type TitleType) const override;
+	UE_API virtual FText GetNodeTitle(ENodeTitleType::Type TitleType) const override;
 	//~ End UEdGraphNode Interface.
 
 protected:
 
 	//~ Begin UK2Node_UpdateVirtualSubjectDataBase interface
-	virtual UScriptStruct* GetStructTypeFromRole(ULiveLinkRole* Role) const override;
-	virtual FName GetUpdateFunctionName() const override;
-	virtual FText GetStructPinName() const override;
+	UE_API virtual UScriptStruct* GetStructTypeFromRole(ULiveLinkRole* Role) const override;
+	UE_API virtual FName GetUpdateFunctionName() const override;
+	UE_API virtual FText GetStructPinName() const override;
 	virtual void AddPins(FKismetCompilerContext& CompilerContext, UK2Node_CallFunction* UpdateVirtualSubjectDataFunction) const override {}
 	//~ End UK2Node_UpdateVirtualSubjectDataBase interface
 };
 
-UCLASS()
-class LIVELINKGRAPHNODE_API UK2Node_UpdateVirtualSubjectFrameData : public UK2Node_UpdateVirtualSubjectDataBase
+UCLASS(MinimalAPI)
+class UK2Node_UpdateVirtualSubjectFrameData : public UK2Node_UpdateVirtualSubjectDataBase
 {
 	GENERATED_BODY()
 
 public:
 
 	//~ Begin UEdGraphNode Interface.
-	virtual void AllocateDefaultPins() override;
-	virtual FText GetNodeTitle(ENodeTitleType::Type TitleType) const override;
+	UE_API virtual void AllocateDefaultPins() override;
+	UE_API virtual FText GetNodeTitle(ENodeTitleType::Type TitleType) const override;
 	//~ End UEdGraphNode Interface.
 
 protected:
 
 	//~ Begin UK2Node_UpdateVirtualSubjectDataBase interface
-	virtual UScriptStruct* GetStructTypeFromRole(ULiveLinkRole* Role) const override;
-	virtual FName GetUpdateFunctionName() const override;
-	virtual FText GetStructPinName() const override;
-	virtual void AddPins(FKismetCompilerContext& CompilerContext, UK2Node_CallFunction* UpdateVirtualSubjectDataFunction) const override;
+	UE_API virtual UScriptStruct* GetStructTypeFromRole(ULiveLinkRole* Role) const override;
+	UE_API virtual FName GetUpdateFunctionName() const override;
+	UE_API virtual FText GetStructPinName() const override;
+	UE_API virtual void AddPins(FKismetCompilerContext& CompilerContext, UK2Node_CallFunction* UpdateVirtualSubjectDataFunction) const override;
 	//~ End UK2Node_UpdateVirtualSubjectDataBase interface
 
 	/** Returns the Timestamp pin  */
-	UEdGraphPin* GetTimestampFramePin() const;
+	UE_API UEdGraphPin* GetTimestampFramePin() const;
 
 private:
 
 	/** Name of the pin to enable/disable timestamping */
-	static const FName LiveLinkTimestampFramePinName;
-};
\ No newline at end of file
+	static UE_API const FName LiveLinkTimestampFramePinName;
+};
+
+#undef UE_API
diff --git a/Engine/Plugins/Animation/LiveLink/Source/LiveLinkGraphNode/Public/AnimGraphNode_LiveLinkPose.h b/Engine/Plugins/Animation/LiveLink/Source/LiveLinkGraphNode/Public/AnimGraphNode_LiveLinkPose.h
--- a/Engine/Plugins/Animation/LiveLink/Source/LiveLinkGraphNode/Public/AnimGraphNode_LiveLinkPose.h
+++ b/Engine/Plugins/Animation/LiveLink/Source/LiveLinkGraphNode/Public/AnimGraphNode_LiveLinkPose.h
@@ -7,26 +7,28 @@
 
 #include "AnimGraphNode_LiveLinkPose.generated.h"
 
+#define UE_API LIVELINKGRAPHNODE_API
+
 namespace ENodeTitleType { enum Type : int; }
 
 
 
-UCLASS()
-class LIVELINKGRAPHNODE_API UAnimGraphNode_LiveLinkPose : public UAnimGraphNode_Base
+UCLASS(MinimalAPI)
+class UAnimGraphNode_LiveLinkPose : public UAnimGraphNode_Base
 {
 	GENERATED_BODY()
 
 public:
 
 	//~ UEdGraphNode interface
-	virtual FText GetNodeTitle(ENodeTitleType::Type TitleType) const override;
-	virtual FText GetTooltipText() const override;
-	virtual FText GetMenuCategory() const;
+	UE_API virtual FText GetNodeTitle(ENodeTitleType::Type TitleType) const override;
+	UE_API virtual FText GetTooltipText() const override;
+	UE_API virtual FText GetMenuCategory() const;
 	//~ End of UEdGraphNode
 
 	//~ Begin UK2Node interface
-	virtual void ConvertDeprecatedNode(UEdGraph* Graph, bool bOnlySafeChanges) override;
-	virtual void GetOutputLinkAttributes(FNodeAttributeArray& OutAttributes) const override;
+	UE_API virtual void ConvertDeprecatedNode(UEdGraph* Graph, bool bOnlySafeChanges) override;
+	UE_API virtual void GetOutputLinkAttributes(FNodeAttributeArray& OutAttributes) const override;
 	//~ End UK2Node interface
 
 public:
@@ -35,3 +37,5 @@ public:
 	FAnimNode_LiveLinkPose Node;
 
 };
+
+#undef UE_API
diff --git a/Engine/Plugins/Animation/LiveLink/Source/LiveLinkGraphNode/Public/K2Node_EvaluateLiveLinkFrame.h b/Engine/Plugins/Animation/LiveLink/Source/LiveLinkGraphNode/Public/K2Node_EvaluateLiveLinkFrame.h
--- a/Engine/Plugins/Animation/LiveLink/Source/LiveLinkGraphNode/Public/K2Node_EvaluateLiveLinkFrame.h
+++ b/Engine/Plugins/Animation/LiveLink/Source/LiveLinkGraphNode/Public/K2Node_EvaluateLiveLinkFrame.h
@@ -7,6 +7,8 @@
 #include "UObject/Package.h"
 #include "K2Node_EvaluateLiveLinkFrame.generated.h"
 
+#define UE_API LIVELINKGRAPHNODE_API
+
 class ULiveLinkRole;
 
 class FBlueprintActionDatabaseRegistrar;
@@ -16,50 +18,50 @@ class UEdGraph;
 class UK2Node_CallFunction;
 
 
-UCLASS(Abstract)
-class LIVELINKGRAPHNODE_API UK2Node_EvaluateLiveLinkFrame : public UK2Node
+UCLASS(MinimalAPI, Abstract)
+class UK2Node_EvaluateLiveLinkFrame : public UK2Node
 {
 	GENERATED_BODY()
 
 public:
 	//~ Begin UEdGraphNode Interface.
-	virtual void AllocateDefaultPins() override;
-	virtual void PinDefaultValueChanged(UEdGraphPin* Pin) override;
-	virtual FText GetTooltipText() const override;
-	virtual void ExpandNode(FKismetCompilerContext& CompilerContext, UEdGraph* SourceGraph) override;
-	virtual FSlateIcon GetIconAndTint(FLinearColor& OutColor) const override;
-	virtual void PostReconstructNode() override;
+	UE_API virtual void AllocateDefaultPins() override;
+	UE_API virtual void PinDefaultValueChanged(UEdGraphPin* Pin) override;
+	UE_API virtual FText GetTooltipText() const override;
+	UE_API virtual void ExpandNode(FKismetCompilerContext& CompilerContext, UEdGraph* SourceGraph) override;
+	UE_API virtual FSlateIcon GetIconAndTint(FLinearColor& OutColor) const override;
+	UE_API virtual void PostReconstructNode() override;
 	//~ End UEdGraphNode Interface.
 
 	//~ Begin UK2Node Interface
 	virtual bool IsNodeSafeToIgnore() const override { return true; }
-	virtual void ReallocatePinsDuringReconstruction(TArray<UEdGraphPin*>& OldPins) override;
-	virtual void GetMenuActions(FBlueprintActionDatabaseRegistrar& ActionRegistrar) const override;
-	virtual FText GetMenuCategory() const override;
-	virtual bool IsConnectionDisallowed(const UEdGraphPin* MyPin, const UEdGraphPin* OtherPin, FString& OutReason) const override;
-	virtual void EarlyValidation(class FCompilerResultsLog& MessageLog) const override;
-	virtual void PreloadRequiredAssets() override;
-	virtual void NotifyPinConnectionListChanged(UEdGraphPin* Pin) override;
+	UE_API virtual void ReallocatePinsDuringReconstruction(TArray<UEdGraphPin*>& OldPins) override;
+	UE_API virtual void GetMenuActions(FBlueprintActionDatabaseRegistrar& ActionRegistrar) const override;
+	UE_API virtual FText GetMenuCategory() const override;
+	UE_API virtual bool IsConnectionDisallowed(const UEdGraphPin* MyPin, const UEdGraphPin* OtherPin, FString& OutReason) const override;
+	UE_API virtual void EarlyValidation(class FCompilerResultsLog& MessageLog) const override;
+	UE_API virtual void PreloadRequiredAssets() override;
+	UE_API virtual void NotifyPinConnectionListChanged(UEdGraphPin* Pin) override;
 	//~ End UK2Node Interface
 
 	/** Get the return types of our struct */
-	UScriptStruct* GetReturnTypeForOutputDataStruct() const;
+	UE_API UScriptStruct* GetReturnTypeForOutputDataStruct() const;
 
 	/** Get the Live Link Role input pin */
-	UEdGraphPin* GetLiveLinkRolePin() const;
+	UE_API UEdGraphPin* GetLiveLinkRolePin() const;
 
 	/** Get the Live Link Subject input pin */
-	UEdGraphPin* GetLiveLinkSubjectPin() const;
+	UE_API UEdGraphPin* GetLiveLinkSubjectPin() const;
 
 	/** Get the exec output pin for when no frame is available for the desired role */
-	UEdGraphPin* GetFrameNotAvailablePin() const;
+	UE_API UEdGraphPin* GetFrameNotAvailablePin() const;
 
 	/** Get the result output pin */
-	UEdGraphPin* GetResultingDataPin() const;
+	UE_API UEdGraphPin* GetResultingDataPin() const;
 
 	/** Get the type that the Live Link Role evaluates to to return */
-	UScriptStruct* GetLiveLinkRoleOutputStructType() const;
-	UScriptStruct* GetLiveLinkRoleOutputFrameStructType() const;
+	UE_API UScriptStruct* GetLiveLinkRoleOutputStructType() const;
+	UE_API UScriptStruct* GetLiveLinkRoleOutputFrameStructType() const;
 
 protected:
 
@@ -76,16 +78,18 @@ protected:
 	 * @param   MutatablePin	The pin you want to set tool-tip text on
 	 * @param   PinDescription	A string describing the pin's purpose
 	 */
-	void SetPinToolTip(UEdGraphPin& InOutMutatablePin, const FText& InPinDescription) const;
+	UE_API void SetPinToolTip(UEdGraphPin& InOutMutatablePin, const FText& InPinDescription) const;
 
 	/** Set the return type of our structs */
-	void SetReturnTypeForOutputStruct(UScriptStruct* InClass);
+	UE_API void SetReturnTypeForOutputStruct(UScriptStruct* InClass);
 
 	/** Queries for the authoritative return type, then modifies the return pin to match */
-	void RefreshDataOutputPinType();
+	UE_API void RefreshDataOutputPinType();
 
-	TSubclassOf<ULiveLinkRole> GetDefaultRolePinValue() const;
+	UE_API TSubclassOf<ULiveLinkRole> GetDefaultRolePinValue() const;
 
 	/** Verify if selected role class is valid for evaluation */
-	bool IsRoleValidForEvaluation(TSubclassOf<ULiveLinkRole> InRoleClass) const;
+	UE_API bool IsRoleValidForEvaluation(TSubclassOf<ULiveLinkRole> InRoleClass) const;
 };
+
+#undef UE_API
diff --git a/Engine/Plugins/Animation/LiveLink/Source/LiveLinkGraphNode/Public/K2Node_UpdateVirtualSubjectDataBase.h b/Engine/Plugins/Animation/LiveLink/Source/LiveLinkGraphNode/Public/K2Node_UpdateVirtualSubjectDataBase.h
--- a/Engine/Plugins/Animation/LiveLink/Source/LiveLinkGraphNode/Public/K2Node_UpdateVirtualSubjectDataBase.h
+++ b/Engine/Plugins/Animation/LiveLink/Source/LiveLinkGraphNode/Public/K2Node_UpdateVirtualSubjectDataBase.h
@@ -6,6 +6,8 @@
 
 #include "K2Node_UpdateVirtualSubjectDataBase.generated.h"
 
+#define UE_API LIVELINKGRAPHNODE_API
+
 namespace ENodeTitleType { enum Type : int; }
 
 class FBlueprintActionDatabaseRegistrar;
@@ -15,31 +17,31 @@ class UEdGraph;
 class UK2Node_CallFunction;
 class ULiveLinkRole;
 
-UCLASS(Abstract)
-class LIVELINKGRAPHNODE_API UK2Node_UpdateVirtualSubjectDataBase : public UK2Node
+UCLASS(MinimalAPI, Abstract)
+class UK2Node_UpdateVirtualSubjectDataBase : public UK2Node
 {
 	GENERATED_BODY()
 
 public:
 	//~ Begin UEdGraphNode Interface.
-	virtual void AllocateDefaultPins() override;
-	virtual FText GetNodeTitle(ENodeTitleType::Type TitleType) const override;
-	virtual FText GetTooltipText() const override;
-	virtual void ExpandNode(FKismetCompilerContext& CompilerContext, UEdGraph* SourceGraph) override;
-	virtual FSlateIcon GetIconAndTint(FLinearColor& OutColor) const override;
-	virtual bool IsCompatibleWithGraph(const UEdGraph* TargetGraph) const override;
+	UE_API virtual void AllocateDefaultPins() override;
+	UE_API virtual FText GetNodeTitle(ENodeTitleType::Type TitleType) const override;
+	UE_API virtual FText GetTooltipText() const override;
+	UE_API virtual void ExpandNode(FKismetCompilerContext& CompilerContext, UEdGraph* SourceGraph) override;
+	UE_API virtual FSlateIcon GetIconAndTint(FLinearColor& OutColor) const override;
+	UE_API virtual bool IsCompatibleWithGraph(const UEdGraph* TargetGraph) const override;
 	//~ End UEdGraphNode Interface.
 
 	//~ Begin UK2Node Interface
 	virtual bool IsNodeSafeToIgnore() const override { return true; }
-	virtual void GetMenuActions(FBlueprintActionDatabaseRegistrar& ActionRegistrar) const override;
-	virtual FText GetMenuCategory() const override;
-	virtual bool IsConnectionDisallowed(const UEdGraphPin* MyPin, const UEdGraphPin* OtherPin, FString& OutReason) const override;
-	virtual void EarlyValidation(class FCompilerResultsLog& MessageLog) const override;
+	UE_API virtual void GetMenuActions(FBlueprintActionDatabaseRegistrar& ActionRegistrar) const override;
+	UE_API virtual FText GetMenuCategory() const override;
+	UE_API virtual bool IsConnectionDisallowed(const UEdGraphPin* MyPin, const UEdGraphPin* OtherPin, FString& OutReason) const override;
+	UE_API virtual void EarlyValidation(class FCompilerResultsLog& MessageLog) const override;
 	//~ End UK2Node Interface
 
 	/** Get the Live Link Struct input pin */
-	UEdGraphPin* GetLiveLinkStructPin() const;
+	UE_API UEdGraphPin* GetLiveLinkStructPin() const;
 
 protected:
 	/**
@@ -49,7 +51,7 @@ protected:
 	 * @param   MutatablePin	The pin you want to set tool-tip text on
 	 * @param   PinDescription	A string describing the pin's purpose
 	 */
-	void SetPinToolTip(UEdGraphPin& InOutMutatablePin, const FText& InPinDescription) const;
+	UE_API void SetPinToolTip(UEdGraphPin& InOutMutatablePin, const FText& InPinDescription) const;
 
 	/** Returns Struct type associated to the subject's role (static or frame) */
 	virtual UScriptStruct* GetStructTypeFromRole(ULiveLinkRole* Role) const PURE_VIRTUAL(UK2Node_UpdateVirtualSubjectDataBase::GetStructTypeFromRole, return nullptr; );
@@ -66,13 +68,15 @@ protected:
 private:
 
 	/** Returns the Struct type associated to the role */
-	UScriptStruct* GetStructTypeFromBlueprint() const;
+	UE_API UScriptStruct* GetStructTypeFromBlueprint() const;
 
 private:
 	
 	/** Name of the struct pin */
-	static const FName LiveLinkStructPinName;
+	static UE_API const FName LiveLinkStructPinName;
 
 	/** Struct pin description */
-	static const FText LiveLinkStructPinDescription;
+	static UE_API const FText LiveLinkStructPinDescription;
 };
+
+#undef UE_API
diff --git a/Engine/Plugins/Animation/LiveLink/Source/LiveLinkSequencer/Public/LiveLinkSequencerSettings.h b/Engine/Plugins/Animation/LiveLink/Source/LiveLinkSequencer/Public/LiveLinkSequencerSettings.h
--- a/Engine/Plugins/Animation/LiveLink/Source/LiveLinkSequencer/Public/LiveLinkSequencerSettings.h
+++ b/Engine/Plugins/Animation/LiveLink/Source/LiveLinkSequencer/Public/LiveLinkSequencerSettings.h
@@ -8,6 +8,8 @@
 
 #include "LiveLinkSequencerSettings.generated.h"
 
+#define UE_API LIVELINKSEQUENCER_API
+
 class ULiveLinkControllerBase;
 class UMovieSceneLiveLinkControllerTrackRecorder;
 
@@ -15,17 +17,17 @@ class UMovieSceneLiveLinkControllerTrackRecorder;
 /**
  * Settings for LiveLink Sequence Editor
  */
-UCLASS(config=Game)
-class LIVELINKSEQUENCER_API ULiveLinkSequencerSettings : public UDeveloperSettings
+UCLASS(MinimalAPI, config=Game)
+class ULiveLinkSequencerSettings : public UDeveloperSettings
 {
 	GENERATED_BODY()
 
 public:
 	//~ Begin UDeveloperSettings interface
-	virtual FName GetCategoryName() const;
+	UE_API virtual FName GetCategoryName() const;
 #if WITH_EDITOR
-	virtual FText GetSectionText() const override;
-	virtual FName GetSectionName() const override;
+	UE_API virtual FText GetSectionText() const override;
+	UE_API virtual FName GetSectionName() const override;
 #endif
 	//~ End UDeveloperSettings interface
 
@@ -35,3 +37,5 @@ public:
 	UPROPERTY(config, EditAnywhere, Category = "LiveLink", meta = (AllowAbstract = "false"))
 	TMap<TSubclassOf<ULiveLinkControllerBase>, TSubclassOf<UMovieSceneLiveLinkControllerTrackRecorder>> DefaultTrackRecordersForController;
 };
+
+#undef UE_API
diff --git a/Engine/Plugins/Animation/LiveLink/Source/LiveLinkSequencer/Public/MovieSceneLiveLinkControllerMapTrackRecorder.h b/Engine/Plugins/Animation/LiveLink/Source/LiveLinkSequencer/Public/MovieSceneLiveLinkControllerMapTrackRecorder.h
--- a/Engine/Plugins/Animation/LiveLink/Source/LiveLinkSequencer/Public/MovieSceneLiveLinkControllerMapTrackRecorder.h
+++ b/Engine/Plugins/Animation/LiveLink/Source/LiveLinkSequencer/Public/MovieSceneLiveLinkControllerMapTrackRecorder.h
@@ -8,20 +8,22 @@
 
 #include "MovieSceneLiveLinkControllerMapTrackRecorder.generated.h"
 
+#define UE_API LIVELINKSEQUENCER_API
+
 class ULiveLinkControllerBase;
 class UMovieSceneLiveLinkControllerTrackRecorder;
 
 /** Movie Scene track recorder factory for the LiveLink Component's Controller Map */
-class LIVELINKSEQUENCER_API FMovieSceneLiveLinkControllerMapTrackRecorderFactory : public IMovieSceneTrackRecorderFactory
+class FMovieSceneLiveLinkControllerMapTrackRecorderFactory : public IMovieSceneTrackRecorderFactory
 {
 public:
 	virtual ~FMovieSceneLiveLinkControllerMapTrackRecorderFactory() = default;
 
 	//~ Begin IMovieSceneTrackRecorderFactory interface
 	virtual bool CanRecordObject(class UObject* InObjectToRecord) const override { return false; }
-	virtual bool CanRecordProperty(class UObject* InObjectToRecord, class FProperty* InPropertyToRecord) const override;
+	UE_API virtual bool CanRecordProperty(class UObject* InObjectToRecord, class FProperty* InPropertyToRecord) const override;
 
-	virtual UMovieSceneTrackRecorder* CreateTrackRecorderForObject() const override;
+	UE_API virtual UMovieSceneTrackRecorder* CreateTrackRecorderForObject() const override;
 	virtual UMovieSceneTrackRecorder* CreateTrackRecorderForProperty(UObject* InObjectToRecord, const FName& InPropertyToRecord) const override { return nullptr; }
 
 	virtual FText GetDisplayName() const override { return NSLOCTEXT("MovieSceneLiveLinkControllerTrackRecorderFactory", "DisplayName", "LiveLink ControllerMap Track Recorder"); }
@@ -29,27 +31,29 @@ public:
 };
 
 /** Movie Scene track recorder for LiveLink Component's Controller Map */
-UCLASS(BlueprintType)
-class LIVELINKSEQUENCER_API UMovieSceneLiveLinkControllerMapTrackRecorder : public UMovieSceneTrackRecorder
+UCLASS(MinimalAPI, BlueprintType)
+class UMovieSceneLiveLinkControllerMapTrackRecorder : public UMovieSceneTrackRecorder
 {
 	GENERATED_BODY()
 public:
 	virtual ~UMovieSceneLiveLinkControllerMapTrackRecorder() = default;
 
 	//~ Begin UMovieSceneTrackRecorder Interface
-	virtual void CreateTrackImpl() override;
-	virtual void SetSectionStartTimecodeImpl(const FTimecode& InSectionStartTimecode, const FFrameNumber& InSectionFirstFrame) override;
-	virtual void RecordSampleImpl(const FQualifiedFrameTime& CurrentTime) override;
-	virtual void StopRecordingImpl() override;
-	virtual void FinalizeTrackImpl() override;
+	UE_API virtual void CreateTrackImpl() override;
+	UE_API virtual void SetSectionStartTimecodeImpl(const FTimecode& InSectionStartTimecode, const FFrameNumber& InSectionFirstFrame) override;
+	UE_API virtual void RecordSampleImpl(const FQualifiedFrameTime& CurrentTime) override;
+	UE_API virtual void StopRecordingImpl() override;
+	UE_API virtual void FinalizeTrackImpl() override;
 	//~ End UMovieSceneTrackRecorder Interface
 
 private:
 	/** Return the desired track recorder class that supports the input LiveLink controller class */
-	TSubclassOf<UMovieSceneLiveLinkControllerTrackRecorder> GetRecorderClassForLiveLinkController(const TSubclassOf<ULiveLinkControllerBase> ControllerClass);
+	UE_API TSubclassOf<UMovieSceneLiveLinkControllerTrackRecorder> GetRecorderClassForLiveLinkController(const TSubclassOf<ULiveLinkControllerBase> ControllerClass);
 
 private:
 	/** Array of track recorders that will record each of the LiveLink controller's in the Controller Map */
 	UPROPERTY(Transient)
 	TArray<TObjectPtr<UMovieSceneLiveLinkControllerTrackRecorder>> ControllerRecorders;
 };
+
+#undef UE_API
diff --git a/Engine/Plugins/Animation/LiveLink/Source/LiveLinkSequencer/Public/MovieSceneLiveLinkControllerTrackRecorder.h b/Engine/Plugins/Animation/LiveLink/Source/LiveLinkSequencer/Public/MovieSceneLiveLinkControllerTrackRecorder.h
--- a/Engine/Plugins/Animation/LiveLink/Source/LiveLinkSequencer/Public/MovieSceneLiveLinkControllerTrackRecorder.h
+++ b/Engine/Plugins/Animation/LiveLink/Source/LiveLinkSequencer/Public/MovieSceneLiveLinkControllerTrackRecorder.h
@@ -9,8 +9,8 @@
 class ULiveLinkControllerBase;
 
 /** Abstract based for movie scene track recorders that can record LiveLink Controllers */
-UCLASS(Abstract)
-class LIVELINKSEQUENCER_API UMovieSceneLiveLinkControllerTrackRecorder : public UMovieSceneTrackRecorder
+UCLASS(MinimalAPI, Abstract)
+class UMovieSceneLiveLinkControllerTrackRecorder : public UMovieSceneTrackRecorder
 {
 	GENERATED_BODY()
 
diff --git a/Engine/Plugins/Animation/LiveLink/Source/LiveLinkSequencer/Public/MovieSceneLiveLinkTrackRecorder.h b/Engine/Plugins/Animation/LiveLink/Source/LiveLinkSequencer/Public/MovieSceneLiveLinkTrackRecorder.h
--- a/Engine/Plugins/Animation/LiveLink/Source/LiveLinkSequencer/Public/MovieSceneLiveLinkTrackRecorder.h
+++ b/Engine/Plugins/Animation/LiveLink/Source/LiveLinkSequencer/Public/MovieSceneLiveLinkTrackRecorder.h
@@ -8,6 +8,8 @@
 
 #include "MovieSceneLiveLinkTrackRecorder.generated.h"
 
+#define UE_API LIVELINKSEQUENCER_API
+
 class UMovieSceneTrackRecorderSettings;
 
 
@@ -21,42 +23,42 @@ class UMovieSceneLiveLinkTrack;
 class ULevelSequence;
 
 
-UCLASS(BlueprintType)
-class LIVELINKSEQUENCER_API UMovieSceneLiveLinkTrackRecorder : public UMovieSceneTrackRecorder
+UCLASS(MinimalAPI, BlueprintType)
+class UMovieSceneLiveLinkTrackRecorder : public UMovieSceneTrackRecorder
 {
 	GENERATED_BODY()
 public:
 	virtual ~UMovieSceneLiveLinkTrackRecorder() = default;
 
 	// UMovieSceneTrackRecorder Interface
-	virtual void RecordSampleImpl(const FQualifiedFrameTime& CurrentTime) override;
-	virtual void FinalizeTrackImpl() override;
-	virtual void SetSectionStartTimecodeImpl(const FTimecode& InSectionStartTimecode, const FFrameNumber& InSectionFirstFrame) override;
+	UE_API virtual void RecordSampleImpl(const FQualifiedFrameTime& CurrentTime) override;
+	UE_API virtual void FinalizeTrackImpl() override;
+	UE_API virtual void SetSectionStartTimecodeImpl(const FTimecode& InSectionStartTimecode, const FFrameNumber& InSectionFirstFrame) override;
 	virtual UMovieSceneSection* GetMovieSceneSection() const override { return Cast<UMovieSceneSection>(MovieSceneSection.Get()); }
-	virtual void StopRecordingImpl() override;
+	UE_API virtual void StopRecordingImpl() override;
 	virtual void SetSavedRecordingDirectory(const FString& InDirectory)
 	{
 		Directory = InDirectory;
 	}
-	virtual bool LoadRecordedFile(const FString& InFileName, UMovieScene *InMovieScene, TMap<FGuid, AActor*>& ActorGuidToActorMap, TFunction<void()> InCompletionCallback) override;
+	UE_API virtual bool LoadRecordedFile(const FString& InFileName, UMovieScene *InMovieScene, TMap<FGuid, AActor*>& ActorGuidToActorMap, TFunction<void()> InCompletionCallback) override;
 
 
 public:
 	//we don't call UMovieSceneTrackRecorder::CreateTrack or CreateTrackImpl since that expects an  ObjectToRecord and a GUID which isn't needed.
-	void CreateTrack(UMovieScene* InMovieScene, const FName& InSubjectName, bool bInSaveSubjectSettings, bool bInAlwaysUseTimecode, bool bDiscardSamplesBeforeStart, UMovieSceneTrackRecorderSettings* InSettingsObject);
-	void AddContentsToFolder(UMovieSceneFolder* InFolder);
+	UE_API void CreateTrack(UMovieScene* InMovieScene, const FName& InSubjectName, bool bInSaveSubjectSettings, bool bInAlwaysUseTimecode, bool bDiscardSamplesBeforeStart, UMovieSceneTrackRecorderSettings* InSettingsObject);
+	UE_API void AddContentsToFolder(UMovieSceneFolder* InFolder);
 	void SetReduceKeys(bool bInReduce) { bReduceKeys = bInReduce; }
 
 	/** Tell the recorder to write the take track using the recorded times. */
-	void ProcessRecordedTimes(ULevelSequence *InLevelSequence);
+	UE_API void ProcessRecordedTimes(ULevelSequence *InLevelSequence);
 private:
 
-	UMovieSceneLiveLinkTrack* DoesLiveLinkTrackExist(const FName& TrackName, const TSubclassOf<ULiveLinkRole>& InTrackRole);
+	UE_API UMovieSceneLiveLinkTrack* DoesLiveLinkTrackExist(const FName& TrackName, const TSubclassOf<ULiveLinkRole>& InTrackRole);
 
-	void CreateTracks();
+	UE_API void CreateTracks();
 
-	void OnStaticDataReceived(FLiveLinkSubjectKey InSubjectKey, TSubclassOf<ULiveLinkRole> InSubjectRole, const FLiveLinkStaticDataStruct& InStaticData);
-	void OnFrameDataReceived(FLiveLinkSubjectKey InSubjectKey, TSubclassOf<ULiveLinkRole> InSubjectRole, const FLiveLinkFrameDataStruct& InFrameData);
+	UE_API void OnStaticDataReceived(FLiveLinkSubjectKey InSubjectKey, TSubclassOf<ULiveLinkRole> InSubjectRole, const FLiveLinkStaticDataStruct& InStaticData);
+	UE_API void OnFrameDataReceived(FLiveLinkSubjectKey InSubjectKey, TSubclassOf<ULiveLinkRole> InSubjectRole, const FLiveLinkFrameDataStruct& InFrameData);
 private:
 
 	/** Name of Subject To Record */
@@ -111,3 +113,5 @@ private:
 	/** Pairs of recorded times with the corresponding timecode.*/
 	TArray<TPair<FQualifiedFrameTime, FQualifiedFrameTime>> RecordedTimes;
 };
+
+#undef UE_API
diff --git a/Engine/Plugins/Animation/LiveLinkHub/Source/LiveLinkHub/Private/Settings/LiveLinkHubSettings.h b/Engine/Plugins/Animation/LiveLinkHub/Source/LiveLinkHub/Private/Settings/LiveLinkHubSettings.h
--- a/Engine/Plugins/Animation/LiveLinkHub/Source/LiveLinkHub/Private/Settings/LiveLinkHubSettings.h
+++ b/Engine/Plugins/Animation/LiveLinkHub/Source/LiveLinkHub/Private/Settings/LiveLinkHubSettings.h
@@ -10,28 +10,30 @@
 
 #include "LiveLinkHubSettings.generated.h"
 
+#define UE_API LIVELINKHUB_API
+
 /**
  * Settings for LiveLinkHub.
  */
-UCLASS(config=Engine, defaultconfig)
-class LIVELINKHUB_API ULiveLinkHubSettings : public UObject
+UCLASS(MinimalAPI, config=Engine, defaultconfig)
+class ULiveLinkHubSettings : public UObject
 {
 	GENERATED_BODY()
 
 public:
-	ULiveLinkHubSettings();
+	UE_API ULiveLinkHubSettings();
 	
-	virtual void PostInitProperties() override;
+	UE_API virtual void PostInitProperties() override;
 	
 	/** Parse templates and set example output fields. */
-	void CalculateExampleOutput();
+	UE_API void CalculateExampleOutput();
 
 	/** Get the naming tokens for Live Link Hub. */
-	TObjectPtr<ULiveLinkHubNamingTokens> GetNamingTokens() const;
+	UE_API TObjectPtr<ULiveLinkHubNamingTokens> GetNamingTokens() const;
 	
 protected:
 	//~ Begin UObject interface
-	virtual void PostEditChangeProperty(FPropertyChangedEvent& PropertyChangedEvent) override;
+	UE_API virtual void PostEditChangeProperty(FPropertyChangedEvent& PropertyChangedEvent) override;
 	//~ End UObject interface
 	
 public:
@@ -119,8 +121,8 @@ private:
 /**
  * User Settings for LiveLinkHub.
  */
-UCLASS(config = EditorPerProjectUserSettings)
-class LIVELINKHUB_API ULiveLinkHubUserSettings : public UObject
+UCLASS(MinimalAPI, config = EditorPerProjectUserSettings)
+class ULiveLinkHubUserSettings : public UObject
 {
 	GENERATED_BODY()
 
@@ -133,3 +135,5 @@ public:
 	UPROPERTY(config)
 	FString LastConfigDirectory;
 };
+
+#undef UE_API
diff --git a/Engine/Plugins/Animation/LiveLinkHub/Source/LiveLinkHub/Public/ILiveLinkHubModule.h b/Engine/Plugins/Animation/LiveLinkHub/Source/LiveLinkHub/Public/ILiveLinkHubModule.h
--- a/Engine/Plugins/Animation/LiveLinkHub/Source/LiveLinkHub/Public/ILiveLinkHubModule.h
+++ b/Engine/Plugins/Animation/LiveLinkHub/Source/LiveLinkHub/Public/ILiveLinkHubModule.h
@@ -5,7 +5,7 @@
 #include "Containers/UnrealString.h"
 #include "Modules/ModuleInterface.h"
 
-class LIVELINKHUB_API ILiveLinkHubModule : public IModuleInterface
+class ILiveLinkHubModule : public IModuleInterface
 {
 public:
 	UE_DEPRECATED(5.6, "There's no need to call this method anymore, it's now handled by StartupModule and ShutdownModule.")
diff --git a/Engine/Plugins/Animation/LiveLinkHub/Source/LiveLinkHub/Public/LiveLinkHubApplicationBase.h b/Engine/Plugins/Animation/LiveLinkHub/Source/LiveLinkHub/Public/LiveLinkHubApplicationBase.h
--- a/Engine/Plugins/Animation/LiveLinkHub/Source/LiveLinkHub/Public/LiveLinkHubApplicationBase.h
+++ b/Engine/Plugins/Animation/LiveLinkHub/Source/LiveLinkHub/Public/LiveLinkHubApplicationBase.h
@@ -28,10 +28,10 @@ public:
 	virtual ~FLiveLinkHubApplicationBase() {}
 
 	/** Add a LiveLinkHub application mode. */
-	LIVELINKHUB_API virtual void AddLiveLinkHubApplicationMode(FName ModeName, TSharedRef<FLiveLinkHubApplicationMode> Mode) = 0;
+	virtual void AddLiveLinkHubApplicationMode(FName ModeName, TSharedRef<FLiveLinkHubApplicationMode> Mode) = 0;
 
 	/** Push tab factories for the given mode. */
-	LIVELINKHUB_API virtual void PushTabFactories(const class FWorkflowAllowedTabSet& FactorySetToPush, TSharedPtr<FLiveLinkHubApplicationMode> ApplicationMode) = 0;
+	virtual void PushTabFactories(const class FWorkflowAllowedTabSet& FactorySetToPush, TSharedPtr<FLiveLinkHubApplicationMode> ApplicationMode) = 0;
 };
 
 
diff --git a/Engine/Plugins/Animation/LiveLinkHub/Source/LiveLinkHub/Public/LiveLinkHubApplicationMode.h b/Engine/Plugins/Animation/LiveLinkHub/Source/LiveLinkHub/Public/LiveLinkHubApplicationMode.h
--- a/Engine/Plugins/Animation/LiveLinkHub/Source/LiveLinkHub/Public/LiveLinkHubApplicationMode.h
+++ b/Engine/Plugins/Animation/LiveLinkHub/Source/LiveLinkHub/Public/LiveLinkHubApplicationMode.h
@@ -8,6 +8,8 @@
 #include "WorkflowOrientedApp/ApplicationMode.h"
 #include "WorkflowOrientedApp/WorkflowTabManager.h"
 
+#define UE_API LIVELINKHUB_API
+
 struct FSlateIcon;
 class FToolBarBuilder;
 class FWorkflowCentricApplication;
@@ -15,25 +17,25 @@ class FLiveLinkHubApplicationBase;
 class SWidget;
 
 /** Factory used by external plugins to register custom LiveLinkHub layouts. */
-class LIVELINKHUB_API ILiveLinkHubApplicationModeFactory : public IModularFeature
+class ILiveLinkHubApplicationModeFactory : public IModularFeature
 {
 public:
 	/** Name of the modular feature. */
-	static const FName ModularFeatureName;
+	static UE_API const FName ModularFeatureName;
 
 
 	UE_DEPRECATED(5.6, "Use CreateLiveLinkHubAppMode.")
-	virtual TSharedRef<class FLiveLinkHubApplicationMode> CreateApplicationMode(TSharedPtr<FWorkflowCentricApplication> InApp);
+	UE_API virtual TSharedRef<class FLiveLinkHubApplicationMode> CreateApplicationMode(TSharedPtr<FWorkflowCentricApplication> InApp);
 
 	/** Instantiate an application mode so LiveLinkHub can register it and display it in its Layout Selector. */
 	virtual TSharedRef<class FLiveLinkHubApplicationMode> CreateLiveLinkHubAppMode(TSharedPtr<FLiveLinkHubApplicationBase> InApp) = 0;
 };
 
-class LIVELINKHUB_API FLiveLinkHubApplicationMode : public FApplicationMode
+class FLiveLinkHubApplicationMode : public FApplicationMode
 {
 public:
 	/** Name of the File menu extension point. Can be used to insert menus before/after in the main menu bar. */
-	static const FName FileMenuExtensionPoint;
+	static UE_API const FName FileMenuExtensionPoint;
 
 	/** Get the icon for this mode. */
 	virtual FSlateIcon GetModeIcon() const { return FSlateIcon(); }
@@ -42,7 +44,7 @@ public:
 	const FText& GetDisplayName() const { return DisplayName; }
 
 	/** Method used to gather the toolbar widgets that should be visible when this mode is active. */
-	TArray<TSharedRef<SWidget>> GetToolbarWidgets();
+	UE_API TArray<TSharedRef<SWidget>> GetToolbarWidgets();
 
 	/** Get the name of this layout. */
 	FName GetLayoutName()
@@ -76,11 +78,11 @@ public:
 	virtual TArray<TSharedRef<SWidget>> GetStatusBarWidgets_Impl() { return {}; }
 
 protected:
-	FLiveLinkHubApplicationMode(FName InApplicationMode, FText InDisplayName, TSharedPtr<FLiveLinkHubApplicationBase> InApp);
+	UE_API FLiveLinkHubApplicationMode(FName InApplicationMode, FText InDisplayName, TSharedPtr<FLiveLinkHubApplicationBase> InApp);
 
 	//~ Begin FApplicationMode interface
-	virtual void PostActivateMode() override;
-	virtual void PreDeactivateMode() override;
+	UE_API virtual void PostActivateMode() override;
+	UE_API virtual void PreDeactivateMode() override;
 	virtual void RegisterTabFactories(TSharedPtr<FTabManager> InTabManager) override
 	{
 		if (!bRegisteredTabFactories)
@@ -117,4 +119,6 @@ private:
 
 	/** Keeps track of whether or not the tab factories were registered, so we don't end up with duplicate tab spawners for this mode. */
 	bool bRegisteredTabFactories = false;
-};
\ No newline at end of file
+};
+
+#undef UE_API
diff --git a/Engine/Plugins/Animation/LiveLinkHub/Source/LiveLinkHub/Public/LiveLinkHubSessionExtraData.h b/Engine/Plugins/Animation/LiveLinkHub/Source/LiveLinkHub/Public/LiveLinkHubSessionExtraData.h
--- a/Engine/Plugins/Animation/LiveLinkHub/Source/LiveLinkHub/Public/LiveLinkHubSessionExtraData.h
+++ b/Engine/Plugins/Animation/LiveLinkHub/Source/LiveLinkHub/Public/LiveLinkHubSessionExtraData.h
@@ -8,8 +8,8 @@
 
 
 /** Derive from this class to define additional UProperties that participate in session serialization. */
-UCLASS(Abstract)
-class LIVELINKHUB_API ULiveLinkHubSessionExtraData : public UObject
+UCLASS(MinimalAPI, Abstract)
+class ULiveLinkHubSessionExtraData : public UObject
 {
 	GENERATED_BODY()
 };
diff --git a/Engine/Plugins/Animation/LiveLinkHub/Source/LiveLinkHubEditor/Private/LiveLinkHubEditorSettings.h b/Engine/Plugins/Animation/LiveLinkHub/Source/LiveLinkHubEditor/Private/LiveLinkHubEditorSettings.h
--- a/Engine/Plugins/Animation/LiveLinkHub/Source/LiveLinkHubEditor/Private/LiveLinkHubEditorSettings.h
+++ b/Engine/Plugins/Animation/LiveLinkHub/Source/LiveLinkHubEditor/Private/LiveLinkHubEditorSettings.h
@@ -9,8 +9,8 @@
 /**
  * Settings for LiveLinkHub in editor.
  */
-UCLASS(config = Editor, defaultconfig)
-class LIVELINKHUBEDITOR_API ULiveLinkHubEditorSettings : public UObject
+UCLASS(MinimalAPI, config = Editor, defaultconfig)
+class ULiveLinkHubEditorSettings : public UObject
 {
 public:
 	GENERATED_BODY()
@@ -30,4 +30,4 @@ public:
 	/** The version of LiveLinkHub that should be used. If empty, any version will be allowed. */
 	UPROPERTY(config)
 	FString LiveLinkHubTargetVersion;
-};
\ No newline at end of file
+};
diff --git a/Engine/Plugins/Animation/MLDeformer/MLDeformerFramework/Source/MLDeformerFrameworkEditor/Private/SMLDeformerDebugSelectionWidget.h b/Engine/Plugins/Animation/MLDeformer/MLDeformerFramework/Source/MLDeformerFrameworkEditor/Private/SMLDeformerDebugSelectionWidget.h
--- a/Engine/Plugins/Animation/MLDeformer/MLDeformerFramework/Source/MLDeformerFrameworkEditor/Private/SMLDeformerDebugSelectionWidget.h
+++ b/Engine/Plugins/Animation/MLDeformer/MLDeformerFramework/Source/MLDeformerFrameworkEditor/Private/SMLDeformerDebugSelectionWidget.h
@@ -8,6 +8,8 @@
 #include "Widgets/Input/SComboBox.h"
 #include "Containers/UnrealString.h"
 
+#define UE_API MLDEFORMERFRAMEWORKEDITOR_API
+
 class UWorld;
 class AActor;
 
@@ -33,7 +35,7 @@ namespace UE::MLDeformer
 	/**
 	 * 
 	 */
-	class MLDEFORMERFRAMEWORKEDITOR_API SMLDeformerDebugSelectionWidget
+	class SMLDeformerDebugSelectionWidget
 		: public SCompoundWidget
 	{
 	public:
@@ -41,24 +43,24 @@ namespace UE::MLDeformer
 		SLATE_ARGUMENT(FMLDeformerEditorToolkit*, MLDeformerEditor)
 		SLATE_END_ARGS()
 
-		void Construct(const FArguments& InArgs);
+		UE_API void Construct(const FArguments& InArgs);
 
 		/** Refresh the list of debuggable actors in the combobox. */
-		void Refresh();
+		UE_API void Refresh();
 
 		/** Get the current actor we're debugging, or a nullptr when debugging is disabled. */
-		AActor* GetDebugActor() const;
+		UE_API AActor* GetDebugActor() const;
 
 		const TArray<TSharedPtr<FMLDeformerDebugActor>>& GetActors() const { return Actors; }
 
 	private:
-		TSharedRef<SWidget> OnGenerateActorComboBoxItemWidget(TSharedPtr<FMLDeformerDebugActor> Item);
-		TArray<TSharedPtr<FMLDeformerDebugActor>> GetDebugActorsForWorld(TObjectPtr<UWorld> World) const;
-		FText GetComboBoxText() const;
-		void OnActorSelectionChanged(TSharedPtr<FMLDeformerDebugActor> Item, ESelectInfo::Type SelectInfo);
-		void RefreshActorList();
-		void UpdateDebugActorSelectedFlags();
-		bool IsDebuggingDisabled() const;
+		UE_API TSharedRef<SWidget> OnGenerateActorComboBoxItemWidget(TSharedPtr<FMLDeformerDebugActor> Item);
+		UE_API TArray<TSharedPtr<FMLDeformerDebugActor>> GetDebugActorsForWorld(TObjectPtr<UWorld> World) const;
+		UE_API FText GetComboBoxText() const;
+		UE_API void OnActorSelectionChanged(TSharedPtr<FMLDeformerDebugActor> Item, ESelectInfo::Type SelectInfo);
+		UE_API void RefreshActorList();
+		UE_API void UpdateDebugActorSelectedFlags();
+		UE_API bool IsDebuggingDisabled() const;
 
 	private:
 		/** The combobox that contains the actor names. */
@@ -78,3 +80,5 @@ namespace UE::MLDeformer
 	};
 
 }	// namespace UE::MLDeformer
+
+#undef UE_API
diff --git a/Engine/Plugins/Animation/MLDeformer/MLDeformerFramework/Source/MLDeformerFrameworkEditor/Public/MLDeformerApplicationMode.h b/Engine/Plugins/Animation/MLDeformer/MLDeformerFramework/Source/MLDeformerFrameworkEditor/Public/MLDeformerApplicationMode.h
--- a/Engine/Plugins/Animation/MLDeformer/MLDeformerFramework/Source/MLDeformerFrameworkEditor/Public/MLDeformerApplicationMode.h
+++ b/Engine/Plugins/Animation/MLDeformer/MLDeformerFramework/Source/MLDeformerFrameworkEditor/Public/MLDeformerApplicationMode.h
@@ -5,6 +5,8 @@
 #include "WorkflowOrientedApp/ApplicationMode.h"
 #include "WorkflowOrientedApp/WorkflowTabManager.h"
 
+#define UE_API MLDEFORMERFRAMEWORKEDITOR_API
+
 class FWorkflowCentricApplication;
 class IPersonaPreviewScene;
 
@@ -16,17 +18,17 @@ namespace UE::MLDeformer
 	 * The application mode for the ML Deformer editor.
 	 * This defines the layout of the UI. It basically spawns all the tabs, such as the viewport, details panels, etc.
 	 */
-	class MLDEFORMERFRAMEWORKEDITOR_API FMLDeformerApplicationMode
+	class FMLDeformerApplicationMode
 		: public FApplicationMode
 	{
 	public:
 		/** The name of this mode. */
-		static FName ModeName;
+		static UE_API FName ModeName;
 
-		FMLDeformerApplicationMode(TSharedRef<FWorkflowCentricApplication> InHostingApp, TSharedRef<IPersonaPreviewScene> InPreviewScene);
+		UE_API FMLDeformerApplicationMode(TSharedRef<FWorkflowCentricApplication> InHostingApp, TSharedRef<IPersonaPreviewScene> InPreviewScene);
 
 		// FApplicationMode overrides.
-		virtual void RegisterTabFactories(TSharedPtr<FTabManager> InTabManager) override;
+		UE_API virtual void RegisterTabFactories(TSharedPtr<FTabManager> InTabManager) override;
 		// ~END FApplicationMode overrides.
 
 	protected:
@@ -37,4 +39,6 @@ namespace UE::MLDeformer
 		FWorkflowAllowedTabSet TabFactories;
 	};
 
-}	// namespace UE::MLDeformer
\ No newline at end of file
+}	// namespace UE::MLDeformer
+
+#undef UE_API
diff --git a/Engine/Plugins/Animation/MLDeformer/MLDeformerFramework/Source/MLDeformerFrameworkEditor/Public/MLDeformerAssetFactory.h b/Engine/Plugins/Animation/MLDeformer/MLDeformerFramework/Source/MLDeformerFrameworkEditor/Public/MLDeformerAssetFactory.h
--- a/Engine/Plugins/Animation/MLDeformer/MLDeformerFramework/Source/MLDeformerFrameworkEditor/Public/MLDeformerAssetFactory.h
+++ b/Engine/Plugins/Animation/MLDeformer/MLDeformerFramework/Source/MLDeformerFrameworkEditor/Public/MLDeformerAssetFactory.h
@@ -6,27 +6,31 @@
 #include "Factories/Factory.h"
 #include "MLDeformerAssetFactory.generated.h"
 
+#define UE_API MLDEFORMERFRAMEWORKEDITOR_API
+
 /**
  * The factory for the ML Deformer asset type.
  * This basically integrates the new asset type into the editor, so you can right click and create a new ML Deformer asset.
  */
-UCLASS(hidecategories=Object)
-class MLDEFORMERFRAMEWORKEDITOR_API UMLDeformerFactory
+UCLASS(MinimalAPI, hidecategories=Object)
+class UMLDeformerFactory
 	: public UFactory
 {
 	GENERATED_BODY()
 
 public:
-	UMLDeformerFactory();
+	UE_API UMLDeformerFactory();
 
 	// UFactory overrides.
-	virtual FText GetDisplayName() const override;
-	virtual uint32 GetMenuCategories() const override;
-	virtual FText GetToolTip() const override;
-	virtual FString GetDefaultNewAssetName() const override;
-	virtual UObject* FactoryCreateNew(UClass* Class, UObject* InParent, FName Name, EObjectFlags Flags, UObject* Context, FFeedbackContext* Warn) override;
-	virtual bool ConfigureProperties() override;
-	virtual bool ShouldShowInNewMenu() const override;
-	virtual const TArray<FText>& GetMenuCategorySubMenus() const override;
+	UE_API virtual FText GetDisplayName() const override;
+	UE_API virtual uint32 GetMenuCategories() const override;
+	UE_API virtual FText GetToolTip() const override;
+	UE_API virtual FString GetDefaultNewAssetName() const override;
+	UE_API virtual UObject* FactoryCreateNew(UClass* Class, UObject* InParent, FName Name, EObjectFlags Flags, UObject* Context, FFeedbackContext* Warn) override;
+	UE_API virtual bool ConfigureProperties() override;
+	UE_API virtual bool ShouldShowInNewMenu() const override;
+	UE_API virtual const TArray<FText>& GetMenuCategorySubMenus() const override;
 	// ~END UFactory overrides.
 };
+
+#undef UE_API
diff --git a/Engine/Plugins/Animation/MLDeformer/MLDeformerFramework/Source/MLDeformerFrameworkEditor/Public/MLDeformerCurveReferenceCustomization.h b/Engine/Plugins/Animation/MLDeformer/MLDeformerFramework/Source/MLDeformerFrameworkEditor/Public/MLDeformerCurveReferenceCustomization.h
--- a/Engine/Plugins/Animation/MLDeformer/MLDeformerFramework/Source/MLDeformerFrameworkEditor/Public/MLDeformerCurveReferenceCustomization.h
+++ b/Engine/Plugins/Animation/MLDeformer/MLDeformerFramework/Source/MLDeformerFrameworkEditor/Public/MLDeformerCurveReferenceCustomization.h
@@ -7,6 +7,8 @@
 #include "PropertyHandle.h"
 #include "DetailWidgetRow.h"
 
+#define UE_API MLDEFORMERFRAMEWORKEDITOR_API
+
 class USkeletalMesh;
 
 namespace UE::MLDeformer
@@ -14,21 +16,21 @@ namespace UE::MLDeformer
 	/**
 	 * The curve reference property detail customization.
 	 */
-	class MLDEFORMERFRAMEWORKEDITOR_API FMLDeformerCurveReferenceCustomization
+	class FMLDeformerCurveReferenceCustomization
 		: public IPropertyTypeCustomization
 	{
 	public:
-		static TSharedRef<IPropertyTypeCustomization> MakeInstance();
+		static UE_API TSharedRef<IPropertyTypeCustomization> MakeInstance();
 
 		// IPropertyTypeCustomization overrides.
-		virtual void CustomizeHeader(TSharedRef<IPropertyHandle> StructPropertyHandle, class FDetailWidgetRow& HeaderRow, IPropertyTypeCustomizationUtils& StructCustomizationUtils) override;
+		UE_API virtual void CustomizeHeader(TSharedRef<IPropertyHandle> StructPropertyHandle, class FDetailWidgetRow& HeaderRow, IPropertyTypeCustomizationUtils& StructCustomizationUtils) override;
 		virtual void CustomizeChildren(TSharedRef<IPropertyHandle> StructPropertyHandle, class IDetailChildrenBuilder& ChildBuilder, IPropertyTypeCustomizationUtils& StructCustomizationUtils) override {}
 		// ~END IPropertyTypeCustomization overrides.
 
 	protected:
-		void SetSkeleton(TSharedRef<IPropertyHandle> StructPropertyHandle);
-		virtual void SetPropertyHandle(TSharedRef<IPropertyHandle> StructPropertyHandle);
-		TSharedPtr<IPropertyHandle> FindStructMemberProperty(TSharedRef<IPropertyHandle> PropertyHandle, const FName& PropertyName);
+		UE_API void SetSkeleton(TSharedRef<IPropertyHandle> StructPropertyHandle);
+		UE_API virtual void SetPropertyHandle(TSharedRef<IPropertyHandle> StructPropertyHandle);
+		UE_API TSharedPtr<IPropertyHandle> FindStructMemberProperty(TSharedRef<IPropertyHandle> PropertyHandle, const FName& PropertyName);
 
 		/** Property to change after curve has been picked. */
 		TSharedPtr<IPropertyHandle> CurveNameProperty = nullptr;
@@ -37,8 +39,10 @@ namespace UE::MLDeformer
 		TObjectPtr<USkeletalMesh> SkeletalMesh = nullptr;
 
 	private:
-		virtual void OnCurveSelectionChanged(const FString& Name);
-		virtual FString OnGetSelectedCurve() const;
-		virtual USkeletalMesh* OnGetSkeletalMesh() const;
+		UE_API virtual void OnCurveSelectionChanged(const FString& Name);
+		UE_API virtual FString OnGetSelectedCurve() const;
+		UE_API virtual USkeletalMesh* OnGetSkeletalMesh() const;
 	};
 }	// namespace UE::MLDeformer
+
+#undef UE_API
diff --git a/Engine/Plugins/Animation/MLDeformer/MLDeformerFramework/Source/MLDeformerFrameworkEditor/Public/MLDeformerEditorActor.h b/Engine/Plugins/Animation/MLDeformer/MLDeformerFramework/Source/MLDeformerFrameworkEditor/Public/MLDeformerEditorActor.h
--- a/Engine/Plugins/Animation/MLDeformer/MLDeformerFramework/Source/MLDeformerFrameworkEditor/Public/MLDeformerEditorActor.h
+++ b/Engine/Plugins/Animation/MLDeformer/MLDeformerFramework/Source/MLDeformerFrameworkEditor/Public/MLDeformerEditorActor.h
@@ -5,6 +5,8 @@
 #include "CoreTypes.h"
 #include "GameFramework/Actor.h"
 
+#define UE_API MLDEFORMERFRAMEWORKEDITOR_API
+
 class UDebugSkelMeshComponent;
 class UTextRenderComponent;
 class UMLDeformerComponent;
@@ -44,13 +46,13 @@ namespace UE::MLDeformer
 	 * can create components on it.
 	 * It also contains some methods to control visibility and playback, among some other helper methods.
 	 */
-	class MLDEFORMERFRAMEWORKEDITOR_API FMLDeformerEditorActor
+	class FMLDeformerEditorActor
 	{
 	public:
 		/**
 		 * The construction settings, which can be used as info when creating new components.
 		 */
-		struct MLDEFORMERFRAMEWORKEDITOR_API FConstructSettings
+		struct FConstructSettings
 		{
 			/** The UE actor object that we should add components to. */
 			AActor* Actor = nullptr;
@@ -74,22 +76,22 @@ namespace UE::MLDeformer
 			bool bIsTrainingActor = false;
 		};
 
-		FMLDeformerEditorActor(const FConstructSettings& Settings);
-		virtual ~FMLDeformerEditorActor();
+		UE_API FMLDeformerEditorActor(const FConstructSettings& Settings);
+		UE_API virtual ~FMLDeformerEditorActor();
 
 		// Main methods you can override.
-		virtual void SetVisibility(bool bIsVisible);
-		virtual bool IsVisible() const;
-		virtual void SetPlayPosition(float TimeInSeconds, bool bAutoPause = true);
-		virtual float GetPlayPosition() const;
-		virtual void SetPlaySpeed(float PlaySpeed);
-		virtual float GetPlaySpeed();
-
-		virtual FBox GetBoundingBox() const;
-		virtual void Pause(bool bPaused);
-		virtual bool IsPlaying() const;
-		virtual bool IsGroundTruthActor() const;
-		virtual bool HasVisualMesh() const;
+		UE_API virtual void SetVisibility(bool bIsVisible);
+		UE_API virtual bool IsVisible() const;
+		UE_API virtual void SetPlayPosition(float TimeInSeconds, bool bAutoPause = true);
+		UE_API virtual float GetPlayPosition() const;
+		UE_API virtual void SetPlaySpeed(float PlaySpeed);
+		UE_API virtual float GetPlaySpeed();
+
+		UE_API virtual FBox GetBoundingBox() const;
+		UE_API virtual void Pause(bool bPaused);
+		UE_API virtual bool IsPlaying() const;
+		UE_API virtual bool IsGroundTruthActor() const;
+		UE_API virtual bool HasVisualMesh() const;
 
 		AActor* GetActor() const													{ return Actor; }
 		int32 GetTypeID() const														{ return TypeID; }
@@ -107,7 +109,7 @@ namespace UE::MLDeformer
 		void SetActorTypeInstanceIndex(int32 Index)									{ ActorTypeInstanceIndex = Index; }
 
 	protected:
-		UTextRenderComponent* CreateLabelComponent(AActor* InActor, FLinearColor Color, const FText& Text) const;
+		UE_API UTextRenderComponent* CreateLabelComponent(AActor* InActor, FLinearColor Color, const FText& Text) const;
 
 	protected:
 		/**
@@ -141,3 +143,5 @@ namespace UE::MLDeformer
 		bool bIsTrainingActor = true;
 	};
 } // namespace UE::MLDeformer
+
+#undef UE_API
diff --git a/Engine/Plugins/Animation/MLDeformer/MLDeformerFramework/Source/MLDeformerFrameworkEditor/Public/MLDeformerEditorMode.h b/Engine/Plugins/Animation/MLDeformer/MLDeformerFramework/Source/MLDeformerFrameworkEditor/Public/MLDeformerEditorMode.h
--- a/Engine/Plugins/Animation/MLDeformer/MLDeformerFramework/Source/MLDeformerFrameworkEditor/Public/MLDeformerEditorMode.h
+++ b/Engine/Plugins/Animation/MLDeformer/MLDeformerFramework/Source/MLDeformerFrameworkEditor/Public/MLDeformerEditorMode.h
@@ -5,6 +5,8 @@
 #include "CoreTypes.h"
 #include "IPersonaEditMode.h"
 
+#define UE_API MLDEFORMERFRAMEWORKEDITOR_API
+
 class UMLDeformerVizSettings;
 
 namespace UE::MLDeformer
@@ -16,25 +18,25 @@ namespace UE::MLDeformer
 	 * The ML Deformer Persona editor mode.
 	 * This handles ticking and rendering of the render viewport and camera focus.
 	 */
-	class MLDEFORMERFRAMEWORKEDITOR_API FMLDeformerEditorMode
+	class FMLDeformerEditorMode
 		: public IPersonaEditMode
 	{
 	public:
 		/** The name of the mode. */
-		static FName ModeName;
+		static UE_API FName ModeName;
 
 		void SetEditorToolkit(FMLDeformerEditorToolkit* InToolkit) { DeformerEditorToolkit = InToolkit; }
 
 		// IPersonaEditMode overrides.
-		virtual bool GetCameraTarget(FSphere& OutTarget) const override;
-		virtual class IPersonaPreviewScene& GetAnimPreviewScene() const override;
+		UE_API virtual bool GetCameraTarget(FSphere& OutTarget) const override;
+		UE_API virtual class IPersonaPreviewScene& GetAnimPreviewScene() const override;
 		virtual void GetOnScreenDebugInfo(TArray<FText>& OutDebugInfo) const override {}
 		// ~END IPersonaEditMode overrides.
 
 		// FEdMode overrides.
-		virtual void Tick(FEditorViewportClient* ViewportClient, float DeltaTime) override;
-		virtual void Render(const FSceneView* View, FViewport* Viewport, FPrimitiveDrawInterface* PDI) override;
-		virtual void DrawHUD(FEditorViewportClient* ViewportClient, FViewport* Viewport, const FSceneView* View, FCanvas* Canvas) override;
+		UE_API virtual void Tick(FEditorViewportClient* ViewportClient, float DeltaTime) override;
+		UE_API virtual void Render(const FSceneView* View, FViewport* Viewport, FPrimitiveDrawInterface* PDI) override;
+		UE_API virtual void DrawHUD(FEditorViewportClient* ViewportClient, FViewport* Viewport, const FSceneView* View, FCanvas* Canvas) override;
 		virtual bool IsCompatibleWith(FEditorModeID OtherModeID) const override { return true; }
 		virtual bool AllowWidgetMove() override { return false; }
 		virtual bool ShouldDrawWidget() const override { return false; }
@@ -46,3 +48,5 @@ namespace UE::MLDeformer
 		FMLDeformerEditorToolkit* DeformerEditorToolkit = nullptr;
 	};
 }	// namespace UE::MLDeformer
+
+#undef UE_API
