Reverted change from 43132987

diff --git a/Engine/Plugins/Animation/LiveLink/Source/LiveLinkEditor/Public/LiveLinkClientPanelToolbar.h b/Engine/Plugins/Animation/LiveLink/Source/LiveLinkEditor/Public/LiveLinkClientPanelToolbar.h
--- a/Engine/Plugins/Animation/LiveLink/Source/LiveLinkEditor/Public/LiveLinkClientPanelToolbar.h
+++ b/Engine/Plugins/Animation/LiveLink/Source/LiveLinkEditor/Public/LiveLinkClientPanelToolbar.h
@@ -5,8 +5,6 @@
 #include "Widgets/SCompoundWidget.h"
 #include "UObject/GCObject.h"
 
-#define UE_API LIVELINKEDITOR_API
-
 
 class ULiveLinkPreset;
 template <typename T> class TSubclassOf;
@@ -21,7 +19,7 @@ class STextEntryPopup;
 class ULiveLinkSourceFactory;
 struct FAssetData;
 
-class SLiveLinkClientPanelToolbar : public SCompoundWidget, public FGCObject
+class LIVELINKEDITOR_API SLiveLinkClientPanelToolbar : public SCompoundWidget, public FGCObject
 {
 	SLATE_BEGIN_ARGS(SLiveLinkClientPanelToolbar)
 		: _SourceButtonAlignment(HAlign_Fill)
@@ -41,38 +39,41 @@ class SLiveLinkClientPanelToolbar : public SCompoundWidget, public FGCObject
 	SLATE_END_ARGS()
 
 
-	UE_API void Construct(const FArguments& Args, FLiveLinkClient* InClient);
+	void Construct(const FArguments& Args, FLiveLinkClient* InClient);
 
 	//~ FEditorUndoClient interface
-	UE_API virtual void AddReferencedObjects(FReferenceCollector& Collector) override;
+	virtual void AddReferencedObjects(FReferenceCollector& Collector) override;
 	virtual FString GetReferencerName() const override
 	{
 		return TEXT("SLiveLinkClientPanelToolbar");
 	}
 
 private:
-	UE_API TSharedRef<SWidget> OnGenerateSourceMenu();
-	UE_API void RetrieveFactorySourcePanel(FMenuBuilder& MenuBuilder, int32 FactoryIndex);
-	UE_API void ExecuteCreateSource(int32 FactoryIndex);
-	UE_API void OnSourceCreated(TSharedPtr<ILiveLinkSource> NewSource, FString ConnectionString, TSubclassOf<ULiveLinkSourceFactory> Factory);
-
-	UE_API TSharedRef<SWidget> OnPresetGeneratePresetsMenu();
-	UE_API TSharedRef<SWidget> HandleSourceSelectionComboButton();
-	UE_API void HandleVirtualSourceSelection(const FGuid& InSourceGuid);
-	UE_API bool IsVirtualSourceSelected(const FGuid& InSourceGuid);
-	UE_API void OnSaveAsPreset();
-	UE_API void OnImportPreset(const FAssetData& InPreset);
-	UE_API FReply OnRevertChanges();
-	UE_API bool HasLoadedLiveLinkPreset() const;
-
-	UE_API void PopulateVirtualSubjectSourceCreationMenu(FMenuBuilder& InMenuBuilder);
-	UE_API FReply OnAddVirtualSubjectSource();
-	UE_API void AddVirtualSubject();
+	TSharedRef<SWidget> OnGenerateSourceMenu();
+	/** Open a modal window that hosts the source factory widget. */
+	void OpenCreateMenuWindow(int32 FactoryIndex);
+	void ExecuteCreateSource(int32 FactoryIndex);
+	void OnSourceCreated(TSharedPtr<ILiveLinkSource> NewSource, FString ConnectionString, TSubclassOf<ULiveLinkSourceFactory> Factory);
+
+	TSharedRef<SWidget> OnPresetGeneratePresetsMenu();
+	TSharedRef<SWidget> HandleSourceSelectionComboButton();
+	void HandleVirtualSourceSelection(const FGuid& InSourceGuid);
+	bool IsVirtualSourceSelected(const FGuid& InSourceGuid);
+	void OnSaveAsPreset();
+	void OnImportPreset(const FAssetData& InPreset);
+	FReply OnRevertChanges();
+	bool HasLoadedLiveLinkPreset() const;
+
+	void PopulateVirtualSubjectSourceCreationMenu(FMenuBuilder& InMenuBuilder);
+	FReply OnAddVirtualSubjectSource();
+	void AddVirtualSubject();
 
 
 private:
 	FLiveLinkClient* Client;
 
+	/** Source creation modal window. */
+	TSharedPtr<SWindow> SourceCreationWindow;
 	TWeakPtr<IMenu> AddSubjectMenu;
 	TWeakPtr<IMenu> VirtualSubjectMenu;
 	TWeakPtr<STextEntryPopup> VirtualSubjectPopup;
@@ -84,5 +85,3 @@ private:
 	/** Parent window override used for creating asset dialogs when running in LiveLinkHub. */
 	TSharedPtr<SWindow> ParentWindowOverride;
 };
-
-#undef UE_API
