LiveLinkHub - Fix crash when syncing a virtual subject to a real subject

diff --git a/Engine/Plugins/Animation/LiveLink/Source/LiveLink/Private/LiveLinkClient.cpp b/Engine/Plugins/Animation/LiveLink/Source/LiveLink/Private/LiveLinkClient.cpp
--- a/Engine/Plugins/Animation/LiveLink/Source/LiveLink/Private/LiveLinkClient.cpp
+++ b/Engine/Plugins/Animation/LiveLink/Source/LiveLink/Private/LiveLinkClient.cpp
@@ -262,7 +262,12 @@ void FLiveLinkClient::BuildThisTicksSubjectSnapshot()
 			if (SubjectItem.bEnabled)
 			{
 				VSubject->Update();
-				HandleSubjectRebroadcast(VSubject, VSubject->GetFrameData());
+				if (SourceItem.Setting->ParentSubject == FLiveLinkSubjectName())
+				{
+					// Don't rebroadcast if this is being synchronized to another subject.
+					HandleSubjectRebroadcast(VSubject, VSubject->GetFrameData());
+				}
+
 				EnabledSubjects.Add(SubjectItem.Key.SubjectName, SubjectItem.Key);
 			}
 			else
@@ -1112,7 +1117,21 @@ void FLiveLinkClient::PushSubjectFrameData_Internal(FPendingSubjectFrame&& Subje
 
 					// todo: Time offset evaluation
 					FLiveLinkSubjectFrameData ChildData;
-					if (SubjectItem.GetLiveSubject()->EvaluateFrameAtWorldTime(SubjectFrameData.FrameData.GetBaseData()->WorldTime.GetSourceTime(), SubjectItem.GetLinkSettings()->Role, ChildData))
+					bool bSuccess = true;
+					ILiveLinkSubject* LiveSubject = nullptr;
+
+					if (SubjectItem.GetLiveSubject())
+					{
+						LiveSubject = SubjectItem.GetLiveSubject();
+						bSuccess = SubjectItem.GetLiveSubject()->EvaluateFrameAtWorldTime(SubjectFrameData.FrameData.GetBaseData()->WorldTime.GetSourceTime(), SubjectItem.GetLinkSettings()->Role, ChildData);
+					}
+					else
+					{
+						LiveSubject = SubjectItem.GetVirtualSubject();
+						bSuccess = SubjectItem.GetVirtualSubject()->EvaluateFrame(SubjectItem.GetVirtualSubject()->GetRole(), ChildData);
+					}
+
+					if (bSuccess)
 					{
 						const FTimecode FrameTC = FTimecode::FromFrameNumber(FrameData.GetBaseData()->MetaData.SceneTime.Time.GetFrame(), FrameData.GetBaseData()->MetaData.SceneTime.Rate);
 						UE_LOG(LogLiveLink, Verbose, TEXT("LiveLinkHub Parent (%s) - Child '%s' adding frame with Timecode:[%s.%0.3f] - SourceTime: %0.4f, Offset: %0.6f, CorrectedTime: %0.4f"), *SubjectKey.SubjectName.ToString(), *SubjectItem.Key.SubjectName.ToString(), *FrameTC.ToString(), FrameData.GetBaseData()->MetaData.SceneTime.Time.GetSubFrame(), FrameData.GetBaseData()->WorldTime.GetSourceTime(), FrameData.GetBaseData()->WorldTime.GetOffset(), FrameData.GetBaseData()->WorldTime.GetOffsettedTime());
@@ -1120,7 +1139,7 @@ void FLiveLinkClient::PushSubjectFrameData_Internal(FPendingSubjectFrame&& Subje
 						ChildData.FrameData.GetBaseData()->MetaData.SceneTime = FrameData.GetBaseData()->MetaData.SceneTime;
 						ChildData.FrameData.GetBaseData()->MetaData.SceneTime.Rate = FrameData.GetBaseData()->MetaData.SceneTime.Rate;
 
-						HandleSubjectRebroadcast(SubjectItem.GetLiveSubject(), ChildData.FrameData);
+						HandleSubjectRebroadcast(LiveSubject, ChildData.FrameData);
 					}
 					else
 					{
