Fix unresolved symbol error caused by unnecessary api macro in SLiveLinkFilterSearchBox.h

diff --git a/Engine/Plugins/Animation/LiveLink/Source/LiveLinkEditor/Public/SLiveLinkFilterSearchBox.h b/Engine/Plugins/Animation/LiveLink/Source/LiveLinkEditor/Public/SLiveLinkFilterSearchBox.h
--- a/Engine/Plugins/Animation/LiveLink/Source/LiveLinkEditor/Public/SLiveLinkFilterSearchBox.h
+++ b/Engine/Plugins/Animation/LiveLink/Source/LiveLinkEditor/Public/SLiveLinkFilterSearchBox.h
@@ -52,7 +52,7 @@ public:
 	}
 
 	/** Trigger a deferred update of the filtered items. */
-	LIVELINKEDITOR_API void Update()
+	void Update()
 	{
 		bUpdateFilteredItems = true;
 	}
