LiveLink - Allow synchronizing Virtual Subject rebroadcast to one of its subject

Move the Virtual Subject rebroadcast logic to the BuildSubjectSnapshot method

diff --git a/Engine/Plugins/Animation/LiveLink/Source/LiveLink/Private/LiveLinkClient.cpp b/Engine/Plugins/Animation/LiveLink/Source/LiveLink/Private/LiveLinkClient.cpp
--- a/Engine/Plugins/Animation/LiveLink/Source/LiveLink/Private/LiveLinkClient.cpp
+++ b/Engine/Plugins/Animation/LiveLink/Source/LiveLink/Private/LiveLinkClient.cpp
@@ -53,7 +53,6 @@ static TAutoConsoleVariable<int32> CVarMaxNewFrameDataPerUpdate(
 	TEXT("Maximun number of new frame data that can be added in a single UE frame."),
 	ECVF_Default);
 
-
 FLiveLinkClient::FLiveLinkClient()
 	: FLiveLinkClient(FCoreDelegates::OnSamplingInput)
 {
@@ -87,10 +86,14 @@ void FLiveLinkClient::Tick()
 {
 	SCOPE_CYCLE_COUNTER(STAT_LiveLink_Client_Tick);
 
-	DoPendingWork();
+	TSet<FLiveLinkSubjectKey> ReceivedSubjects;
+
+	DoPendingWork(ReceivedSubjects);
 	CacheValues();
 	UpdateSources();
-	BuildThisTicksSubjectSnapshot();
+
+	TSet<FLiveLinkSubjectKey> UpdatedSubjects = BuildThisTicksSubjectSnapshot(ReceivedSubjects);
+	RebroadcastSubjects(UpdatedSubjects);
 
 	OnLiveLinkTickedDelegate.Broadcast();
 }
@@ -114,7 +117,7 @@ void FLiveLinkClient::Initialize()
 	bEnableParentSubjects = GConfig->GetBoolOrDefault(TEXT("LiveLink"), TEXT("bEnableParentSubjects"), false, GEngineIni);
 }
 
-void FLiveLinkClient::DoPendingWork()
+void FLiveLinkClient::DoPendingWork(TSet<FLiveLinkSubjectKey>& OutReceivedSubjects)
 {
 	FScopeLock PendingFramesLock(&PendingFramesCriticalSection);
 
@@ -127,6 +130,7 @@ void FLiveLinkClient::DoPendingWork()
 		// Add new Subject static data
 		for (FPendingSubjectStatic& SubjectStatic : SubjectStaticToPush)
 		{
+			OutReceivedSubjects.Add(SubjectStatic.SubjectKey);
 			PushSubjectStaticData_Internal(MoveTemp(SubjectStatic));
 		}
 		SubjectStaticToPush.Reset();
@@ -134,6 +138,7 @@ void FLiveLinkClient::DoPendingWork()
 		// Add new Subject frame data
 		for (FPendingSubjectFrame& SubjectFrame : SubjectFrameToPush)
 		{
+			OutReceivedSubjects.Add(SubjectFrame.SubjectKey);
 			PushSubjectFrameData_Internal(MoveTemp(SubjectFrame));
 		}
 		SubjectFrameToPush.Reset();
@@ -148,7 +153,7 @@ void FLiveLinkClient::UpdateSources()
 	});
 }
 
-void FLiveLinkClient::HandleSubjectRebroadcast(ILiveLinkSubject* InSubject, FLiveLinkSubjectFrameData&& SubjectFrameData)
+bool FLiveLinkClient::HandleSubjectRebroadcast(ILiveLinkSubject* InSubject, FLiveLinkSubjectFrameData&& SubjectFrameData)
 {
 	// Check the rebroadcast flag and act accordingly, creating the LiveLinkProvider and/or sending the static data if needed
 	if (InSubject->IsRebroadcasted())
@@ -180,6 +185,8 @@ void FLiveLinkClient::HandleSubjectRebroadcast(ILiveLinkSubject* InSubject, FLiv
 				}
 
 				RebroadcastLiveLinkProvider->UpdateSubjectFrameData(RebroadcastName, MoveTemp(SubjectFrameData.FrameData), ExtraAnnotations);
+
+				return true;
 			}
 			else
 			{
@@ -192,20 +199,19 @@ void FLiveLinkClient::HandleSubjectRebroadcast(ILiveLinkSubject* InSubject, FLiv
 		RemoveRebroadcastedSubject(InSubject->GetSubjectKey());
 		InSubject->SetStaticDataAsRebroadcasted(false);
 	}
+
+	return false;
 }
 
-void FLiveLinkClient::BuildThisTicksSubjectSnapshot()
+TSet<FLiveLinkSubjectKey> FLiveLinkClient::BuildThisTicksSubjectSnapshot(const TSet<FLiveLinkSubjectKey>& ReceivedSubjects)
 {
-	check(Collection);
-
 	EnabledSubjects.Reset();
 
 	TSet<FGuid> TaggedSources;
+	TSet<FLiveLinkSubjectKey> UpdatedSubjects;
 
 	// Update the Live Subject before the Virtual Subject
-	Collection->ForEachSubject([this, &TaggedSources](FLiveLinkCollectionSourceItem& SourceItem, FLiveLinkCollectionSubjectItem& SubjectItem)
-	{
-		if (FLiveLinkSubject* LiveSubject = SubjectItem.GetLiveSubject())
+	ForEachLiveSubject([this, &TaggedSources, &UpdatedSubjects](FLiveLinkSubject* LiveSubject, FLiveLinkCollectionSourceItem& SourceItem, FLiveLinkCollectionSubjectItem& SubjectItem)
 		{
 			if (SubjectItem.bEnabled)
 			{
@@ -214,13 +220,9 @@ void FLiveLinkClient::BuildThisTicksSubjectSnapshot()
 
 				LiveSubject->CacheSettings(SourceSettings, SubjectSettings);
 				LiveSubject->Update();
+				UpdatedSubjects.Add(SubjectItem.Key);
 				EnabledSubjects.Add(SubjectItem.Key.SubjectName, SubjectItem.Key);
 
-				if (LiveSubject->IsPaused())
-				{
-					HandleSubjectRebroadcast(LiveSubject, LiveSubject->GetFrameSnapshot().FrameData);
-				}
-
 				// Update Source FrameRate from first enabled subject with valid data.
 				if (LiveSubject->HasValidFrameSnapshot() && SubjectSettings->FrameRate.IsValid() && !TaggedSources.Contains(SubjectItem.Key.Source))
 				{
@@ -232,40 +234,39 @@ void FLiveLinkClient::BuildThisTicksSubjectSnapshot()
 			{
 				LiveSubject->ClearFrames();
 			}
-		}
-	});
+		});
 
-	Collection->ForEachSubject([this](FLiveLinkCollectionSourceItem& SourceItem, FLiveLinkCollectionSubjectItem& SubjectItem)
+	// Update Virtual Subjects
+	ForEachVirtualSubject([this, &ReceivedSubjects, &UpdatedSubjects](ULiveLinkVirtualSubject* VSubject, FLiveLinkCollectionSourceItem& SourceItem, FLiveLinkCollectionSubjectItem& SubjectItem)
 		{
-			if (FLiveLinkSubject* Subject = SubjectItem.GetLiveSubject())
+			if (SubjectItem.bEnabled)
 			{
-				if (SubjectItem.bEnabled && !Subject->IsPaused() && GetDefault<ULiveLinkSettings>()->bTransmitEvaluatedData)
+				const bool bSubjectIsSynchronized = SourceItem.Setting->ParentSubject != FLiveLinkSubjectName() || VSubject->SyncSubject != FLiveLinkSubjectName();
+
+				bool bSubjectNeedsUpdate = false;
+
+				if (bSubjectIsSynchronized)
 				{
-					FLiveLinkSubjectFrameData FrameData;
-					if (Subject->EvaluateFrame(Subject->GetRole(), FrameData))
+					for (const FLiveLinkSubjectKey& ReceivedSubject : ReceivedSubjects)
 					{
-						const FLiveLinkFrameIdentifier FrameId = FrameData.FrameData.GetBaseData()->FrameId;
-						if (Subject->GetLastRebroadcastedFrameId() != FrameId)
+						// We can either synchronize at the source or subject level to another subject.
+						if (SourceItem.Setting->ParentSubject == ReceivedSubject.SubjectName
+							|| VSubject->SyncSubject == ReceivedSubject.SubjectName)
 						{
-							HandleSubjectRebroadcast(SubjectItem.GetLiveSubject(), MoveTemp(FrameData));
-							Subject->SetLastRebroadcastedFrameId(FrameId);
+							bSubjectNeedsUpdate = true;
+							break;
 						}
 					}
 				}
-			}
-		});
+				else
+				{
+					bSubjectNeedsUpdate = true;
+				}
 
-	Collection->ForEachSubject([this](FLiveLinkCollectionSourceItem& SourceItem, FLiveLinkCollectionSubjectItem& SubjectItem)
-	{
-		if (ULiveLinkVirtualSubject* VSubject = SubjectItem.GetVirtualSubject())
-		{
-			if (SubjectItem.bEnabled)
-			{
-				VSubject->Update();
-				if (SourceItem.Setting->ParentSubject == FLiveLinkSubjectName())
+				if (bSubjectNeedsUpdate)
 				{
-					// Don't rebroadcast if this is being synchronized to another subject.
-					HandleSubjectRebroadcast(VSubject, VSubject->GetFrameData());
+					VSubject->Update();
+					UpdatedSubjects.Add(SubjectItem.Key);
 				}
 
 				EnabledSubjects.Add(SubjectItem.Key.SubjectName, SubjectItem.Key);
@@ -274,10 +275,63 @@ void FLiveLinkClient::BuildThisTicksSubjectSnapshot()
 			{
 				VSubject->ClearFrames();
 			}
-		}
-	});
+		});
+
+	return UpdatedSubjects;
+}
+
+void FLiveLinkClient::RebroadcastSubjects(const TSet<FLiveLinkSubjectKey>& UpdatedSubjects)
+{
+	// Rebroadcast paused live subjects
+	ForEachLiveSubject([this](FLiveLinkSubject* LiveSubject, FLiveLinkCollectionSourceItem& SourceItem, FLiveLinkCollectionSubjectItem& SubjectItem)
+		{
+			if (SubjectItem.bEnabled && LiveSubject->IsPaused())
+			{
+				HandleSubjectRebroadcast(LiveSubject, LiveSubject->GetFrameSnapshot().FrameData);
+			}
+		});
+
+	// Rebroadcast Live Subjects (TransmitEvaluatedData path)
+	ForEachLiveSubject([this, UpdatedSubjects](FLiveLinkSubject* LiveSubject, FLiveLinkCollectionSourceItem& SourceItem, FLiveLinkCollectionSubjectItem& SubjectItem)
+		{
+			if (GetDefault<ULiveLinkSettings>()->bTransmitEvaluatedData)
+			{
+				// Evaluated path acts like a VSubject with a single dependant as far as rebroadcast goes.
+				// So make sure to only rebroadcast it when we received an update for it.
+				if (SubjectItem.bEnabled
+					&& !LiveSubject->IsPaused()
+					&& UpdatedSubjects.Contains(SubjectItem.Key))
+				{
+					FLiveLinkFrameIdentifier SnapshotFrameId = LiveSubject->GetSnapshotFrameId();
+
+					if (LiveSubject->GetSnapshotFrameId() != LiveSubject->GetLastRebroadcastedFrameId())
+					{
+						FLiveLinkSubjectFrameData FrameData;
+
+						// todo: Evaluate at time
+						if (LiveSubject->EvaluateFrame(LiveSubject->GetRole(), FrameData))
+						{
+							if (HandleSubjectRebroadcast(SubjectItem.GetLiveSubject(), MoveTemp(FrameData)))
+							{
+								LiveSubject->SetLastRebroadcastedFrameId(SnapshotFrameId);
+							}
+						}
+					}
+				}
+			}
+		});
+
+	// Rebroadcast Virtual Subjects
+	ForEachVirtualSubject([this, UpdatedSubjects](ULiveLinkVirtualSubject* VSubject, FLiveLinkCollectionSourceItem& SourceItem, FLiveLinkCollectionSubjectItem& SubjectItem)
+		{
+			if (UpdatedSubjects.Contains(SubjectItem.Key))
+			{
+				HandleSubjectRebroadcast(VSubject, VSubject->GetFrameData());
+			}
+		});
 }
 
+
 void FLiveLinkClient::CacheValues()
 {
 #if WITH_EDITOR
@@ -320,7 +374,7 @@ void FLiveLinkClient::Shutdown()
 	}
 }
 
-void FLiveLinkClient::HandleSubjectRebroadcast(ILiveLinkSubject* InSubject, const FLiveLinkFrameDataStruct& InFrameData)
+bool FLiveLinkClient::HandleSubjectRebroadcast(ILiveLinkSubject* InSubject, const FLiveLinkFrameDataStruct& InFrameData)
 {
 	check(InSubject);
 
@@ -380,7 +434,10 @@ void FLiveLinkClient::HandleSubjectRebroadcast(ILiveLinkSubject* InSubject, cons
 					RebroadcastedSubjects.Add(InSubject->GetSubjectKey());
 				}
 
+				TRACE_CPUPROFILER_EVENT_SCOPE_TEXT(*FString::Printf(TEXT("FLiveLinkClient::HandleSubjectRebroadcast %s"), *InSubject->GetSubjectKey().SubjectName.Name.ToString()));
 				RebroadcastLiveLinkProvider->UpdateSubjectFrameData(RebroadcastName, MoveTemp(FrameDataCopy), ExtraAnnotations);
+
+				return true;
 			}
 			else
 			{
@@ -393,6 +450,8 @@ void FLiveLinkClient::HandleSubjectRebroadcast(ILiveLinkSubject* InSubject, cons
 		RemoveRebroadcastedSubject(InSubject->GetSubjectKey());
 		InSubject->SetStaticDataAsRebroadcasted(false);
 	}
+
+	return false;
 }
 
 void FLiveLinkClient::OnSubjectRemovedCallback(FLiveLinkSubjectKey InSubjectKey)
@@ -419,6 +478,32 @@ void FLiveLinkClient::RemoveRebroadcastedSubject(FLiveLinkSubjectKey InSubjectKe
 	}
 }
 
+void FLiveLinkClient::ForEachLiveSubject(TFunctionRef<void(FLiveLinkSubject*, FLiveLinkCollectionSourceItem&, FLiveLinkCollectionSubjectItem&)> VisitorFunc)
+{
+	check(Collection);
+
+	Collection->ForEachSubject([VisitorFunc](FLiveLinkCollectionSourceItem& SourceItem, FLiveLinkCollectionSubjectItem& SubjectItem)
+		{
+			if (SubjectItem.GetLiveSubject())
+			{
+				VisitorFunc(SubjectItem.GetLiveSubject(), SourceItem, SubjectItem);
+			}
+		});
+}
+
+void FLiveLinkClient::ForEachVirtualSubject(TFunctionRef<void(ULiveLinkVirtualSubject*, FLiveLinkCollectionSourceItem&, FLiveLinkCollectionSubjectItem&)> VisitorFunc)
+{
+	check(Collection);
+
+	Collection->ForEachSubject([VisitorFunc](FLiveLinkCollectionSourceItem& SourceItem, FLiveLinkCollectionSubjectItem& SubjectItem)
+		{
+			if (SubjectItem.GetVirtualSubject())
+			{
+				VisitorFunc(SubjectItem.GetVirtualSubject(), SourceItem, SubjectItem);
+			}
+		});
+}
+
 ELiveLinkSubjectState FLiveLinkClient::GetSubjectState(FLiveLinkSubjectName InSubjectName) const
 {
 	const FLiveLinkSubjectKey* SubjectKey = EnabledSubjects.Find(InSubjectName);
@@ -957,6 +1042,9 @@ void FLiveLinkClient::PushSubjectFrameData_AnyThread(const FLiveLinkSubjectKey&
 	bool bLogError = false;
 
 	bool bCanPushFrame = true;
+
+	TRACE_CPUPROFILER_EVENT_SCOPE_TEXT(*FString::Printf(TEXT("FLiveLinkClient::PushSubjectFrameData_AnyThread %s"), *InSubjectKey.SubjectName.Name.ToString()));
+
 	{
 		FScopeLock Lock(&PendingFramesCriticalSection);
 
@@ -1111,39 +1199,31 @@ void FLiveLinkClient::PushSubjectFrameData_Internal(FPendingSubjectFrame&& Subje
 			Collection->ForEachSubject([this, &SubjectFrameData](const FLiveLinkCollectionSourceItem& SourceItem, const FLiveLinkCollectionSubjectItem& SubjectItem)
 			{
 				const FLiveLinkSubjectKey SubjectKey = SubjectFrameData.SubjectKey;
+
+				// Handle virtual subjects synchronized rebroadcast.
 				if (SourceItem.Setting->ParentSubject.Name == SubjectKey.SubjectName)
 				{
 					const FLiveLinkFrameDataStruct& FrameData = SubjectFrameData.FrameData;
 
 					// todo: Time offset evaluation
-					FLiveLinkSubjectFrameData ChildData;
-					bool bSuccess = true;
-					ILiveLinkSubject* LiveSubject = nullptr;
-
-					if (SubjectItem.GetLiveSubject())
+					if (ILiveLinkSubject* LiveSubject = SubjectItem.GetLiveSubject())
 					{
-						LiveSubject = SubjectItem.GetLiveSubject();
-						bSuccess = SubjectItem.GetLiveSubject()->EvaluateFrameAtWorldTime(SubjectFrameData.FrameData.GetBaseData()->WorldTime.GetSourceTime(), SubjectItem.GetLinkSettings()->Role, ChildData);
-					}
-					else
-					{
-						LiveSubject = SubjectItem.GetVirtualSubject();
-						bSuccess = SubjectItem.GetVirtualSubject()->EvaluateFrame(SubjectItem.GetVirtualSubject()->GetRole(), ChildData);
-					}
+						FLiveLinkSubjectFrameData ChildData;
 
-					if (bSuccess)
-					{
-						const FTimecode FrameTC = FTimecode::FromFrameNumber(FrameData.GetBaseData()->MetaData.SceneTime.Time.GetFrame(), FrameData.GetBaseData()->MetaData.SceneTime.Rate);
-						UE_LOG(LogLiveLink, Verbose, TEXT("LiveLinkHub Parent (%s) - Child '%s' adding frame with Timecode:[%s.%0.3f] - SourceTime: %0.4f, Offset: %0.6f, CorrectedTime: %0.4f"), *SubjectKey.SubjectName.ToString(), *SubjectItem.Key.SubjectName.ToString(), *FrameTC.ToString(), FrameData.GetBaseData()->MetaData.SceneTime.Time.GetSubFrame(), FrameData.GetBaseData()->WorldTime.GetSourceTime(), FrameData.GetBaseData()->WorldTime.GetOffset(), FrameData.GetBaseData()->WorldTime.GetOffsettedTime());
+						if (SubjectItem.GetLiveSubject()->EvaluateFrameAtWorldTime(SubjectFrameData.FrameData.GetBaseData()->WorldTime.GetSourceTime(), SubjectItem.GetLinkSettings()->Role, ChildData))
+						{
+							const FTimecode FrameTC = FTimecode::FromFrameNumber(FrameData.GetBaseData()->MetaData.SceneTime.Time.GetFrame(), FrameData.GetBaseData()->MetaData.SceneTime.Rate);
+							UE_LOG(LogLiveLink, Verbose, TEXT("LiveLinkHub Parent (%s) - Child '%s' adding frame with Timecode:[%s.%0.3f] - SourceTime: %0.4f, Offset: %0.6f, CorrectedTime: %0.4f"), *SubjectKey.SubjectName.ToString(), *SubjectItem.Key.SubjectName.ToString(), *FrameTC.ToString(), FrameData.GetBaseData()->MetaData.SceneTime.Time.GetSubFrame(), FrameData.GetBaseData()->WorldTime.GetSourceTime(), FrameData.GetBaseData()->WorldTime.GetOffset(), FrameData.GetBaseData()->WorldTime.GetOffsettedTime());
 
-						ChildData.FrameData.GetBaseData()->MetaData.SceneTime = FrameData.GetBaseData()->MetaData.SceneTime;
-						ChildData.FrameData.GetBaseData()->MetaData.SceneTime.Rate = FrameData.GetBaseData()->MetaData.SceneTime.Rate;
+							ChildData.FrameData.GetBaseData()->MetaData.SceneTime = FrameData.GetBaseData()->MetaData.SceneTime;
+							ChildData.FrameData.GetBaseData()->MetaData.SceneTime.Rate = FrameData.GetBaseData()->MetaData.SceneTime.Rate;
 
-						HandleSubjectRebroadcast(LiveSubject, ChildData.FrameData);
-					}
-					else
-					{
-						FLiveLinkLog::Warning(TEXT("Child subjects %s could not be evaluated for data resampling."), *SubjectKey.SubjectName.Name.ToString());
+							HandleSubjectRebroadcast(LiveSubject, ChildData.FrameData);
+						}
+						else
+						{
+							FLiveLinkLog::Warning(TEXT("Child subjects %s could not be evaluated for data resampling."), *SubjectKey.SubjectName.Name.ToString());
+						}
 					}
 				}
 			});
diff --git a/Engine/Plugins/Animation/LiveLink/Source/LiveLink/Public/LiveLinkClient.h b/Engine/Plugins/Animation/LiveLink/Source/LiveLink/Public/LiveLinkClient.h
--- a/Engine/Plugins/Animation/LiveLink/Source/LiveLink/Public/LiveLinkClient.h
+++ b/Engine/Plugins/Animation/LiveLink/Source/LiveLink/Public/LiveLinkClient.h
@@ -37,7 +37,6 @@ public:
 	//~ Constructor that allow providing a custom delegate for ticking LiveLink.  
 	UE_API FLiveLinkClient(FSimpleMulticastDelegate& InTickingDelegate);
 	UE_API FLiveLinkClient(FTSSimpleMulticastDelegate& InTickingDelegate);
-	
 
 	UE_API virtual ~FLiveLinkClient();
 
@@ -206,18 +205,19 @@ private:
 	void Initialize();
 
 	/** Remove old sources & subject,  */
-	void DoPendingWork();
+	void DoPendingWork(TSet<FLiveLinkSubjectKey>& OutReceivedSubjects);
 
 	/** Update the added sources */
 	void UpdateSources();
 
 	/**
-	 * Build subject data so that during the rest of the tick it can be read without
-	 * thread locking or mem copying
+	 * Build subject data so that during the rest of the tick it can be read without thread locking or mem copying
 	 */
-	void BuildThisTicksSubjectSnapshot();
+	TSet<FLiveLinkSubjectKey> BuildThisTicksSubjectSnapshot(const TSet<FLiveLinkSubjectKey>& ReceivedSubjects);
+
+	/** Rebroadcast live (if using evaluated data) and virtual subjects. */
+	void RebroadcastSubjects(const TSet<FLiveLinkSubjectKey>& UpdatedSubjects);
 
-	/** Cache the game thread values to be reused on any thread */
 	void CacheValues();
 
 	void PushSubjectStaticData_Internal(FPendingSubjectStatic&& SubjectStaticData);
@@ -227,12 +227,12 @@ private:
 	void Shutdown();
 
 	/** Process virtual subject for rebroadcast purpose */
-	void HandleSubjectRebroadcast(ILiveLinkSubject* InSubject, const FLiveLinkFrameDataStruct& InFrameData);
+	bool HandleSubjectRebroadcast(ILiveLinkSubject* InSubject, const FLiveLinkFrameDataStruct& InFrameData);
 	/**
 	 * Rebroadcast evaluated livelink data.
 	 * @note Calling this method will not apply preprocessors and translators since it's meant to be used to transmit evaluated data. 
 	 */
-	void HandleSubjectRebroadcast(ILiveLinkSubject* InSubject, FLiveLinkSubjectFrameData&& SubjectFrameData);
+	bool HandleSubjectRebroadcast(ILiveLinkSubject* InSubject, FLiveLinkSubjectFrameData&& SubjectFrameData);
 
 	/** Called when a subject is removed. Used to remove rebroadcasted subjects */
 	void OnSubjectRemovedCallback(FLiveLinkSubjectKey InSubjectKey);
@@ -240,6 +240,12 @@ private:
 	/** Removes a subject from the rebroadcast provider and resets it if there are no more subjects */
 	void RemoveRebroadcastedSubject(FLiveLinkSubjectKey InSubjectKey);
 
+	/** Iterate over every live subject. */
+	void ForEachLiveSubject(TFunctionRef<void(class FLiveLinkSubject*, struct FLiveLinkCollectionSourceItem&, struct FLiveLinkCollectionSubjectItem&)> VisitorFunc);
+
+	/** Iterate over every virtual subject. */
+	void ForEachVirtualSubject(TFunctionRef<void(class ULiveLinkVirtualSubject*, struct FLiveLinkCollectionSourceItem&, struct FLiveLinkCollectionSubjectItem&)> VisitorFunc);
+
 protected:
 	/** Broadcast out to the SubjectFrameAddedHandles a frame data update. */
 	UE_API void BroadcastFrameDataUpdate(const FLiveLinkSubjectKey& InSubjectKey, const FLiveLinkFrameDataStruct& InFrameData);
diff --git a/Engine/Plugins/Animation/LiveLink/Source/LiveLinkEditor/Private/LiveLinkVirtualSubjectDetailCustomization.cpp b/Engine/Plugins/Animation/LiveLink/Source/LiveLinkEditor/Private/LiveLinkVirtualSubjectDetailCustomization.cpp
--- a/Engine/Plugins/Animation/LiveLink/Source/LiveLinkEditor/Private/LiveLinkVirtualSubjectDetailCustomization.cpp
+++ b/Engine/Plugins/Animation/LiveLink/Source/LiveLinkEditor/Private/LiveLinkVirtualSubjectDetailCustomization.cpp
@@ -11,6 +11,7 @@
 #include "Roles/LiveLinkAnimationRole.h"
 #include "Roles/LiveLinkBasicRole.h"
 #include "ScopedTransaction.h"
+#include "SLiveLinkSubjectRepresentationPicker.h"
 #include "Widgets/Input/SCheckBox.h"
 #include "Widgets/Input/SComboButton.h"
 
@@ -107,6 +108,47 @@ void FLiveLinkVirtualSubjectDetailCustomization::CustomizeDetails(IDetailLayoutB
 		[
 			ComboButton
 		];
+
+	TSharedRef<IPropertyHandle> SyncSubjectProperty = DetailBuilder.GetProperty(GET_MEMBER_NAME_CHECKED(ULiveLinkVirtualSubject, SyncSubject));
+	IDetailCategoryBuilder& SyncSubjectPropertyGroup = DetailBuilder.EditCategory(*SyncSubjectProperty->GetMetaData(TEXT("Category")));
+
+	DetailBuilder.HideProperty(SyncSubjectProperty);
+
+	FIsResetToDefaultVisible IsResetToDefaultVisible = FIsResetToDefaultVisible::CreateLambda(
+		[](TSharedPtr<IPropertyHandle> PropertyHandle)
+		{
+			FName Value;
+			PropertyHandle->GetChildHandle(GET_MEMBER_NAME_CHECKED(FLiveLinkSubjectName, Name))->GetValue(Value);
+			return Value != NAME_None;
+		});
+
+	FResetToDefaultHandler ResetToDefaultHandler = FResetToDefaultHandler::CreateLambda(
+		[](TSharedPtr<IPropertyHandle> PropertyHandle)
+		{
+			PropertyHandle->GetChildHandle(GET_MEMBER_NAME_CHECKED(FLiveLinkSubjectName, Name))->SetValue(NAME_None);
+		});
+
+	const FResetToDefaultOverride Override = FResetToDefaultOverride::Create(IsResetToDefaultVisible, ResetToDefaultHandler);
+
+	SyncSubjectPropertyGroup.AddCustomRow(LOCTEXT("SyncSubjectLabel", "Sync Subject"))
+		.PropertyHandleList({ SyncSubjectProperty })
+		.NameContent()
+		[
+			SNew(STextBlock)
+				.Font(IDetailLayoutBuilder::GetDetailFont())
+				.Text(LOCTEXT("SyncSubjectLabel", "Sync Subject"))
+		]
+		.ValueContent()
+		[
+			SAssignNew(SyncSubjectWidget, SLiveLinkSubjectRepresentationPicker)
+				.ShowRole(false)
+				.Font(IDetailLayoutBuilder::GetDetailFont())
+				.HasMultipleValues(false)
+				.OnGetSubjects(this, &FLiveLinkVirtualSubjectDetailCustomization::OnGetSubjects)
+				.Value(this, &FLiveLinkVirtualSubjectDetailCustomization::GetSyncSubject)
+				.OnValueChanged(this, &FLiveLinkVirtualSubjectDetailCustomization::SetSyncSubject)
+		]
+		.OverrideResetToDefault(Override);
 }
 
 int32 GetArrayPropertyIndex(TSharedPtr<IPropertyHandleArray> ArrayProperty, FName ItemToSearchFor, uint32 NumItems)
@@ -190,6 +232,16 @@ TSharedRef<ITableRow> FLiveLinkVirtualSubjectDetailCustomization::OnGenerateWidg
 						check(Result == FPropertyAccess::Success);
 						Result = SubjectsArrayPropertyHandle->GetElement(NumItems)->SetValueFromFormattedString(TextValue, EPropertyValueSetFlags::NotTransactable);
 						check(Result == FPropertyAccess::Success);
+
+						// Sync subject handling
+						if (TStrongObjectPtr<ULiveLinkVirtualSubject> VSubject = CaptureThis->SubjectPtr.Pin())
+						{
+							// Automatically set a sync subject if none has been selected.
+							if (VSubject->SyncSubject == FLiveLinkSubjectName())
+							{
+								VSubject->SyncSubject = NewSubjectName;
+							}
+						}
 					}
 					else
 					{
@@ -198,6 +250,15 @@ TSharedRef<ITableRow> FLiveLinkVirtualSubjectDetailCustomization::OnGenerateWidg
 						{
 							SubjectsArrayPropertyHandle->DeleteItem(RemoveIndex);
 						}
+
+						if (TStrongObjectPtr<ULiveLinkVirtualSubject> VSubject = CaptureThis->SubjectPtr.Pin())
+						{
+							// Clear out the sync subject if we disable it.
+							if (VSubject->SyncSubject == *InItem)
+							{
+								VSubject->SyncSubject = FLiveLinkSubjectName();
+							}
+						}
 					}
 
 					CaptureThis->UpdateSelectedSubjects();
@@ -207,9 +268,9 @@ TSharedRef<ITableRow> FLiveLinkVirtualSubjectDetailCustomization::OnGenerateWidg
 			.FillWidth(1.f)
 			[
 				SNew(STextBlock)
-				.Text(FText::FromName(*InItem))
-				.ColorAndOpacity(this, &FLiveLinkVirtualSubjectDetailCustomization::HandleSubjectItemColor, InItem)
-				.ToolTipText(this, &FLiveLinkVirtualSubjectDetailCustomization::HandleSubjectItemToolTip, InItem)
+					.Text(FText::FromName(*InItem))
+					.ColorAndOpacity(this, &FLiveLinkVirtualSubjectDetailCustomization::HandleSubjectItemColor, InItem)
+					.ToolTipText(this, &FLiveLinkVirtualSubjectDetailCustomization::HandleSubjectItemToolTip, InItem)
 			]
 		];
 }
@@ -226,7 +287,7 @@ FSlateColor FLiveLinkVirtualSubjectDetailCustomization::HandleSubjectItemColor(F
 		if (!SubjectKeys.ContainsByPredicate([ThisItem](const FLiveLinkSubjectKey& Other) { return Other.SubjectName.Name == ThisItem; }))
 		{
 			Result = FLinearColor::Red;
-		} 
+		}
 	}
 
 	return Result;
@@ -267,10 +328,48 @@ void FLiveLinkVirtualSubjectDetailCustomization::UpdateSelectedSubjects()
 		}
 	}
 
+	if (TStrongObjectPtr<ULiveLinkVirtualSubject> VSubject = SubjectPtr.Pin())
+	{
+		if (!VSubject->GetSubjects().Contains(VSubject->SyncSubject))
+		{
+			VSubject->SyncSubject = FLiveLinkSubjectName();
+		}
+	}
+
 	if (SelectedSubjectsListView)
 	{
 		SelectedSubjectsListView->RequestListRefresh();
 	}
 }
 
+void FLiveLinkVirtualSubjectDetailCustomization::OnGetSubjects(TArray<FLiveLinkSubjectKey>& OutSubjectsList) const
+{
+	if (TStrongObjectPtr<ULiveLinkVirtualSubject> VSubject = SubjectPtr.Pin())
+	{
+		const TArray<FLiveLinkSubjectName>& Subjects = VSubject->GetSubjects();
+		OutSubjectsList.Reset(Subjects.Num());
+		Algo::Transform(Subjects, OutSubjectsList, [](FLiveLinkSubjectName Name) { return FLiveLinkSubjectKey{ FGuid(), Name }; });
+	}
+}
+
+SLiveLinkSubjectRepresentationPicker::FLiveLinkSourceSubjectRole FLiveLinkVirtualSubjectDetailCustomization::GetSyncSubject() const
+{
+	SLiveLinkSubjectRepresentationPicker::FLiveLinkSourceSubjectRole Role;
+	if (TStrongObjectPtr<ULiveLinkVirtualSubject> VSubject = SubjectPtr.Pin())
+	{
+		Role.Subject = VSubject->SyncSubject;
+	}
+
+	return Role;
+}
+
+void FLiveLinkVirtualSubjectDetailCustomization::SetSyncSubject(SLiveLinkSubjectRepresentationPicker::FLiveLinkSourceSubjectRole Role)
+{
+	if (TStrongObjectPtr<ULiveLinkVirtualSubject> VSubject = SubjectPtr.Pin())
+	{
+		VSubject->Modify();
+		VSubject->SyncSubject = Role.Subject;
+	}
+}
+
 #undef LOCTEXT_NAMESPACE
diff --git a/Engine/Plugins/Animation/LiveLink/Source/LiveLinkEditor/Private/LiveLinkVirtualSubjectDetailCustomization.h b/Engine/Plugins/Animation/LiveLink/Source/LiveLinkEditor/Private/LiveLinkVirtualSubjectDetailCustomization.h
--- a/Engine/Plugins/Animation/LiveLink/Source/LiveLinkEditor/Private/LiveLinkVirtualSubjectDetailCustomization.h
+++ b/Engine/Plugins/Animation/LiveLink/Source/LiveLinkEditor/Private/LiveLinkVirtualSubjectDetailCustomization.h
@@ -3,7 +3,7 @@
 #pragma once
 
 #include "IDetailCustomization.h"
-
+#include "SLiveLinkSubjectRepresentationPicker.h"
 
 #include "Widgets/Views/SListView.h"
 
@@ -56,6 +56,15 @@ private:
 	// Update the list of selected subjects displayed above the subject picker.
 	void UpdateSelectedSubjects();
 
+	// Returns the list of subjects that constitute the virtual subject.
+	void OnGetSubjects(TArray<FLiveLinkSubjectKey>& OutSubjectsList) const;
+
+	// Returns the subject that synchronizes the virtual subject.
+	SLiveLinkSubjectRepresentationPicker::FLiveLinkSourceSubjectRole GetSyncSubject() const;
+
+	// Sets the SyncSubject property on the virtual subject,
+	void SetSyncSubject(SLiveLinkSubjectRepresentationPicker::FLiveLinkSourceSubjectRole Role);
+
 private:
 	// The tile set being edited
 	TWeakObjectPtr<ULiveLinkVirtualSubject> SubjectPtr;
@@ -72,4 +81,7 @@ private:
 	TArray<FSubjectEntryPtr> SubjectsListItems;
 	TArray<FSubjectEntryPtr> SelectedSubjectsListItems;
 	TSharedPtr< SListView< FSubjectEntryPtr > > SelectedSubjectsListView;
+
+	// LiveLink subject picker that restricts to the subjects that make up the virtual subject.
+	TSharedPtr<class SLiveLinkSubjectRepresentationPicker> SyncSubjectWidget;
 };
diff --git a/Engine/Source/Runtime/LiveLinkInterface/Public/ILiveLinkSubject.h b/Engine/Source/Runtime/LiveLinkInterface/Public/ILiveLinkSubject.h
--- a/Engine/Source/Runtime/LiveLinkInterface/Public/ILiveLinkSubject.h
+++ b/Engine/Source/Runtime/LiveLinkInterface/Public/ILiveLinkSubject.h
@@ -17,7 +17,6 @@ class ULiveLinkFrameInterpolationProcessor;
 class ULiveLinkFramePreProcessor;
 class ULiveLinkRole;
 
-
 /**
  * The base class of a LiveLink subject.
  * Subjects are individual streams of data within the client.
@@ -75,6 +74,16 @@ public:
 	/** Unpause subject. */
 	virtual void UnpauseSubject() = 0;
 
+	/** Get the frame ID of the last snapshot that was captured. */
+	FLiveLinkFrameIdentifier GetSnapshotFrameId() const
+	{
+		if (HasValidFrameSnapshot())
+		{
+			return GetFrameSnapshot().FrameData.GetBaseData()->FrameId;
+		}
+		return INDEX_NONE;
+	}
+
 protected:
 	virtual const FLiveLinkSubjectFrameData& GetFrameSnapshot() const = 0;
 
diff --git a/Engine/Source/Runtime/LiveLinkInterface/Public/LiveLinkVirtualSubject.h b/Engine/Source/Runtime/LiveLinkInterface/Public/LiveLinkVirtualSubject.h
--- a/Engine/Source/Runtime/LiveLinkInterface/Public/LiveLinkVirtualSubject.h
+++ b/Engine/Source/Runtime/LiveLinkInterface/Public/LiveLinkVirtualSubject.h
@@ -91,7 +91,12 @@ public:
 	{
 		return SubjectKey.SubjectName;
 	}
-	
+
+public:
+	/** If set, the virtual subject will only update and transmit data when the sync subject receives new data. */
+	UPROPERTY(EditAnywhere, Category = "LiveLink")
+	FLiveLinkSubjectName SyncSubject;
+
 protected:
 
 	/** Updates the list of translators valid for this frame */
