LiveLinkHub - Allow enabling transmission of evaluated data on a per-source basis
Hide buffer properties if evaluated data is not enabled

diff --git a/Engine/Plugins/Animation/LiveLink/Source/LiveLink/Private/LiveLinkClient.cpp b/Engine/Plugins/Animation/LiveLink/Source/LiveLink/Private/LiveLinkClient.cpp
--- a/Engine/Plugins/Animation/LiveLink/Source/LiveLink/Private/LiveLinkClient.cpp
+++ b/Engine/Plugins/Animation/LiveLink/Source/LiveLink/Private/LiveLinkClient.cpp
@@ -247,14 +247,17 @@ TSet<FLiveLinkSubjectKey> FLiveLinkClient::BuildThisTicksSubjectSnapshot(const T
 
 				if (bSubjectIsSynchronized)
 				{
-					for (const FLiveLinkSubjectKey& ReceivedSubject : ReceivedSubjects)
+					if (SubjectItem.bEnabled && !VSubject->IsPaused() && SourceItem.Setting->bTransmitEvaluatedData)
 					{
-						// We can either synchronize at the source or subject level to another subject.
-						if (SourceItem.Setting->ParentSubject == ReceivedSubject.SubjectName
-							|| VSubject->SyncSubject == ReceivedSubject.SubjectName)
+						for (const FLiveLinkSubjectKey& ReceivedSubject : ReceivedSubjects)
 						{
-							bSubjectNeedsUpdate = true;
-							break;
+							// We can either synchronize at the source or subject level to another subject.
+							if (SourceItem.Setting->ParentSubject == ReceivedSubject.SubjectName
+								|| VSubject->SyncSubject == ReceivedSubject.SubjectName)
+							{
+								bSubjectNeedsUpdate = true;
+								break;
+							}
 						}
 					}
 				}
@@ -294,7 +297,7 @@ void FLiveLinkClient::RebroadcastSubjects(const TSet<FLiveLinkSubjectKey>& Updat
 	// Rebroadcast Live Subjects (TransmitEvaluatedData path)
 	ForEachLiveSubject([this, UpdatedSubjects](FLiveLinkSubject* LiveSubject, FLiveLinkCollectionSourceItem& SourceItem, FLiveLinkCollectionSubjectItem& SubjectItem)
 		{
-			if (GetDefault<ULiveLinkSettings>()->bTransmitEvaluatedData)
+			if (SourceItem.Setting->bTransmitEvaluatedData)
 			{
 				// Evaluated path acts like a VSubject with a single dependant as far as rebroadcast goes.
 				// So make sure to only rebroadcast it when we received an update for it.
@@ -1188,7 +1191,7 @@ void FLiveLinkClient::PushSubjectFrameData_Internal(FPendingSubjectFrame&& Subje
 		if (!LinkSubject->IsPaused())
 		{
 			/** Only rebroadcast here if we're transmitted non-evaluated data. */
-			if (!GetDefault<ULiveLinkSettings>()->bTransmitEvaluatedData)
+			if (SourceItem->Setting->bTransmitEvaluatedData)
 			{
 				HandleSubjectRebroadcast(LinkSubject, SubjectFrameData.FrameData);
 			}
diff --git a/Engine/Plugins/Animation/LiveLink/Source/LiveLink/Public/LiveLinkSettings.h b/Engine/Plugins/Animation/LiveLink/Source/LiveLink/Public/LiveLinkSettings.h
--- a/Engine/Plugins/Animation/LiveLink/Source/LiveLink/Public/LiveLinkSettings.h
+++ b/Engine/Plugins/Animation/LiveLink/Source/LiveLink/Public/LiveLinkSettings.h
@@ -131,13 +131,6 @@ public:
 	UPROPERTY(config)
 	bool bTranslateRebroadcastFrames = false;
 
-	/**
-	 * Experimental
-	 * Enabling this will evaluate livelink data before rebroadcasting it.
-	 */
-	UPROPERTY(config, EditAnywhere, Category = "LiveLink")
-	bool bTransmitEvaluatedData = false;
-
 	/**
 	 * A source may still exist but does not send frames for a subject.
 	 * Time before considering the subject as "invalid".
@@ -165,6 +158,14 @@ public:
 	UPROPERTY(config, EditAnywhere, Category = "LiveLink|UI|Debug")
 	uint8 TextSizeSubject;
 
+#if WITH_EDITORONLY_DATA
+	/** Whether to rebroadcast evaluted data. */
+	UE_DEPRECATED(5.7, "bTransmitEvaluatedData has been moved to the source settings.")
+		UPROPERTY(meta = (DeprecatedProperty))
+	bool bTransmitEvaluatedData_DEPRECATED = false;
+#endif
+
+
 public:
 	UE_API FLiveLinkRoleProjectSetting GetDefaultSettingForRole(TSubclassOf<ULiveLinkRole> Role) const;
 
diff --git a/Engine/Plugins/Animation/LiveLink/Source/LiveLinkEditor/Private/LiveLinkHubSubjectSettingsDetailsCustomization.cpp b/Engine/Plugins/Animation/LiveLink/Source/LiveLinkEditor/Private/LiveLinkHubSubjectSettingsDetailsCustomization.cpp
--- a/Engine/Plugins/Animation/LiveLink/Source/LiveLinkEditor/Private/LiveLinkHubSubjectSettingsDetailsCustomization.cpp
+++ b/Engine/Plugins/Animation/LiveLink/Source/LiveLinkEditor/Private/LiveLinkHubSubjectSettingsDetailsCustomization.cpp
@@ -71,11 +71,6 @@ void FLiveLinkHubSubjectSettingsDetailsCustomization::CustomizeDetails(IDetailLa
 				];
 		}
 
-		if (!GetDefault<ULiveLinkSettings>()->bTransmitEvaluatedData)
-		{
-			InDetailBuilder.HideProperty(GET_MEMBER_NAME_CHECKED(ULiveLinkHubSubjectSettings, InterpolationProcessor), ULiveLinkSubjectSettings::StaticClass());
-		}
-
 		InDetailBuilder.HideProperty(GET_MEMBER_NAME_CHECKED(ULiveLinkHubSubjectSettings, FrameRate), ULiveLinkSubjectSettings::StaticClass());
 		InDetailBuilder.HideProperty(GET_MEMBER_NAME_CHECKED(ULiveLinkHubSubjectSettings, bRebroadcastSubject), ULiveLinkSubjectSettings::StaticClass());
 		InDetailBuilder.HideProperty(GET_MEMBER_NAME_CHECKED(ULiveLinkHubSubjectSettings, Translators), ULiveLinkSubjectSettings::StaticClass());
diff --git a/Engine/Plugins/Animation/LiveLink/Source/LiveLinkEditor/Private/LiveLinkSourceSettingsDetailCustomization.cpp b/Engine/Plugins/Animation/LiveLink/Source/LiveLinkEditor/Private/LiveLinkSourceSettingsDetailCustomization.cpp
--- a/Engine/Plugins/Animation/LiveLink/Source/LiveLinkEditor/Private/LiveLinkSourceSettingsDetailCustomization.cpp
+++ b/Engine/Plugins/Animation/LiveLink/Source/LiveLinkEditor/Private/LiveLinkSourceSettingsDetailCustomization.cpp
@@ -58,6 +58,30 @@ void FLiveLinkSourceSettingsDetailCustomization::CustomizeDetails(IDetailLayoutB
 	{
 		return;
 	}
+
+	bool bSupportsBuffering = true;
+	TSharedRef<IPropertyHandle> TransmitEvalutedDataHandle = InDetailBuilder.GetProperty(GET_MEMBER_NAME_CHECKED(ULiveLinkSourceSettings, bTransmitEvaluatedData));
+
+	const bool bIsInLiveLinkHubApp = GConfig->GetBoolOrDefault(TEXT("LiveLink"), TEXT("bCreateLiveLinkHubInstance"), false, GEngineIni);
+	if (bIsInLiveLinkHubApp)
+	{
+		bSupportsBuffering = SourceSettings->bTransmitEvaluatedData;
+		TransmitEvalutedDataHandle->SetOnPropertyValueChanged(FSimpleDelegate::CreateSP(this, &FLiveLinkSourceSettingsDetailCustomization::ForceRefresh));
+	}
+	else
+	{
+		InDetailBuilder.HideProperty(TransmitEvalutedDataHandle);
+	}
+
+	if (!bSupportsBuffering)
+	{
+		InDetailBuilder.HideProperty(InDetailBuilder.GetProperty(GET_MEMBER_NAME_CHECKED(ULiveLinkSourceSettings, Mode)));
+		InDetailBuilder.HideProperty(InDetailBuilder.GetProperty(GET_MEMBER_NAME_CHECKED(ULiveLinkSourceSettings, BufferSettings)));
+
+		// Don't add buffering properties if the source doesn't support it. 
+		return;
+	}
+
 	const ELiveLinkSourceMode SourceMode = *reinterpret_cast<ELiveLinkSourceMode*>(ModeValuePtr);
 
 	InDetailBuilder.AddPropertyToCategory(ModePropertyHandle);
diff --git a/Engine/Source/Runtime/LiveLinkInterface/Public/LiveLinkSourceSettings.h b/Engine/Source/Runtime/LiveLinkInterface/Public/LiveLinkSourceSettings.h
--- a/Engine/Source/Runtime/LiveLinkInterface/Public/LiveLinkSourceSettings.h
+++ b/Engine/Source/Runtime/LiveLinkInterface/Public/LiveLinkSourceSettings.h
@@ -176,6 +176,10 @@ public:
 	UPROPERTY(EditAnywhere, Category = "Settings")
 	FLiveLinkSourceBufferManagementSettings BufferSettings;
 
+	/** Whether to rebroadcast evaluted data. */
+	UPROPERTY(EditAnywhere, Category = "Settings")
+	bool bTransmitEvaluatedData = false;
+
 	/** Connection information that is needed by the factory to recreate the source from a preset. */
 	UPROPERTY(EditAnywhere, AdvancedDisplay, Category = "Settings")
 	FString ConnectionString;
