LiveLinkHub - Fix data not being rebroadcast by default

diff --git a/Engine/Plugins/Animation/LiveLink/Source/LiveLink/Private/LiveLinkClient.cpp b/Engine/Plugins/Animation/LiveLink/Source/LiveLink/Private/LiveLinkClient.cpp
--- a/Engine/Plugins/Animation/LiveLink/Source/LiveLink/Private/LiveLinkClient.cpp
+++ b/Engine/Plugins/Animation/LiveLink/Source/LiveLink/Private/LiveLinkClient.cpp
@@ -1191,7 +1191,7 @@ void FLiveLinkClient::PushSubjectFrameData_Internal(FPendingSubjectFrame&& Subje
 		if (!LinkSubject->IsPaused())
 		{
 			/** Only rebroadcast here if we're transmitted non-evaluated data. */
-			if (SourceItem->Setting->bTransmitEvaluatedData)
+			if (!SourceItem->Setting->bTransmitEvaluatedData)
 			{
 				HandleSubjectRebroadcast(LinkSubject, SubjectFrameData.FrameData);
 			}
