LiveLink - Add instrumentation to LiveLinkCustomTimeStep

diff --git a/Engine/Plugins/Animation/LiveLink/Source/LiveLink/Private/LiveLinkCustomTimeStep.cpp b/Engine/Plugins/Animation/LiveLink/Source/LiveLink/Private/LiveLinkCustomTimeStep.cpp
--- a/Engine/Plugins/Animation/LiveLink/Source/LiveLink/Private/LiveLinkCustomTimeStep.cpp
+++ b/Engine/Plugins/Animation/LiveLink/Source/LiveLink/Private/LiveLinkCustomTimeStep.cpp
@@ -132,8 +132,10 @@ uint32 ULiveLinkCustomTimeStep::GetExpectedSyncCountDelta() const
 }
 
 bool ULiveLinkCustomTimeStep::WaitForSync()
-{	
+{
 	EventWaitForSync->Trigger();
+
+	TRACE_CPUPROFILER_EVENT_SCOPE(ULiveLinkCustomTimeStep::EventLiveLink->Wait);
 	return EventLiveLink->Wait(FTimespan(TimeoutInSeconds * ETimespan::TicksPerSecond));
 }
 
@@ -268,6 +270,8 @@ void ULiveLinkCustomTimeStep::OnLiveLinkSubjectRemoved(FLiveLinkSubjectKey InSub
 
 void ULiveLinkCustomTimeStep::OnLiveLinkFrameDataReceived_AnyThread(const FLiveLinkFrameDataStruct& InFrameData)
 {
+	TRACE_CPUPROFILER_EVENT_SCOPE(ULiveLinkCustomTimeStep::OnLiveLinkFrameDataReceived_AnyThread);
+
 	// Increment sync counter based on frame id. For example, if frame id is incremented by 2, it means that a frame
 	// was skipped and we must report 2 sync signals, even though only one was received. The purpose of this is to
 	// reflect source delta times in the engine delta times.
@@ -310,6 +314,7 @@ void ULiveLinkCustomTimeStep::OnLiveLinkFrameDataReceived_AnyThread(const FLiveL
 		// We can't block if this function is being called in the game thread because it would deadlock.
 		if (bLockStepModeAnyThread && !IsInGameThread())
 		{
+			TRACE_CPUPROFILER_EVENT_SCOPE(ULiveLinkCustomTimeStep::EventWaitForSync->Wait);
 			EventWaitForSync->Wait(FTimespan((TimeoutInSecondsAnyThread / 2) * ETimespan::TicksPerSecond));
 		}
 	}
