LiveLink - Fix deadlock when deleting metahuman video source with multiple subjects
Fix being unable to unpause subject after disabling/enabling it

diff --git a/Engine/Plugins/Animation/LiveLink/Source/LiveLink/Private/LiveLinkClient.cpp b/Engine/Plugins/Animation/LiveLink/Source/LiveLink/Private/LiveLinkClient.cpp
--- a/Engine/Plugins/Animation/LiveLink/Source/LiveLink/Private/LiveLinkClient.cpp
+++ b/Engine/Plugins/Animation/LiveLink/Source/LiveLink/Private/LiveLinkClient.cpp
@@ -119,13 +119,15 @@ void FLiveLinkClient::Initialize()
 
 void FLiveLinkClient::DoPendingWork(TSet<FLiveLinkSubjectKey>& OutReceivedSubjects)
 {
-	FScopeLock PendingFramesLock(&PendingFramesCriticalSection);
-
 	check(Collection);
 
 	// Remove Sources and Subjects
+	// Note MH-15807: This needs to happen outside of the PendingFramesCriticalSection, since deleting a subject could cause it to way on its thread,
+	// which could already have locked this critical section, causing a deadlock.
 	Collection->RemovePendingKill();
 
+	FScopeLock PendingFramesLock(&PendingFramesCriticalSection);
+
 	{
 		// Add new Subject static data
 		for (FPendingSubjectStatic& SubjectStatic : SubjectStaticToPush)
@@ -232,7 +234,10 @@ TSet<FLiveLinkSubjectKey> FLiveLinkClient::BuildThisTicksSubjectSnapshot(const T
 			}
 			else
 			{
-				LiveSubject->ClearFrames();
+				if (!LiveSubject->IsPaused())
+				{
+					LiveSubject->ClearFrames();
+				}
 			}
 		});
 
diff --git a/Engine/Plugins/Animation/LiveLink/Source/LiveLinkEditor/Private/LiveLinkClientPanelViews.cpp b/Engine/Plugins/Animation/LiveLink/Source/LiveLinkEditor/Private/LiveLinkClientPanelViews.cpp
--- a/Engine/Plugins/Animation/LiveLink/Source/LiveLinkEditor/Private/LiveLinkClientPanelViews.cpp
+++ b/Engine/Plugins/Animation/LiveLink/Source/LiveLinkEditor/Private/LiveLinkClientPanelViews.cpp
@@ -1069,7 +1069,7 @@ bool FLiveLinkSubjectsView::CanPauseSubject() const
 	TArray<FLiveLinkSubjectUIEntryPtr> Selected;
 	SubjectsTreeView->GetSelectedItems(Selected);
 
-	return Selected.Num() > 0 && Selected[0] && Selected[0]->IsSubjectValid();
+	return Selected.Num() > 0 && Selected[0] && Selected[0]->IsSubjectValid() && Selected[0]->IsSubjectEnabled();
 }
 
 void FLiveLinkSubjectsView::HandlePauseSubject()
