LiveLink - Add a LiveLink Broadcast Component.

This component currently allows streaming LiveLink data for Transform, Camera, and Skeletal Mesh data.
Adding a broadcast component will create a LiveLink subject that's automatically rebroadcasting data from the component.

Also fix changing a subject's role clearing its subject settings, this was required to allow changing the subject role without clearing the rebroadcast flag.

todo: Maybe we just check on end frame all the modified objects

diff --git a/Engine/Plugins/Animation/LiveLink/Source/LiveLink/LiveLink.Build.cs b/Engine/Plugins/Animation/LiveLink/Source/LiveLink/LiveLink.Build.cs
--- a/Engine/Plugins/Animation/LiveLink/Source/LiveLink/LiveLink.Build.cs
+++ b/Engine/Plugins/Animation/LiveLink/Source/LiveLink/LiveLink.Build.cs
@@ -28,7 +28,8 @@ namespace UnrealBuildTool.Rules
 				"Projects",
 				"SlateCore",
 				"Slate",
-				"TimeManagement"
+				"TimeManagement",
+				"CinematicCamera" // Used for the broadcast component to populate camera data.
 			});
 
 			if (Target.bBuildEditor)
diff --git a/Engine/Plugins/Animation/LiveLink/Source/LiveLink/Private/LiveLinkBroadcastComponent.cpp b/Engine/Plugins/Animation/LiveLink/Source/LiveLink/Private/LiveLinkBroadcastComponent.cpp
new file mode 100644
--- /dev/null
+++ b/Engine/Plugins/Animation/LiveLink/Source/LiveLink/Private/LiveLinkBroadcastComponent.cpp
@@ -0,0 +1,424 @@
+// Copyright Epic Games, Inc. All Rights Reserved.
+
+#include "LiveLinkBroadcastComponent.h"
+
+#include "Animation/AnimInstance.h"
+#include "Camera/CameraComponent.h"
+#include "CineCameraComponent.h"
+#include "Components/SceneComponent.h"
+#include "Components/SkeletalMeshComponent.h"
+#include "Engine/Engine.h"
+#include "Features/IModularFeatures.h"
+#include "ILiveLinkClient.h"
+#include "LiveLinkBroadcastSubsystem.h"
+#include "Roles/LiveLinkBasicRole.h"
+#include "Roles/LiveLinkAnimationRole.h"
+#include "Roles/LiveLinkCameraRole.h"
+#include "Roles/LiveLinkCameraTypes.h"
+#include "Roles/LiveLinkTransformRole.h"
+#include "Roles/LiveLinkTransformTypes.h"
+#include "UObject/UObjectGlobals.h"
+
+ULiveLinkBroadcastComponent::ULiveLinkBroadcastComponent()
+{
+	PrimaryComponentTick.bCanEverTick = true;
+	PrimaryComponentTick.bStartWithTickEnabled = true;
+	bTickInEditor = true;
+	bIsDirty = true;
+}
+
+void ULiveLinkBroadcastComponent::OnComponentDestroyed(bool bDestroyingHierarchy)
+{
+	ULiveLinkBroadcastSubsystem* BroadcastSubsystem = GEngine->GetEngineSubsystem<ULiveLinkBroadcastSubsystem>();
+
+	BroadcastSubsystem->RemoveSubject(SubjectKey);
+
+#if WITH_EDITOR
+	FCoreUObjectDelegates::OnObjectModified.RemoveAll(this);
+#endif
+	
+	Super::OnComponentDestroyed(bDestroyingHierarchy);
+}
+
+#if WITH_EDITOR
+void ULiveLinkBroadcastComponent::PostEditChangeProperty(struct FPropertyChangedEvent& PropertyChangedEvent)
+{
+	Super::PostEditChangeProperty(PropertyChangedEvent);
+
+	bIsDirty = PropertyChangedEvent.GetPropertyName() == GET_MEMBER_NAME_CHECKED(ULiveLinkBroadcastComponent, SubjectName)
+		|| PropertyChangedEvent.GetPropertyName() == GET_MEMBER_NAME_CHECKED(ULiveLinkBroadcastComponent, Role)
+		|| PropertyChangedEvent.GetPropertyName() == GET_MEMBER_NAME_CHECKED(ULiveLinkBroadcastComponent, AllowedCurveNames)
+		|| PropertyChangedEvent.GetPropertyName() == GET_MEMBER_NAME_CHECKED(ULiveLinkBroadcastComponent, AllowedBoneNames);
+	
+}
+
+void ULiveLinkBroadcastComponent::OnObjectModified(UObject* ModifiedObject)
+{
+	if (ModifiedObject == CachedResolvedObject)
+	{
+		bIsDirty = true;
+	}
+}
+#endif
+
+void ULiveLinkBroadcastComponent::OnComponentCreated()
+{
+	Super::OnComponentCreated();
+
+	USkeletalMeshComponent* SkeletalMeshComponent = GetOwner()->GetComponentByClass<USkeletalMeshComponent>();
+
+#if WITH_EDITOR
+	FCoreUObjectDelegates::OnObjectModified.AddUObject(this, &ULiveLinkBroadcastComponent::OnObjectModified);
+#endif
+
+	if (!Role)
+	{
+		// Try to infer the role from the components on the actor.
+		if (SkeletalMeshComponent)
+		{
+			Role = ULiveLinkAnimationRole::StaticClass();
+
+			// Set the mesh to the first skel mesh component we find if it's not set.
+			if (SourceMesh.PathToComponent.IsEmpty())
+			{
+				SourceMesh.PathToComponent = SkeletalMeshComponent->GetPathName(GetOwner());
+			}
+
+			SkeletalMeshComponent->SetUpdateAnimationInEditor(true);
+		}
+		else if (UCameraComponent* CameraComponent = GetOwner()->GetComponentByClass<UCameraComponent>())
+		{
+			Role = ULiveLinkCameraRole::StaticClass();
+
+			if (SourceMesh.PathToComponent.IsEmpty())
+			{
+				SourceMesh.PathToComponent = CameraComponent->GetPathName(GetOwner());
+			}
+		}
+		else
+		{
+			Role = ULiveLinkTransformRole::StaticClass();
+
+			if (SourceMesh.PathToComponent.IsEmpty())
+			{
+				SourceMesh.PathToComponent = GetOwner()->GetRootComponent()->GetPathName(GetOwner());
+			}
+		}
+	}
+}
+
+void ULiveLinkBroadcastComponent::TickComponent(float DeltaTime, ELevelTick TickType, FActorComponentTickFunction* ThisTickFunction)
+{
+	Super::TickComponent(DeltaTime, TickType, ThisTickFunction);
+
+	if (!bEnable)
+	{
+		return;
+	}
+
+	if (bIsDirty)
+	{
+		if (SubjectName == NAME_None)
+		{
+#if WITH_EDITOR
+			SubjectName = *GetTypedOuter<AActor>()->GetActorLabel();
+#else
+			SubjectName = GetTypedOuter<AActor>()->GetFName();
+#endif
+		}
+
+		CachedResolvedObject = SourceMesh.GetComponent(GetOwner());
+
+		InitiateLiveLink();
+		
+		BroadcastStaticData();
+
+		bIsDirty = false;
+	}
+
+	BroadcastFrameData();
+}
+
+void ULiveLinkBroadcastComponent::InitiateLiveLink()
+{
+	if (Role && !SourceMesh.PathToComponent.IsEmpty())
+	{
+		SubjectKey = GEngine->GetEngineSubsystem<ULiveLinkBroadcastSubsystem>()->CreateSubject(SubjectName, Role);
+	}
+	else
+	{
+		UE_LOG(LogTemp, Warning, TEXT("LiveLinkBroadcastComponent role or source mesh isn't set on %s"), *GetOwner()->GetName());
+	}
+}
+
+TArray<FName> ULiveLinkBroadcastComponent::GetSkeletalBoneNames() const
+{
+	TArray<FName> BoneNames;
+		
+	if (USkeletalMeshComponent* SourceSkelMesh = Cast<USkeletalMeshComponent>(SourceMesh.GetComponent(GetOwner())))
+	{
+		// If there are any entries in the allowed bones list, just use those bones names.
+		if (AllowedBoneNames.Num() > 0)
+		{
+			BoneNames = AllowedBoneNames;
+		}
+		else
+		{
+			for (int32 Index = 0; Index < SourceSkelMesh->GetNumBones() - 1; Index++)
+			{
+				FName BoneName = SourceSkelMesh->GetBoneName(Index);
+				BoneNames.Emplace(BoneName);
+			}
+		}
+	}
+
+	return BoneNames;
+}
+
+TArray<int32> ULiveLinkBroadcastComponent::GetSkeletalBoneParents() const
+{
+	TArray<int32> BoneParents;
+	TArray<FName> BoneNames;
+	
+	if (USkeletalMeshComponent* SourceSkelMesh = Cast<USkeletalMeshComponent>(SourceMesh.GetComponent(GetOwner())))
+	{
+		if (AllowedBoneNames.Num() > 0) //If there are any entries in the allowed bones list, just use those bones names.
+		{
+			for (FName BoneName : AllowedBoneNames)
+			{
+				FName ParentBoneName = SourceSkelMesh->GetParentBone(BoneName);
+				int32 BoneParentIdx = AllowedBoneNames.Find(ParentBoneName);
+				BoneParents.Add(BoneParentIdx);
+			}
+		}
+		else
+		{
+			for (int32 Index = 0; Index < SourceSkelMesh->GetNumBones() - 1; Index++)
+			{
+				FName BoneName = SourceSkelMesh->GetBoneName(Index);
+				FName ParentBoneName = SourceSkelMesh->GetParentBone(BoneName);
+				int32 BoneParentIndex = SourceSkelMesh->GetBoneIndex(ParentBoneName);
+				BoneParents.Add(BoneParentIndex);
+			}
+		}
+	}
+
+	return BoneParents;
+}
+
+TArray<FTransform> ULiveLinkBroadcastComponent::GetBoneTransforms() const
+{
+	TArray<FTransform> BoneTransforms;
+
+	if (USkeletalMeshComponent* SourceSkelMesh = Cast<USkeletalMeshComponent>(SourceMesh.GetComponent(GetOwner())))
+	{
+		if (AllowedBoneNames.Num() > 0) //If there are any entries in the allowed bones list, just use those bones names.
+		{
+			for (FName BoneName : AllowedBoneNames)
+			{
+				FTransform BoneTransform = SourceSkelMesh->GetSocketTransform(BoneName, RTS_ParentBoneSpace);
+				BoneTransforms.Add(BoneTransform);
+			}
+		}
+		else
+		{
+			for (int32 Index = 0; Index < SourceSkelMesh->GetNumBones() - 1; Index++)
+			{
+				FTransform BoneTransform = SourceSkelMesh->GetSocketTransform(SourceSkelMesh->GetBoneName(Index), RTS_ParentBoneSpace);
+				BoneTransforms.Add(BoneTransform);
+			}
+		}
+	}
+
+	return BoneTransforms;
+}
+
+TArray<FName> ULiveLinkBroadcastComponent::GetCurveNames() const
+{
+	TArray<FName> CurveNames;
+
+	USkeletalMeshComponent* SourceSkelMesh = Cast<USkeletalMeshComponent>(SourceMesh.GetComponent(GetOwner()));
+
+	if (AllowedCurveNames.Num() > 0)
+	{
+		CurveNames = AllowedCurveNames;
+	}
+
+	return CurveNames;
+}
+
+TArray<float> ULiveLinkBroadcastComponent::GetCurveValues() const
+{
+	TArray<float> Curves;
+
+	USkeletalMeshComponent* SourceSkelMesh = Cast<USkeletalMeshComponent>(SourceMesh.GetComponent(GetOwner()));
+
+	TArray<FName> CurveNames = GetCurveNames();
+
+	//TODO filter just for Metahuman control curves. Or allow user to filter on prefix? 
+	for (FName Curve : CurveNames)
+	{
+		if (SourceSkelMesh->GetAnimInstance())
+		{
+			Curves.Add(SourceSkelMesh->GetAnimInstance()->GetCurveValue(Curve));
+		}
+	}
+
+	return Curves;
+}
+
+void ULiveLinkBroadcastComponent::BroadcastStaticData() const
+{
+	if (!Role)
+	{
+		return;
+	}
+
+	FLiveLinkStaticDataStruct StaticDataStruct;
+
+	if (Role->IsChildOf<ULiveLinkAnimationRole>())
+	{
+		StaticDataStruct = GetSkeletonStaticData();
+	}
+	else if (Role->IsChildOf<ULiveLinkCameraRole>())
+	{
+		StaticDataStruct = GetCameraStaticData();
+	}
+	else if (Role->IsChildOf<ULiveLinkTransformRole>())
+	{
+		StaticDataStruct = GetTransformStaticData();
+	}
+
+	ULiveLinkBroadcastSubsystem* BroadcastSubsystem = GEngine->GetEngineSubsystem<ULiveLinkBroadcastSubsystem>();
+	BroadcastSubsystem->BroadcastStaticData(SubjectKey, Role, MoveTemp(StaticDataStruct));
+}
+
+FLiveLinkStaticDataStruct ULiveLinkBroadcastComponent::GetSkeletonStaticData() const
+{
+	FLiveLinkStaticDataStruct StaticDataStruct(FLiveLinkSkeletonStaticData::StaticStruct());
+	FLiveLinkSkeletonStaticData& StaticData = *StaticDataStruct.Cast<FLiveLinkSkeletonStaticData>();
+
+	StaticData.BoneNames = GetSkeletalBoneNames();
+	StaticData.BoneParents = GetSkeletalBoneParents();
+	StaticData.PropertyNames = GetCurveNames();
+
+	return StaticDataStruct;
+}
+
+FLiveLinkFrameDataStruct ULiveLinkBroadcastComponent::GetSkeletonFrameData() const
+{
+	FLiveLinkFrameDataStruct FrameDataStruct(FLiveLinkAnimationFrameData::StaticStruct());
+	FLiveLinkAnimationFrameData& FrameData = *FrameDataStruct.Cast<FLiveLinkAnimationFrameData>();
+	FrameData.Transforms = GetBoneTransforms();
+	FrameData.PropertyValues = GetCurveValues();
+
+	return FrameDataStruct;
+}
+
+FLiveLinkStaticDataStruct ULiveLinkBroadcastComponent::GetCameraStaticData() const
+{
+	FLiveLinkStaticDataStruct StaticDataStruct(FLiveLinkCameraStaticData::StaticStruct());
+	FLiveLinkCameraStaticData& StaticData = *StaticDataStruct.Cast<FLiveLinkCameraStaticData>();
+
+	StaticData.bIsLocationSupported = true;
+	StaticData.bIsRotationSupported = true;
+	StaticData.bIsScaleSupported = true;
+
+	if (UCameraComponent* CameraComponent = Cast<UCameraComponent>(SourceMesh.GetComponent(GetOwner())))
+	{
+		if (UCineCameraComponent* CineCameraComponent = Cast<UCineCameraComponent>(CameraComponent))
+		{
+			// Todo: Detect changes on the camera component to automatically update the static data when it's modified.
+			StaticData.FilmBackWidth = CineCameraComponent->Filmback.SensorWidth;
+			StaticData.FilmBackHeight = CineCameraComponent->Filmback.SensorHeight;
+
+			StaticData.bIsFocalLengthSupported = true;
+			StaticData.bIsDepthOfFieldSupported = true;
+			StaticData.bIsApertureSupported = true;
+			StaticData.bIsFocusDistanceSupported = true;
+			StaticData.bIsDynamicFilmbackSupported = true;
+		}
+
+		StaticData.bIsFieldOfViewSupported = true;
+		StaticData.bIsAspectRatioSupported = true;
+		StaticData.bIsProjectionModeSupported = true;
+	}
+
+	return StaticDataStruct;
+}
+
+FLiveLinkFrameDataStruct ULiveLinkBroadcastComponent::GetCameraFrameData() const
+{
+	FLiveLinkFrameDataStruct FrameDataStruct(FLiveLinkCameraFrameData::StaticStruct());
+	FLiveLinkCameraFrameData& FrameData = *FrameDataStruct.Cast<FLiveLinkCameraFrameData>();
+
+	if (UCameraComponent* CameraComponent = Cast<UCameraComponent>(SourceMesh.GetComponent(GetOwner())))
+	{
+		FrameData.Transform = CameraComponent->GetComponentTransform();
+
+		FrameData.FieldOfView = CameraComponent->FieldOfView;
+		FrameData.AspectRatio = CameraComponent->AspectRatio;
+		FrameData.ProjectionMode = CameraComponent->ProjectionMode == ECameraProjectionMode::Perspective ? ELiveLinkCameraProjectionMode::Perspective  : ELiveLinkCameraProjectionMode::Orthographic;
+
+		if (UCineCameraComponent* CineCameraComponent = Cast<UCineCameraComponent>(CameraComponent))
+		{
+			FrameData.FocalLength = CineCameraComponent->CurrentFocalLength;
+			FrameData.Aperture = CineCameraComponent->CurrentAperture;
+			FrameData.FocusDistance = CineCameraComponent->CurrentFocusDistance;
+		}
+	}
+
+	return FrameDataStruct;
+}
+
+FLiveLinkStaticDataStruct ULiveLinkBroadcastComponent::GetTransformStaticData() const
+{
+	FLiveLinkStaticDataStruct StaticDataStruct(FLiveLinkTransformStaticData::StaticStruct());
+	FLiveLinkTransformStaticData& StaticData = *StaticDataStruct.Cast<FLiveLinkTransformStaticData>();
+
+	StaticData.bIsLocationSupported = true;
+	StaticData.bIsRotationSupported = true;
+	StaticData.bIsScaleSupported = true;
+
+	return StaticDataStruct;
+}
+
+FLiveLinkFrameDataStruct ULiveLinkBroadcastComponent::GetTransformFrameData() const
+{
+	FLiveLinkFrameDataStruct FrameDataStruct(FLiveLinkTransformFrameData::StaticStruct());
+	FLiveLinkTransformFrameData& FrameData = *FrameDataStruct.Cast<FLiveLinkTransformFrameData>();
+
+	if (USceneComponent* SceneComponent = Cast<USceneComponent>(SourceMesh.GetComponent(GetOwner())))
+	{
+		// Todo: Flag for relative vs absolute transform?
+		FrameData.Transform = SceneComponent->GetComponentTransform();
+	}
+
+	return FrameDataStruct;
+}
+
+void ULiveLinkBroadcastComponent::BroadcastFrameData() const
+{
+	if (!Role)
+	{
+		return;
+	}
+
+	FLiveLinkFrameDataStruct FrameDataStruct;
+
+	if (Role->IsChildOf<ULiveLinkAnimationRole>())
+	{
+		FrameDataStruct = GetSkeletonFrameData();
+	}
+	else if (Role->IsChildOf<ULiveLinkCameraRole>())
+	{
+		FrameDataStruct = GetCameraFrameData();
+	}
+	else if (Role->IsChildOf<ULiveLinkTransformRole>())
+	{
+		FrameDataStruct = GetTransformFrameData();
+	}
+
+	ULiveLinkBroadcastSubsystem* BroadcastSubsystem = GEngine->GetEngineSubsystem<ULiveLinkBroadcastSubsystem>();
+	BroadcastSubsystem->BroadcastFrameData(SubjectKey, MoveTemp(FrameDataStruct));
+}
diff --git a/Engine/Plugins/Animation/LiveLink/Source/LiveLink/Private/LiveLinkBroadcastSubsystem.cpp b/Engine/Plugins/Animation/LiveLink/Source/LiveLink/Private/LiveLinkBroadcastSubsystem.cpp
new file mode 100644
--- /dev/null
+++ b/Engine/Plugins/Animation/LiveLink/Source/LiveLink/Private/LiveLinkBroadcastSubsystem.cpp
@@ -0,0 +1,130 @@
+// Copyright Epic Games, Inc. All Rights Reserved.
+
+
+#include "LiveLinkBroadcastSubsystem.h"
+#include "Features/IModularFeatures.h"
+#include "ILiveLinkClient.h"
+#include "LiveLinkBroadcastComponent.h"
+#include "LiveLinkPresetTypes.h"
+#include "LiveLinkRole.h"
+#include "LiveLinkSubjectSettings.h"
+
+
+class FLiveLinkBroadcastSource : public ILiveLinkSource
+{
+public:
+	FLiveLinkBroadcastSource(const FString& SubjectName)
+	{
+		SourceName = *FString::Printf(TEXT("%s"), *SubjectName);
+	}
+
+	virtual ~FLiveLinkBroadcastSource() = default;
+
+	//~ Begin ILiveLinkSource interface
+	virtual bool CanBeDisplayedInUI() const override
+	{
+		return true;
+	}
+
+	virtual bool IsSourceStillValid() const override
+	{
+		return true;
+	}
+
+	virtual bool RequestSourceShutdown() override
+	{
+		return true;
+	}
+
+	virtual FText GetSourceType() const override
+	{
+		return FText::FromName(SourceName);
+	}
+
+	virtual FText GetSourceMachineName() const override
+	{
+		return FText().FromString(FPlatformProcess::ComputerName());
+	}
+
+	virtual FText GetSourceStatus() const override
+	{
+		return NSLOCTEXT("LiveLinkBroadcastSource", "BroadcastSourceStatus", "Active");
+	}
+
+	void ReceiveClient(ILiveLinkClient*, FGuid)
+	{
+	}
+
+protected:
+	/** Source name. */
+	FName SourceName;
+};
+
+FLiveLinkSubjectKey ULiveLinkBroadcastSubsystem::CreateSubject(FName SubjectName, TSubclassOf<ULiveLinkRole> Role)
+{
+	FLiveLinkSubjectKey SubjectKey;
+
+	if (IModularFeatures::Get().IsModularFeatureAvailable(ILiveLinkClient::ModularFeatureName))
+	{
+		ILiveLinkClient* LiveLinkClient = &IModularFeatures::Get().GetModularFeature<ILiveLinkClient>(ILiveLinkClient::ModularFeatureName);
+
+		if (!LiveLinkClient->IsSourceStillValid(VirtualSourceGUID))
+		{
+			TSharedRef<FLiveLinkBroadcastSource> NewSource = MakeShared<FLiveLinkBroadcastSource>(TEXT("Live Link Actor Broadcast"));
+			VirtualSourceGUID = LiveLinkClient->AddSource(NewSource);
+		}
+
+		SubjectKey = FLiveLinkSubjectKey{ VirtualSourceGUID, SubjectName };
+
+		// Check that the subject exists before recreating it.
+		if (!LiveLinkClient->GetSubjectSettings(SubjectKey))
+		{
+			FLiveLinkSubjectPreset Preset;
+			Preset.bEnabled = true;
+			Preset.Key = SubjectKey;
+			Preset.Role = Role;
+
+			LiveLinkClient->CreateSubject(Preset);
+			ULiveLinkSubjectSettings* Settings = CastChecked<ULiveLinkSubjectSettings>(LiveLinkClient->GetSubjectSettings(SubjectKey));
+			Settings->bRebroadcastSubject = true;
+		}
+	}
+
+	return SubjectKey;
+}
+
+void ULiveLinkBroadcastSubsystem::RemoveSubject(const FLiveLinkSubjectKey& SubjectKey)
+{
+	if (IModularFeatures::Get().IsModularFeatureAvailable(ILiveLinkClient::ModularFeatureName))
+	{
+		ILiveLinkClient* LiveLinkClient = &IModularFeatures::Get().GetModularFeature<ILiveLinkClient>(ILiveLinkClient::ModularFeatureName);
+		LiveLinkClient->RemoveSubject_AnyThread(SubjectKey);
+	}
+}
+
+void ULiveLinkBroadcastSubsystem::BroadcastStaticData(const FLiveLinkSubjectKey& SubjectKey, TSubclassOf<ULiveLinkRole> Role, FLiveLinkStaticDataStruct&& StaticData) const
+{
+	if (IModularFeatures::Get().IsModularFeatureAvailable(ILiveLinkClient::ModularFeatureName))
+	{
+		ILiveLinkClient* LiveLinkClient = &IModularFeatures::Get().GetModularFeature<ILiveLinkClient>(ILiveLinkClient::ModularFeatureName);
+
+		// Don't recreate the source if it's been explicitly deleted.
+		if (LiveLinkClient->IsSourceStillValid(VirtualSourceGUID))
+		{
+			LiveLinkClient->PushSubjectStaticData_AnyThread(SubjectKey, Role, MoveTemp(StaticData));
+		}
+	}
+}
+
+void ULiveLinkBroadcastSubsystem::BroadcastFrameData(const FLiveLinkSubjectKey& SubjectKey, FLiveLinkFrameDataStruct&& FrameData) const
+{
+	if (IModularFeatures::Get().IsModularFeatureAvailable(ILiveLinkClient::ModularFeatureName))
+	{
+		ILiveLinkClient* LiveLinkClient = &IModularFeatures::Get().GetModularFeature<ILiveLinkClient>(ILiveLinkClient::ModularFeatureName);
+
+		if (LiveLinkClient->IsSourceStillValid(VirtualSourceGUID))
+		{
+			LiveLinkClient->PushSubjectFrameData_AnyThread(SubjectKey, MoveTemp(FrameData));
+		}
+	}
+}
diff --git a/Engine/Plugins/Animation/LiveLink/Source/LiveLink/Private/LiveLinkBroadcastSubsystem.h b/Engine/Plugins/Animation/LiveLink/Source/LiveLink/Private/LiveLinkBroadcastSubsystem.h
new file mode 100644
--- /dev/null
+++ b/Engine/Plugins/Animation/LiveLink/Source/LiveLink/Private/LiveLinkBroadcastSubsystem.h
@@ -0,0 +1,37 @@
+// Copyright Epic Games, Inc. All Rights Reserved.
+
+#pragma once
+
+#include "Subsystems/EngineSubsystem.h"
+#include "LiveLinkTypes.h"
+
+#include "LiveLinkBroadcastSubsystem.generated.h"
+
+class FLiveLinkBroadcastSource;
+class ULiveLinkRole;
+
+/**
+ * Subsystem used primarily by the LiveLinkBroadcastComponent to rebroadcast a skeletal mesh on the network.
+ */
+UCLASS()
+class ULiveLinkBroadcastSubsystem : public UEngineSubsystem
+{
+	GENERATED_BODY()
+	
+public:
+	/** Create a broadcast subject with a given subject name. */
+	FLiveLinkSubjectKey CreateSubject(FName SubjectName, TSubclassOf<ULiveLinkRole> Role);
+
+	/** Remove a LiveLink s */
+	void RemoveSubject(const FLiveLinkSubjectKey& SubjectKey);
+
+	/** Broadcast static data for a subject. */
+	void BroadcastStaticData(const FLiveLinkSubjectKey& SubjectKey, TSubclassOf<ULiveLinkRole> Role, FLiveLinkStaticDataStruct&& StaticData) const;
+
+	/** Broadcast frame data for a subject. */
+	void BroadcastFrameData(const FLiveLinkSubjectKey& SubjectKey, FLiveLinkFrameDataStruct&& FrameData) const;
+
+private:
+	/** ID of the source used to rebroadcast the data. */
+	FGuid VirtualSourceGUID;
+};
diff --git a/Engine/Plugins/Animation/LiveLink/Source/LiveLink/Private/LiveLinkClient.cpp b/Engine/Plugins/Animation/LiveLink/Source/LiveLink/Private/LiveLinkClient.cpp
--- a/Engine/Plugins/Animation/LiveLink/Source/LiveLink/Private/LiveLinkClient.cpp
+++ b/Engine/Plugins/Animation/LiveLink/Source/LiveLink/Private/LiveLinkClient.cpp
@@ -879,6 +879,9 @@ void FLiveLinkClient::PushSubjectStaticData_Internal(FPendingSubjectStatic&& Sub
 		return;
 	}
 
+	// We use a strong object ptr to prevent GC for this object in case this was created outside the game thread.
+	TStrongObjectPtr<ULiveLinkSubjectSettings> SubjectSettings;
+
 	FLiveLinkSubject* LiveLinkSubject = nullptr;
 	{
 		if (FLiveLinkCollectionSubjectItem* SubjectItem = Collection->FindSubject(SubjectStaticData.SubjectKey))
@@ -891,6 +894,8 @@ void FLiveLinkClient::PushSubjectStaticData_Internal(FPendingSubjectStatic&& Sub
 				{
 					FLiveLinkLog::Warning(TEXT("Subject '%s' of role '%s' is changing its role to '%s'. Current subject will be removed and a new one will be created"), *SubjectStaticData.SubjectKey.SubjectName.ToString(), *LiveLinkSubject->GetRole().GetDefaultObject()->GetDisplayName().ToString(), *SubjectStaticData.Role.GetDefaultObject()->GetDisplayName().ToString());
 
+					// Make sure to keep the previous settings even if we're changing the role.
+					SubjectSettings = TStrongObjectPtr{ Cast<ULiveLinkSubjectSettings>(GetSubjectSettings(SubjectItem->Key)) };
 					Collection->RemoveSubject(SubjectStaticData.SubjectKey);
 					LiveLinkSubject = nullptr;
 				}
@@ -911,18 +916,20 @@ void FLiveLinkClient::PushSubjectStaticData_Internal(FPendingSubjectStatic&& Sub
 		const ULiveLinkSettings* LiveLinkSettings = GetDefault<ULiveLinkSettings>();
 		const FLiveLinkRoleProjectSetting DefaultSetting = LiveLinkSettings->GetDefaultSettingForRole(SubjectStaticData.Role.Get());
 
-		// We use a strong object ptr to prevent GC for this object in case this was created outside the game thread.
-		TStrongObjectPtr<ULiveLinkSubjectSettings> SubjectSettings = nullptr;
-
-		// Setting class should always be valid
 		{
-			UClass* SettingClass = DefaultSetting.SettingClass.Get();
-			if (SettingClass == nullptr)
+			// By default, try to reuse the previous subject settings when possible. The role and interpolation will be replaced but the other settings (such as bRebroadcastSubject) will be kept.
+			if (!SubjectSettings)
 			{
-				SettingClass = ULiveLinkSubjectSettings::StaticClass();
+				// Setting class should always be valid
+				UClass* SettingClass = DefaultSetting.SettingClass.Get();
+				if (SettingClass == nullptr)
+				{
+					SettingClass = ULiveLinkSubjectSettings::StaticClass();
+				}
+
+				SubjectSettings = TStrongObjectPtr{ NewObject<ULiveLinkSubjectSettings>(GetTransientPackage(), SettingClass) };
 			}
 
-			SubjectSettings = TStrongObjectPtr{NewObject<ULiveLinkSubjectSettings>(GetTransientPackage(), SettingClass)};
 			SubjectSettings->Initialize(SubjectStaticData.SubjectKey);
 
 			SubjectSettings->Role = SubjectStaticData.Role;
diff --git a/Engine/Plugins/Animation/LiveLink/Source/LiveLink/Public/LiveLinkBroadcastComponent.h b/Engine/Plugins/Animation/LiveLink/Source/LiveLink/Public/LiveLinkBroadcastComponent.h
new file mode 100644
--- /dev/null
+++ b/Engine/Plugins/Animation/LiveLink/Source/LiveLink/Public/LiveLinkBroadcastComponent.h
@@ -0,0 +1,123 @@
+// Copyright Epic Games, Inc. All Rights Reserved.
+
+#pragma once
+
+#include "CoreMinimal.h"
+#include "ILiveLinkSource.h"
+#include "LiveLinkPresetTypes.h"
+#include "Components/ActorComponent.h"
+#include "Templates/SharedPointer.h"
+
+#include "LiveLinkBroadcastComponent.generated.h"
+
+
+class FLiveLinkBroadcastSource;
+
+UCLASS(BlueprintType, DisplayName="Live Link Broadcast", ClassGroup=("LiveLink"), meta=(BlueprintSpawnableComponent), HideCategories=(Navigation, Tags, AssetUserData, Activation))
+class ULiveLinkBroadcastComponent : public UActorComponent 
+{
+	GENERATED_BODY()
+
+public:
+	// Sets default values for this component's properties
+	ULiveLinkBroadcastComponent();
+
+	/** Whether this component should broadcast data. */
+	UPROPERTY(BlueprintReadWrite, EditAnywhere, Category = "Live Link Broadcast")
+	bool bEnable = true;
+
+	/**
+	 * The name of the Live Link subject. Will default to the actor name if left blank
+	 */
+	UPROPERTY(BlueprintReadWrite, EditAnywhere, Category = "Live Link Broadcast")
+	FName SubjectName;
+
+	/** 
+	 * What role should be used for evaluating the data.
+	 * Choosing Transform Role will only send the transform of the chosen component, while picking Animation will send the full skeleton transforms.
+	 **/
+	UPROPERTY(BlueprintReadWrite, EditAnywhere, Category = "Live Link Broadcast")
+	TSubclassOf<ULiveLinkRole> Role;
+
+	/**
+	 * The source mesh to broadcast
+	 * todo: Use customization to restrict the allowed classes metadata
+	 */
+	UPROPERTY(BlueprintReadWrite, EditAnywhere, Category = "Live Link Broadcast", meta=(UseComponentPicker))
+	FComponentReference SourceMesh;
+
+	/**
+	 * Optional list of allowed bone names (useful for Metahumans and other complex assets).
+	 */
+	UPROPERTY(BlueprintReadWrite, EditAnywhere, Category = "Live Link Broadcast",  meta=(ToolTip="Use this array to provide an inclusive list of bones to broadcast"))
+	TArray<FName> AllowedBoneNames;
+
+	/**
+	 * Optional list of allowed curve names (useful for Metahumans and other complex assets).
+	 */
+	UPROPERTY(BlueprintReadWrite, EditAnywhere, Category = "Live Link Broadcast", meta=(ToolTip="Curve data will be added to the broadcast data if listed here"))
+	TArray<FName> AllowedCurveNames;
+
+	//~ UActorComponent interface
+	LIVELINK_API virtual void TickComponent(float DeltaTime, ELevelTick TickType, FActorComponentTickFunction* ThisTickFunction) override;
+
+private:
+	/** Creates the rebroadcast source if it doesn't exist. */
+	void InitiateLiveLink();
+
+	//~ Populates a static data struct from the chosen mesh
+	FLiveLinkStaticDataStruct GetSkeletonStaticData() const;
+	FLiveLinkStaticDataStruct GetCameraStaticData() const;
+	FLiveLinkStaticDataStruct GetTransformStaticData() const;
+
+	//~ Populates a frame data struct from the chosen mesh
+	FLiveLinkFrameDataStruct GetSkeletonFrameData() const;
+	FLiveLinkFrameDataStruct GetCameraFrameData() const;
+	FLiveLinkFrameDataStruct GetTransformFrameData() const;
+
+	/** Get bone names from the skekelal mesh. */
+	TArray<FName> GetSkeletalBoneNames() const;
+
+	/** Get bone parents from the skekelal mesh. */
+	TArray<int32> GetSkeletalBoneParents() const;
+
+	/** Get the bone transforms from the mesh. */
+	TArray<FTransform> GetBoneTransforms() const;
+
+	/** Get the curve names from the mesh. */
+	TArray<FName> GetCurveNames() const;
+
+	/** Get the curve values from the mesh. */
+	TArray<float> GetCurveValues() const;
+
+	/** Send the static data to the broadcast subsystem. */
+	void BroadcastStaticData() const;
+
+	/** Send the frame data to the broadcast subsystem. */
+	void BroadcastFrameData() const;
+	
+protected:
+	//~ Begin UActorComponent interface
+	virtual void OnComponentCreated() override;
+	virtual void OnComponentDestroyed(bool bDestroyingHierarchy) override;
+	//~ End UActorComponent interface
+	
+#if WITH_EDITOR
+	//~ Begin UObject interface 
+	virtual void PostEditChangeProperty(struct FPropertyChangedEvent& PropertyChangedEvent) override;
+	//~ End UObject interface 
+
+	/** Object modified callback used to trigger the dirty flag if one of the properties on the tracked source object is modified. */
+	void OnObjectModified(UObject* ModifiedObject);
+#endif
+
+private:
+	/** Cached resolved object, updated when bIsDirty gets triggered. */
+	class UActorComponent* CachedResolvedObject = nullptr;
+	/** Dirty flag used to trigger updates on the broadcast subsystem. */
+	bool bIsDirty;
+	/** Cached subject key for this component. */
+	FLiveLinkSubjectKey SubjectKey;
+};
+
+
diff --git a/Engine/Plugins/Animation/LiveLink/Source/LiveLinkEditor/Private/LiveLinkBroadcastComponentDetailCustomization.cpp b/Engine/Plugins/Animation/LiveLink/Source/LiveLinkEditor/Private/LiveLinkBroadcastComponentDetailCustomization.cpp
new file mode 100644
--- /dev/null
+++ b/Engine/Plugins/Animation/LiveLink/Source/LiveLinkEditor/Private/LiveLinkBroadcastComponentDetailCustomization.cpp
@@ -0,0 +1,40 @@
+// Copyright Epic Games, Inc. All Rights Reserved.
+
+#include "LiveLinkBroadcastComponentDetailCustomization.h"
+
+#include "LiveLinkBroadcastComponent.h"
+#include "DetailLayoutBuilder.h"
+#include "Roles/LiveLinkAnimationRole.h"
+
+void FLiveLinkBroadcastComponentDetailCustomization::CustomizeDetails(IDetailLayoutBuilder& DetailBuilder)
+{
+	DetailLayout = &DetailBuilder;
+
+	TArray<TWeakObjectPtr<UObject>> SelectedObjects = DetailBuilder.GetSelectedObjects();
+
+	// Hide everything when more than one are selected
+	if (SelectedObjects.Num() != 1 || !SelectedObjects[0].IsValid())
+	{
+		return;
+	}
+
+	ULiveLinkBroadcastComponent* Component = CastChecked<ULiveLinkBroadcastComponent>(SelectedObjects[0].Get());
+
+	TSharedRef<IPropertyHandle> RoleProperty = DetailBuilder.GetProperty(GET_MEMBER_NAME_CHECKED(ULiveLinkBroadcastComponent, Role));
+	RoleProperty->SetOnPropertyValueChanged(FSimpleDelegate::CreateRaw(DetailLayout, &IDetailLayoutBuilder::ForceRefreshDetails));
+
+	TSharedRef<IPropertyHandle> SourceMeshProperty = DetailBuilder.GetProperty(GET_MEMBER_NAME_CHECKED(ULiveLinkBroadcastComponent, SourceMesh));
+	TSharedRef<IPropertyHandle> AllowedBonesProperty = DetailBuilder.GetProperty(GET_MEMBER_NAME_CHECKED(ULiveLinkBroadcastComponent, AllowedBoneNames));
+	TSharedRef<IPropertyHandle> AllowedCurvesProperty = DetailBuilder.GetProperty(GET_MEMBER_NAME_CHECKED(ULiveLinkBroadcastComponent, AllowedCurveNames));
+
+	if (Component->Role == ULiveLinkAnimationRole::StaticClass())
+	{
+		SourceMeshProperty->SetInstanceMetaData("AllowedClasses", TEXT("/Script/Engine.SkeletalMeshComponent"));
+	}
+	else
+	{
+		AllowedBonesProperty->MarkHiddenByCustomization();
+		AllowedCurvesProperty->MarkHiddenByCustomization();
+		SourceMeshProperty->SetInstanceMetaData("AllowedClasses", TEXT(""));
+	}
+}
diff --git a/Engine/Plugins/Animation/LiveLink/Source/LiveLinkEditor/Private/LiveLinkBroadcastComponentDetailCustomization.h b/Engine/Plugins/Animation/LiveLink/Source/LiveLinkEditor/Private/LiveLinkBroadcastComponentDetailCustomization.h
new file mode 100644
--- /dev/null
+++ b/Engine/Plugins/Animation/LiveLink/Source/LiveLinkEditor/Private/LiveLinkBroadcastComponentDetailCustomization.h
@@ -0,0 +1,28 @@
+// Copyright Epic Games, Inc. All Rights Reserved.
+
+#pragma once
+
+#include "IDetailCustomization.h"
+
+class IDetailLayoutBuilder;
+
+/**
+* Customizes a LiveLinkBroadcastComponent details.
+*/
+class FLiveLinkBroadcastComponentDetailCustomization : public IDetailCustomization
+{
+public:
+
+	static TSharedRef<IDetailCustomization> MakeInstance()
+	{
+		return MakeShared<FLiveLinkBroadcastComponentDetailCustomization>();
+	}
+
+	// IDetailCustomization interface
+	virtual void CustomizeDetails(class IDetailLayoutBuilder& DetailBuilder) override;
+	// End IDetailCustomization interface
+
+protected:
+	/** Keep a reference to force refresh the layout */
+	IDetailLayoutBuilder* DetailLayout = nullptr;
+};
diff --git a/Engine/Plugins/Animation/LiveLink/Source/LiveLinkEditor/Private/LiveLinkEditorModule.cpp b/Engine/Plugins/Animation/LiveLink/Source/LiveLinkEditor/Private/LiveLinkEditorModule.cpp
--- a/Engine/Plugins/Animation/LiveLink/Source/LiveLinkEditor/Private/LiveLinkEditorModule.cpp
+++ b/Engine/Plugins/Animation/LiveLink/Source/LiveLinkEditor/Private/LiveLinkEditorModule.cpp
@@ -20,6 +20,8 @@
 #include "WorkspaceMenuStructureModule.h"
 
 #include "LiveLinkBoneAttachmentDetailCustomization.h"
+#include "LiveLinkBroadcastComponent.h"
+#include "LiveLinkBroadcastComponentDetailCustomization.h"
 #include "LiveLinkComponentController.h"
 #include "LiveLinkComponentDetailCustomization.h"
 #include "LiveLinkComponentSettings.h"
@@ -298,6 +300,7 @@ private:
 		PropertyEditorModule.RegisterCustomPropertyTypeLayout(FLiveLinkVirtualSubjectBoneAttachment::StaticStruct()->GetFName(), FOnGetPropertyTypeCustomizationInstance::CreateStatic(&FLiveLinkBoneAttachmentDetailCustomization::MakeInstance));
 		PropertyEditorModule.RegisterCustomClassLayout(ULiveLinkSourceSettings::StaticClass()->GetFName(), FOnGetDetailCustomizationInstance::CreateStatic(&FLiveLinkSourceSettingsDetailCustomization::MakeInstance));
 		PropertyEditorModule.RegisterCustomClassLayout(ULiveLinkComponentController::StaticClass()->GetFName(), FOnGetDetailCustomizationInstance::CreateStatic(&FLiveLinkComponentDetailCustomization::MakeInstance));
+		PropertyEditorModule.RegisterCustomClassLayout(ULiveLinkBroadcastComponent::StaticClass()->GetFName(), FOnGetDetailCustomizationInstance::CreateStatic(&FLiveLinkBroadcastComponentDetailCustomization::MakeInstance));
 		PropertyEditorModule.RegisterCustomClassLayout(ULiveLinkControllerBase::StaticClass()->GetFName(), FOnGetDetailCustomizationInstance::CreateStatic(&FLiveLinkControllerBaseDetailCustomization::MakeInstance));
 		PropertyEditorModule.RegisterCustomClassLayout(ULiveLinkHubSubjectSettings::StaticClass()->GetFName(), FOnGetDetailCustomizationInstance::CreateStatic(&FLiveLinkHubSubjectSettingsDetailsCustomization::MakeInstance));
 
@@ -314,6 +317,7 @@ private:
 			if (PropertyEditorModule)
 			{
 				PropertyEditorModule->UnregisterCustomClassLayout(ULiveLinkHubSubjectSettings::StaticClass()->GetFName());
+				PropertyEditorModule->UnregisterCustomClassLayout(ULiveLinkBroadcastComponent::StaticClass()->GetFName());
 				PropertyEditorModule->UnregisterCustomClassLayout(ULiveLinkComponentController::StaticClass()->GetFName());
 				PropertyEditorModule->UnregisterCustomClassLayout(ULiveLinkControllerBase::StaticClass()->GetFName());
 				PropertyEditorModule->UnregisterCustomClassLayout(ULiveLinkSourceSettings::StaticClass()->GetFName());
