LiveLinkHub - Fix Virtual subjects not rebroadcasting if TransmitEvaluatedData was disabled

diff --git a/Engine/Plugins/Animation/LiveLink/Source/LiveLink/Private/LiveLinkClient.cpp b/Engine/Plugins/Animation/LiveLink/Source/LiveLink/Private/LiveLinkClient.cpp
--- a/Engine/Plugins/Animation/LiveLink/Source/LiveLink/Private/LiveLinkClient.cpp
+++ b/Engine/Plugins/Animation/LiveLink/Source/LiveLink/Private/LiveLinkClient.cpp
@@ -252,7 +252,7 @@ TSet<FLiveLinkSubjectKey> FLiveLinkClient::BuildThisTicksSubjectSnapshot(const T
 
 				if (bSubjectIsSynchronized)
 				{
-					if (SubjectItem.bEnabled && !VSubject->IsPaused() && SourceItem.Setting->bTransmitEvaluatedData)
+					if (SubjectItem.bEnabled && !VSubject->IsPaused())
 					{
 						for (const FLiveLinkSubjectKey& ReceivedSubject : ReceivedSubjects)
 						{
diff --git a/Engine/Plugins/Animation/LiveLink/Source/LiveLinkEditor/Private/LiveLinkSourceSettingsDetailCustomization.cpp b/Engine/Plugins/Animation/LiveLink/Source/LiveLinkEditor/Private/LiveLinkSourceSettingsDetailCustomization.cpp
--- a/Engine/Plugins/Animation/LiveLink/Source/LiveLinkEditor/Private/LiveLinkSourceSettingsDetailCustomization.cpp
+++ b/Engine/Plugins/Animation/LiveLink/Source/LiveLinkEditor/Private/LiveLinkSourceSettingsDetailCustomization.cpp
@@ -47,6 +47,8 @@ void FLiveLinkSourceSettingsDetailCustomization::CustomizeDetails(IDetailLayoutB
 		}
 	}
 
+	const bool bIsVirtualSource = LiveLinkClient.GetClient()->GetVirtualSources().Contains(EditedSourceGuid) || EditedSourceGuid == FGuid();
+
 	TSharedPtr<IPropertyHandle> BufferSettingsPropertyHandle = InDetailBuilder.GetProperty(GET_MEMBER_NAME_CHECKED(ULiveLinkSourceSettings, BufferSettings));
 	InDetailBuilder.HideProperty(BufferSettingsPropertyHandle);
 
@@ -65,8 +67,16 @@ void FLiveLinkSourceSettingsDetailCustomization::CustomizeDetails(IDetailLayoutB
 	const bool bIsInLiveLinkHubApp = GConfig->GetBoolOrDefault(TEXT("LiveLink"), TEXT("bCreateLiveLinkHubInstance"), false, GEngineIni);
 	if (bIsInLiveLinkHubApp)
 	{
-		bSupportsBuffering = SourceSettings->bTransmitEvaluatedData;
-		TransmitEvalutedDataHandle->SetOnPropertyValueChanged(FSimpleDelegate::CreateSP(this, &FLiveLinkSourceSettingsDetailCustomization::ForceRefresh));
+		if (bIsVirtualSource)
+		{
+			// Virtual subject transmit evaluated data by default
+			InDetailBuilder.HideProperty(TransmitEvalutedDataHandle);
+		}
+		else
+		{
+			bSupportsBuffering = SourceSettings->bTransmitEvaluatedData;
+			TransmitEvalutedDataHandle->SetOnPropertyValueChanged(FSimpleDelegate::CreateSP(this, &FLiveLinkSourceSettingsDetailCustomization::ForceRefresh));
+		}
 	}
 	else
 	{
