LiveLink - Fix crash when calling ClearSubjectFrames on a subject while it's receiving data.

Added a ApplyToSubject method that allows modifying the subject in a a thread safe way

diff --git a/Engine/Plugins/Animation/LiveLink/Source/LiveLink/Private/LiveLinkClient.cpp b/Engine/Plugins/Animation/LiveLink/Source/LiveLink/Private/LiveLinkClient.cpp
--- a/Engine/Plugins/Animation/LiveLink/Source/LiveLink/Private/LiveLinkClient.cpp
+++ b/Engine/Plugins/Animation/LiveLink/Source/LiveLink/Private/LiveLinkClient.cpp
@@ -1466,12 +1466,10 @@ void FLiveLinkClient::ClearSubjectsFrames_AnyThread(const FLiveLinkSubjectKey& I
 {
 	if (Collection)
 	{
-		if (FLiveLinkCollectionSubjectItem* SubjectItem = Collection->FindSubject(InSubjectKey))
+		Collection->ApplyToSubject(InSubjectKey, [](FLiveLinkCollectionSubjectItem& SubjectItem) 
 		{
-			SubjectItem->GetSubject()->ClearFrames();
-
-
-		}
+			SubjectItem.GetSubject()->ClearFrames();
+		});
 	}
 }
 
diff --git a/Engine/Plugins/Animation/LiveLink/Source/LiveLink/Private/LiveLinkSourceCollection.cpp b/Engine/Plugins/Animation/LiveLink/Source/LiveLink/Private/LiveLinkSourceCollection.cpp
--- a/Engine/Plugins/Animation/LiveLink/Source/LiveLink/Private/LiveLinkSourceCollection.cpp
+++ b/Engine/Plugins/Animation/LiveLink/Source/LiveLink/Private/LiveLinkSourceCollection.cpp
@@ -452,6 +452,15 @@ void FLiveLinkSourceCollection::ForEachSource(TFunctionRef<void(const FLiveLinkC
 	}
 }
 
+void FLiveLinkSourceCollection::ApplyToSubject(const FLiveLinkSubjectKey& InSubjectKey, TFunctionRef<void(FLiveLinkCollectionSubjectItem& /*SubjectItem*/)> VisitorFunc)
+{
+	if (FLiveLinkCollectionSubjectItem* Subject = FindSubject(InSubjectKey))
+	{
+		UE::TUniqueLock Lock(SubjectsLock);
+		VisitorFunc(*Subject);
+	}
+}
+
 void FLiveLinkSourceCollection::HandleSubjectStateChanged(ELiveLinkSubjectState NewState, FLiveLinkSubjectKey SubjectKey)
 {
 	BroadcastOnGameThread(OnLiveLinkSubjectStateChanged(), SubjectKey, NewState);
diff --git a/Engine/Plugins/Animation/LiveLink/Source/LiveLink/Public/LiveLinkSourceCollection.h b/Engine/Plugins/Animation/LiveLink/Source/LiveLink/Public/LiveLinkSourceCollection.h
--- a/Engine/Plugins/Animation/LiveLink/Source/LiveLink/Public/LiveLinkSourceCollection.h
+++ b/Engine/Plugins/Animation/LiveLink/Source/LiveLink/Public/LiveLinkSourceCollection.h
@@ -146,6 +146,11 @@ public:
 	 */
 	UE_API void ForEachSource(TFunctionRef<void(const FLiveLinkCollectionSourceItem&)> VisitorFunc) const;
 
+	/**
+	 * Thread safe way to apply a method over a subject.
+	 */
+	UE_API void ApplyToSubject(const FLiveLinkSubjectKey& InSubjectKey, TFunctionRef<void(FLiveLinkCollectionSubjectItem& /*SubjectItem*/)> VisitorFunc);
+
 	FSimpleMulticastDelegate& OnLiveLinkSourcesChanged() { return OnLiveLinkSourcesChangedDelegate; }
 	FSimpleMulticastDelegate& OnLiveLinkSubjectsChanged() { return OnLiveLinkSubjectsChangedDelegate; }
 
