LiveLink - Add scroll bar to LiveLink sources and subjects

#rnx

diff --git a/Engine/Plugins/Animation/LiveLink/Source/LiveLinkEditor/Private/LiveLinkClientPanelViews.cpp b/Engine/Plugins/Animation/LiveLink/Source/LiveLinkEditor/Private/LiveLinkClientPanelViews.cpp
--- a/Engine/Plugins/Animation/LiveLink/Source/LiveLinkEditor/Private/LiveLinkClientPanelViews.cpp
+++ b/Engine/Plugins/Animation/LiveLink/Source/LiveLinkEditor/Private/LiveLinkClientPanelViews.cpp
@@ -836,11 +836,12 @@ void FLiveLinkSourcesView::CreateSourcesListView(const TSharedPtr<FUICommandList
 		]
 	]
 	+SVerticalBox::Slot()
-	.AutoHeight()
+	.VAlign(VAlign_Fill)
 	[
 		SAssignNew(SourcesListView, SLiveLinkSourceListView, bReadOnly)
 			.ListItemsSource(&FilteredList)
 			.SelectionMode(ESelectionMode::Single)
+			.ScrollbarVisibility(EVisibility::Visible)
 			.OnGenerateRow_Raw(this, &FLiveLinkSourcesView::MakeSourceListViewWidget)
 			.OnContextMenuOpening_Raw(this, &FLiveLinkSourcesView::OnSourceConstructContextMenu, InCommandList)
 			.OnSelectionChanged_Raw(this, &FLiveLinkSourcesView::OnSourceListSelectionChanged)
@@ -1121,7 +1122,7 @@ TSharedRef<SWidget> FLiveLinkSubjectsView::GetWidget()
 				]
 		]
 		+ SVerticalBox::Slot()
-		.AutoHeight()
+		.VAlign(VAlign_Fill)
 		[
 			SubjectsTreeView.ToSharedRef()
 		];
