Live Link subject list multi-selection issue

diff --git a/Engine/Plugins/Animation/LiveLink/Source/LiveLinkEditor/Private/LiveLinkClientPanelViews.cpp b/Engine/Plugins/Animation/LiveLink/Source/LiveLinkEditor/Private/LiveLinkClientPanelViews.cpp
--- a/Engine/Plugins/Animation/LiveLink/Source/LiveLinkEditor/Private/LiveLinkClientPanelViews.cpp
+++ b/Engine/Plugins/Animation/LiveLink/Source/LiveLinkEditor/Private/LiveLinkClientPanelViews.cpp
@@ -991,12 +991,12 @@ bool FLiveLinkSubjectsView::CanRemoveSubject() const
 
 void FLiveLinkSubjectsView::RefreshSubjects()
 {
-	TArray<FLiveLinkSubjectKey> SavedSelection;
+	TArray<TPair<FLiveLinkSubjectKey, bool>> SavedSelection;
 	{
 		TArray<FLiveLinkSubjectUIEntryPtr> SelectedItems = SubjectsTreeView->GetSelectedItems();
 		for (const FLiveLinkSubjectUIEntryPtr& SelectedItem : SelectedItems)
 		{
-			SavedSelection.Add(SelectedItem->SubjectKey);
+			SavedSelection.Add(TPair<FLiveLinkSubjectKey, bool>(SelectedItem->SubjectKey, SelectedItem->IsSubject()));
 		}
 	}
 
@@ -1047,9 +1047,9 @@ void FLiveLinkSubjectsView::RefreshSubjects()
 
 			for (const FLiveLinkSubjectUIEntryPtr& Item : AllItems)
 			{
-				for (FLiveLinkSubjectKey& Selection : SavedSelection)
+				for (const TPair<FLiveLinkSubjectKey, bool>& Selection : SavedSelection)
 				{
-					if (Item->SubjectKey == Selection)
+					if (Item->SubjectKey == Selection.Key && Item->IsSubject() == Selection.Value)
 					{
 						SubjectsTreeView->SetItemSelection(Item, true);
 						break;
