LiveLinkHub - Add tooltip to SyncSubject

diff --git a/Engine/Plugins/Animation/LiveLink/Source/LiveLinkEditor/Private/LiveLinkVirtualSubjectDetailCustomization.cpp b/Engine/Plugins/Animation/LiveLink/Source/LiveLinkEditor/Private/LiveLinkVirtualSubjectDetailCustomization.cpp
--- a/Engine/Plugins/Animation/LiveLink/Source/LiveLinkEditor/Private/LiveLinkVirtualSubjectDetailCustomization.cpp
+++ b/Engine/Plugins/Animation/LiveLink/Source/LiveLinkEditor/Private/LiveLinkVirtualSubjectDetailCustomization.cpp
@@ -130,6 +130,8 @@ void FLiveLinkVirtualSubjectDetailCustomization::CustomizeDetails(IDetailLayoutB
 
 	const FResetToDefaultOverride Override = FResetToDefaultOverride::Create(IsResetToDefaultVisible, ResetToDefaultHandler);
 
+	const FText ToolTipText = LOCTEXT("SyncSubjectToolTipText", "Sync subject will be used to decide when to send this Virtual Subject's data. For smoothest result, choose the subject with the highest rate of update.");
+
 	SyncSubjectPropertyGroup.AddCustomRow(LOCTEXT("SyncSubjectLabel", "Sync Subject"))
 		.PropertyHandleList({ SyncSubjectProperty })
 		.NameContent()
@@ -137,6 +139,7 @@ void FLiveLinkVirtualSubjectDetailCustomization::CustomizeDetails(IDetailLayoutB
 			SNew(STextBlock)
 				.Font(IDetailLayoutBuilder::GetDetailFont())
 				.Text(LOCTEXT("SyncSubjectLabel", "Sync Subject"))
+				.ToolTipText(ToolTipText)
 		]
 		.ValueContent()
 		[
@@ -147,6 +150,7 @@ void FLiveLinkVirtualSubjectDetailCustomization::CustomizeDetails(IDetailLayoutB
 				.OnGetSubjects(this, &FLiveLinkVirtualSubjectDetailCustomization::OnGetSubjects)
 				.Value(this, &FLiveLinkVirtualSubjectDetailCustomization::GetSyncSubject)
 				.OnValueChanged(this, &FLiveLinkVirtualSubjectDetailCustomization::SetSyncSubject)
+				.ToolTipText(ToolTipText)
 		]
 		.OverrideResetToDefault(Override);
 }
