LiveLink - Add multi-select for subjects, allow pausing multiple subjects at once

diff --git a/Engine/Plugins/Animation/LiveLink/Source/LiveLinkEditor/Private/LiveLinkClientPanelViews.cpp b/Engine/Plugins/Animation/LiveLink/Source/LiveLinkEditor/Private/LiveLinkClientPanelViews.cpp
--- a/Engine/Plugins/Animation/LiveLink/Source/LiveLinkEditor/Private/LiveLinkClientPanelViews.cpp
+++ b/Engine/Plugins/Animation/LiveLink/Source/LiveLinkEditor/Private/LiveLinkClientPanelViews.cpp
@@ -913,6 +913,7 @@ void FLiveLinkSourcesView::HandleRemoveSource()
 {
 	TArray<FLiveLinkSourceUIEntryPtr> Selected;
 	SourcesListView->GetSelectedItems(Selected);
+
 	if (Selected.Num() > 0)
 	{
 		Selected[0]->RemoveFromClient();
@@ -986,7 +987,16 @@ bool FLiveLinkSubjectsView::CanRemoveSubject() const
 {
 	TArray<FLiveLinkSubjectUIEntryPtr> Selected;
 	SubjectsTreeView->GetSelectedItems(Selected);
-	return Selected.Num() > 0 && Selected[0] && Selected[0]->IsVirtualSubject();
+
+	for (const FLiveLinkSubjectUIEntryPtr& Entry : Selected)
+	{
+		if (Entry && Entry->IsVirtualSubject())
+		{
+			return true;
+		}
+	}
+
+	return false;
 }
 
 void FLiveLinkSubjectsView::RefreshSubjects()
@@ -1070,16 +1080,28 @@ bool FLiveLinkSubjectsView::CanPauseSubject() const
 	TArray<FLiveLinkSubjectUIEntryPtr> Selected;
 	SubjectsTreeView->GetSelectedItems(Selected);
 
-	return Selected.Num() > 0 && Selected[0] && Selected[0]->IsSubjectValid() && Selected[0]->IsSubjectEnabled();
+	for (const FLiveLinkSubjectUIEntryPtr& Entry : Selected)
+	{
+		if (Entry && Entry->IsSubjectValid() && Entry->IsSubjectEnabled())
+		{
+			return true;
+		}
+	}
+
+	return false;
 }
 
 void FLiveLinkSubjectsView::HandlePauseSubject()
 {
 	TArray<FLiveLinkSubjectUIEntryPtr> Selected;
 	SubjectsTreeView->GetSelectedItems(Selected);
-	if (Selected.Num())
+
+	for (const FLiveLinkSubjectUIEntryPtr& Entry : Selected)
 	{
-		Selected[0]->PauseSubject();
+		if (Entry && Entry->IsSubjectValid() && Entry->IsSubjectEnabled())
+		{
+			Entry->PauseSubject();
+		}
 	}
 }
 
@@ -1136,7 +1158,7 @@ void FLiveLinkSubjectsView::CreateSubjectsTreeView(const TSharedPtr<FUICommandLi
 		.OnGetChildren_Raw(this, &FLiveLinkSubjectsView::GetChildrenForInfo)
 		.OnSelectionChanged_Raw(this, &FLiveLinkSubjectsView::OnSubjectSelectionChanged)
 		.OnContextMenuOpening_Raw(this, &FLiveLinkSubjectsView::OnOpenVirtualSubjectContextMenu, InCommandList)
-		.SelectionMode(ESelectionMode::Single)
+		.SelectionMode(ESelectionMode::Multi)
 		.HeaderRow
 		(
 			SNew(SHeaderRow)
@@ -1183,11 +1205,16 @@ bool FLiveLinkSubjectsView::IsSelectedSubjectPaused() const
 {
 	TArray<FLiveLinkSubjectUIEntryPtr> Selected;
 	SubjectsTreeView->GetSelectedItems(Selected);
-	if (Selected.Num())
+
+	for (const FLiveLinkSubjectUIEntryPtr& EntryPtr : Selected)
 	{
-		return Selected[0]->IsPaused();
+		if (EntryPtr && !EntryPtr->IsPaused())
+		{
+			return false;
+		}
 	}
-	return false;
+
+	return true;
 }
 
 TSharedRef<SWidget> FLiveLinkSourcesView::OnGenerateSourceMenu()
diff --git a/Engine/Plugins/Animation/LiveLink/Source/LiveLinkEditor/Private/LiveLinkPanelController.cpp b/Engine/Plugins/Animation/LiveLink/Source/LiveLinkEditor/Private/LiveLinkPanelController.cpp
--- a/Engine/Plugins/Animation/LiveLink/Source/LiveLinkEditor/Private/LiveLinkPanelController.cpp
+++ b/Engine/Plugins/Animation/LiveLink/Source/LiveLinkEditor/Private/LiveLinkPanelController.cpp
@@ -252,10 +252,5 @@ bool FLiveLinkPanelController::CanPauseSubject() const
 
 void FLiveLinkPanelController::HandlePauseSubject()
 {
-	TArray<FLiveLinkSubjectUIEntryPtr> Selected;
-	SubjectsView->SubjectsTreeView->GetSelectedItems(Selected);
-	if (Selected.Num() > 0 && Selected[0])
-	{
-		Selected[0]->PauseSubject();
-	}
+	SubjectsView->HandlePauseSubject();
 }
diff --git a/Engine/Plugins/Animation/LiveLink/Source/LiveLinkEditor/Public/LiveLinkClientPanelViews.h b/Engine/Plugins/Animation/LiveLink/Source/LiveLinkEditor/Public/LiveLinkClientPanelViews.h
--- a/Engine/Plugins/Animation/LiveLink/Source/LiveLinkEditor/Public/LiveLinkClientPanelViews.h
+++ b/Engine/Plugins/Animation/LiveLink/Source/LiveLinkEditor/Public/LiveLinkClientPanelViews.h
@@ -147,6 +147,22 @@ namespace UE::LiveLink
 						Item->RemoveFromClient();
 					}
 				}
+
+				return FReply::Handled();
+			}
+			else if (InKeyEvent.GetModifierKeys().IsControlDown() && InKeyEvent.GetKey() == EKeys::A)
+			{
+				// Use SListView<ListElementType>::GetItems to select all items while avoiding the deprecation warning in STreeView.
+				for (const ListElementType& Element: SListView<ListElementType>::GetItems())
+				{
+					ListType::SetItemSelection(Element, true);
+				}
+
+				return FReply::Handled();
+			}
+			else if (InKeyEvent.GetKey() == EKeys::Escape)
+			{
+				ListType::ClearSelection();
 				return FReply::Handled();
 			}
 
