LiveLinkHub - Add a column that displays the translated role for a subject (hidden by default)

diff --git a/Engine/Plugins/Animation/LiveLink/Source/LiveLink/Private/LiveLinkClient.cpp b/Engine/Plugins/Animation/LiveLink/Source/LiveLink/Private/LiveLinkClient.cpp
--- a/Engine/Plugins/Animation/LiveLink/Source/LiveLink/Private/LiveLinkClient.cpp
+++ b/Engine/Plugins/Animation/LiveLink/Source/LiveLink/Private/LiveLinkClient.cpp
@@ -1639,6 +1639,26 @@ TSubclassOf<ULiveLinkRole> FLiveLinkClient::GetSubjectRole_AnyThread(FLiveLinkSu
 	return TSubclassOf<ULiveLinkRole>();
 }
 
+TSubclassOf<ULiveLinkRole> FLiveLinkClient::GetSubjectTranslatedRole_AnyThread(const FLiveLinkSubjectKey& InSubjectKey) const
+{
+	if (GetDefault<ULiveLinkSettings>()->bTranslateRebroadcastFrames)
+	{
+		if (const FLiveLinkCollectionSubjectItem* SubjectItem = Collection->FindSubject(InSubjectKey))
+		{
+			if (FLiveLinkSubject* Subject = SubjectItem->GetLiveSubject())
+			{
+				TArray<ULiveLinkFrameTranslator::FWorkerSharedPtr> Translators = Subject->GetFrameTranslators();
+				if (Translators.Num() && Translators[0])
+				{
+					return Translators[0]->GetToRole();
+				}
+			}
+		}
+	}
+
+	return nullptr;
+}
+
 bool FLiveLinkClient::DoesSubjectSupportsRole_AnyThread(const FLiveLinkSubjectKey& InSubjectKey, TSubclassOf<ULiveLinkRole> InSupportedRole) const
 {
 	if (const FLiveLinkCollectionSubjectItem* SubjectItem = Collection->FindSubject(InSubjectKey))
diff --git a/Engine/Plugins/Animation/LiveLink/Source/LiveLink/Public/LiveLinkClient.h b/Engine/Plugins/Animation/LiveLink/Source/LiveLink/Public/LiveLinkClient.h
--- a/Engine/Plugins/Animation/LiveLink/Source/LiveLink/Public/LiveLinkClient.h
+++ b/Engine/Plugins/Animation/LiveLink/Source/LiveLink/Public/LiveLinkClient.h
@@ -70,6 +70,7 @@ public:
 	UE_API virtual void ClearAllSubjectsFrames_AnyThread() override;
 	UE_API virtual TSubclassOf<ULiveLinkRole> GetSubjectRole_AnyThread(const FLiveLinkSubjectKey& SubjectKey) const override;
 	UE_API virtual TSubclassOf<ULiveLinkRole> GetSubjectRole_AnyThread(FLiveLinkSubjectName SubjectName) const override;
+	UE_API virtual TSubclassOf<ULiveLinkRole> GetSubjectTranslatedRole_AnyThread(const FLiveLinkSubjectKey& SubjectKey) const override;
 	UE_API virtual bool DoesSubjectSupportsRole_AnyThread(const FLiveLinkSubjectKey& SubjectKey, TSubclassOf<ULiveLinkRole> SupportedRole) const override;
 	UE_API virtual bool DoesSubjectSupportsRole_AnyThread(FLiveLinkSubjectName SubjectName, TSubclassOf<ULiveLinkRole> SupportedRole) const override;
 
diff --git a/Engine/Plugins/Animation/LiveLink/Source/LiveLinkEditor/Private/LiveLinkClientPanelViews.cpp b/Engine/Plugins/Animation/LiveLink/Source/LiveLinkEditor/Private/LiveLinkClientPanelViews.cpp
--- a/Engine/Plugins/Animation/LiveLink/Source/LiveLinkEditor/Private/LiveLinkClientPanelViews.cpp
+++ b/Engine/Plugins/Animation/LiveLink/Source/LiveLinkEditor/Private/LiveLinkClientPanelViews.cpp
@@ -51,6 +51,7 @@ namespace SubjectTreeUI
 	static const FName EnabledColumnName(TEXT("Enabled"));
 	static const FName NameColumnName(TEXT("Name"));
 	static const FName RoleColumnName(TEXT("Role"));
+	static const FName TranslatedRoleColumnName(TEXT("TranslatedRole"));
 	static const FName ActionsColumnName(TEXT("Action"));
 };
 
@@ -473,6 +474,11 @@ TSubclassOf<ULiveLinkRole> FLiveLinkSubjectUIEntry::GetItemRole() const
 	return IsSubject() ? Client->GetSubjectRole_AnyThread(SubjectKey) : TSubclassOf<ULiveLinkRole>();
 }
 
+TSubclassOf<ULiveLinkRole> FLiveLinkSubjectUIEntry::GetItemTranslatedRole() const
+{
+	return IsSubject() ? Client->GetSubjectTranslatedRole_AnyThread(SubjectKey) : TSubclassOf<ULiveLinkRole>();
+}
+
 void FLiveLinkSubjectUIEntry::RemoveFromClient() const
 {
 	Client->RemoveSubject_AnyThread(SubjectKey);
@@ -556,10 +562,29 @@ public:
 		}
 		else if (ColumnName == SubjectTreeUI::RoleColumnName)
 		{
-			auto RoleAttribute = MakeAttributeSP(this, &SLiveLinkClientPanelSubjectRow::GetItemRole);
-			return SNew(STextBlock)
+			TAttribute<FText> RoleAttribute = MakeAttributeSP(this, &SLiveLinkClientPanelSubjectRow::GetItemRole);
+
+			return SNew(SHorizontalBox)
+			+ SHorizontalBox::Slot()
+			.AutoWidth()
+			[
+				SNew(STextBlock)
 				.ColorAndOpacity(this, &SLiveLinkClientPanelSubjectRow::GetSubjectTextColor)
-				.Text(EntryPtr->IsSubject() ? RoleAttribute : FText::GetEmpty());
+				.Text(EntryPtr->IsSubject() ? RoleAttribute : FText::GetEmpty())
+			];
+		}
+		else if (ColumnName == SubjectTreeUI::TranslatedRoleColumnName)
+		{
+			TAttribute<FText> TranslatedRoleAttribute = MakeAttributeSP(this, &SLiveLinkClientPanelSubjectRow::GetItemTranslatedRole);
+
+			return SNew(SHorizontalBox)
+			+ SHorizontalBox::Slot()
+			.AutoWidth()
+			[
+				SNew(STextBlock)
+					.ColorAndOpacity(this, &SLiveLinkClientPanelSubjectRow::GetSubjectTextColor)
+					.Text(EntryPtr->IsSubject() ? TranslatedRoleAttribute : FText::GetEmpty())
+			];
 		}
 		else if (ColumnName == SubjectTreeUI::ActionsColumnName)
 		{
@@ -583,13 +608,25 @@ private:
 	FText GetItemRole() const
 	{
 		TSubclassOf<ULiveLinkRole> Role = EntryPtr->GetItemRole();
-		if (Role.Get())
+		if (Role)
 		{
 			return Role->GetDefaultObject<ULiveLinkRole>()->GetDisplayName();
 		}
 		return FText::GetEmpty();
 	}
 
+	// If role is translated, put the original role in the extra info.
+	FText GetItemTranslatedRole() const
+	{
+		TSubclassOf<ULiveLinkRole> TranslatedRole = EntryPtr->GetItemTranslatedRole();
+
+		if (TranslatedRole.Get())
+		{
+			return TranslatedRole->GetDefaultObject<ULiveLinkRole>()->GetDisplayName();
+		}
+		return FText::GetEmpty();
+	}
+
 	/** Get widget visibility according to whether or not the panel is in read-only mode. */
 	EVisibility GetVisibilityFromReadOnly() const
 	{
@@ -1162,15 +1199,24 @@ void FLiveLinkSubjectsView::CreateSubjectsTreeView(const TSharedPtr<FUICommandLi
 		.HeaderRow
 		(
 			SNew(SHeaderRow)
+			.CanSelectGeneratedColumn(true)
+			.HiddenColumnsList({ SubjectTreeUI::TranslatedRoleColumnName })
 			+ SHeaderRow::Column(SubjectTreeUI::EnabledColumnName)
+			.DefaultTooltip(LOCTEXT("EnabledToolTip", "Whether this subject should be evaluated."))
 			.DefaultLabel(FText())
 			.FixedWidth(22)
 			+ SHeaderRow::Column(SubjectTreeUI::NameColumnName)
 			.DefaultLabel(LOCTEXT("SubjectItemName", "Subject Name"))
+			.DefaultTooltip(LOCTEXT("NameToolTip", "Unique name for this subject."))
 			.FillWidth(0.60f)
 			+ SHeaderRow::Column(SubjectTreeUI::RoleColumnName)
 			.DefaultLabel(LOCTEXT("RoleName", "Role"))
-			.FillWidth(0.40f)
+			.ManualWidth(90.f)
+			.DefaultTooltip(LOCTEXT("RoleToolTip", "Role of this subject."))
+			+ SHeaderRow::Column(SubjectTreeUI::TranslatedRoleColumnName)
+			.DefaultLabel(LOCTEXT("TranslatedRoleName", "Translated"))
+			.DefaultTooltip(LOCTEXT("TranslatedRoleToolTip", "Get the translated role of a subject if it's being translated before being rebroadcast. This should only be relevant in Live Link Hub."))
+			.ManualWidth(70.f)
 			+ SHeaderRow::Column(SubjectTreeUI::ActionsColumnName)
 			.ManualWidth(20.f)
 			.DefaultLabel(FText())
diff --git a/Engine/Plugins/Animation/LiveLink/Source/LiveLinkEditor/Public/LiveLinkClientPanelViews.h b/Engine/Plugins/Animation/LiveLink/Source/LiveLinkEditor/Public/LiveLinkClientPanelViews.h
--- a/Engine/Plugins/Animation/LiveLink/Source/LiveLinkEditor/Public/LiveLinkClientPanelViews.h
+++ b/Engine/Plugins/Animation/LiveLink/Source/LiveLinkEditor/Public/LiveLinkClientPanelViews.h
@@ -72,6 +72,8 @@ struct FLiveLinkSubjectUIEntry
 	FText GetItemText() const;
 	// Get the livelink role of this entry
 	TSubclassOf<ULiveLinkRole> GetItemRole() const;
+	// Get the translated role for this subject (if translating rebroadcast subjects).
+	TSubclassOf<ULiveLinkRole> GetItemTranslatedRole() const;
 	// Remove the subject or source from the livelink client
 	void RemoveFromClient() const;
 	// Pause or unpause a subject if it's already paused
diff --git a/Engine/Source/Runtime/LiveLinkInterface/Public/ILiveLinkClient.h b/Engine/Source/Runtime/LiveLinkInterface/Public/ILiveLinkClient.h
--- a/Engine/Source/Runtime/LiveLinkInterface/Public/ILiveLinkClient.h
+++ b/Engine/Source/Runtime/LiveLinkInterface/Public/ILiveLinkClient.h
@@ -155,6 +155,12 @@ public:
 	/** Get the role of the subject with this name */
 	virtual TSubclassOf<ULiveLinkRole> GetSubjectRole_AnyThread(FLiveLinkSubjectName SubjectName) const = 0;
 
+	/**
+	 * Get the translated role of a subject (if it's being translated before being rebroadcast). 
+	 * If the subject is not being translated, this function will return an invalid TSubclassOf.
+	 */
+	virtual TSubclassOf<ULiveLinkRole> GetSubjectTranslatedRole_AnyThread(const FLiveLinkSubjectKey& SubjectKey) const = 0;
+
 	/** Whether a subject support a particular role, either directly or through a translator */
 	virtual bool DoesSubjectSupportsRole_AnyThread(const FLiveLinkSubjectKey& SubjectKey, TSubclassOf<ULiveLinkRole> SupportedRole) const = 0;
 
