LiveLinkTimecode - Rely on subject name rather than subject key to allow driving timecode with a livelinkhub recording.
Fix timecode not working after a recording looped or was scrubbed. LiveLinkTimecodeProvider's frame history is now cleared whenever a recording is scrubbed, stepped through or looped in order to ensure the buffer only contains the latest frames provided by the recording.

Note: LiveLinkTimecodeProvider will not work with LiveLink recordings if the timecode provider is in Lerp Evaluation Mode, since the platform time will greatly differ from the timecode in the recording.
This could eventually be corrected by adding an API to feed a "current time" from the recording so that lerping works correctly.

#rnx

diff --git a/Engine/Plugins/Animation/LiveLink/Source/LiveLink/Private/LiveLinkTimecodeProvider.cpp b/Engine/Plugins/Animation/LiveLink/Source/LiveLink/Private/LiveLinkTimecodeProvider.cpp
--- a/Engine/Plugins/Animation/LiveLink/Source/LiveLink/Private/LiveLinkTimecodeProvider.cpp
+++ b/Engine/Plugins/Animation/LiveLink/Source/LiveLink/Private/LiveLinkTimecodeProvider.cpp
@@ -71,6 +71,12 @@ FQualifiedFrameTime ULiveLinkTimecodeProvider::GetQualifiedFrameTime() const
 }
 
 
+void ULiveLinkTimecodeProvider::ClearFrameHistory()
+{
+	FScopeLock ScopeLock(&SubjectFrameLock);
+	SubjectFrameTimes.Reset();
+}
+
 FQualifiedFrameTime ULiveLinkTimecodeProvider::ConvertTo(FQualifiedFrameTime Value) const
 {
 	if (bOverrideFrameRate)
@@ -140,13 +146,13 @@ void ULiveLinkTimecodeProvider::PostEditChangeProperty(FPropertyChangedEvent& Pr
 	if (PropertyChangedEvent.Property && PropertyChangedEvent.Property->GetFName() == GET_MEMBER_NAME_CHECKED(ULiveLinkTimecodeProvider, SubjectKey))
 	{
 		// If was already registered
-		if (RegisterForFrameDataReceivedHandle.IsValid() && RegisteredSubjectKey != SubjectKey)
+		if (RegisterForFrameDataReceivedHandle.IsValid() && RegisteredSubjectKey.SubjectName != SubjectKey.SubjectName)
 		{
 			UnregisterSubject();
 			RegisterSubject();
 		}
 		// If was waiting for the subject to be added to the client
-		else if (LiveLinkClient && LiveLinkClient->GetSubjectSettings(SubjectKey))
+		else if (LiveLinkClient && LiveLinkClient->IsSubjectValid(SubjectKey.SubjectName))
 		{
 			OnLiveLinkSubjectAdded(SubjectKey);
 		}
@@ -166,7 +172,7 @@ void ULiveLinkTimecodeProvider::InitClient()
 		LiveLinkClient->OnLiveLinkSubjectRemoved().AddUObject(this, &ULiveLinkTimecodeProvider::OnLiveLinkSubjectRemoved);
 
 		// if the subject already exist
-		if (LiveLinkClient->GetSubjectSettings(SubjectKey))
+		if (LiveLinkClient->IsSubjectValid(SubjectKey.SubjectName))
 		{
 			OnLiveLinkSubjectAdded(SubjectKey);
 		}
@@ -220,7 +226,6 @@ void ULiveLinkTimecodeProvider::UnregisterSubject()
 		FScopeLock ScopeLock(&SubjectFrameLock);
 		SubjectFrameTimes.Reset();
 	}
-
 }
 
 
@@ -245,8 +250,11 @@ void ULiveLinkTimecodeProvider::OnLiveLinkClientUnregistered(const FName& Type,
 
 void ULiveLinkTimecodeProvider::OnLiveLinkSubjectAdded(FLiveLinkSubjectKey InSubjectKey)
 {
-	if (InSubjectKey == SubjectKey)
+	if (InSubjectKey.SubjectName == SubjectKey.SubjectName)
 	{
+		// This provider might have been initialized with an invalid key because the subject wasn't live at the time of creation. 
+		// So we update the SubjectKey to match the new subject since it's now been added.
+		SubjectKey = InSubjectKey;
 		RegisterSubject();
 	}
 }
@@ -254,7 +262,7 @@ void ULiveLinkTimecodeProvider::OnLiveLinkSubjectAdded(FLiveLinkSubjectKey InSub
 
 void ULiveLinkTimecodeProvider::OnLiveLinkSubjectRemoved(FLiveLinkSubjectKey InSubjectKey)
 {
-	if (InSubjectKey == RegisteredSubjectKey)
+	if (InSubjectKey.SubjectName == RegisteredSubjectKey.SubjectName)
 	{
 		State = ETimecodeProviderSynchronizationState::Error;
 		UnregisterSubject();
diff --git a/Engine/Plugins/Animation/LiveLink/Source/LiveLink/Public/LiveLinkTimecodeProvider.h b/Engine/Plugins/Animation/LiveLink/Source/LiveLink/Public/LiveLinkTimecodeProvider.h
--- a/Engine/Plugins/Animation/LiveLink/Source/LiveLink/Public/LiveLinkTimecodeProvider.h
+++ b/Engine/Plugins/Animation/LiveLink/Source/LiveLink/Public/LiveLinkTimecodeProvider.h
@@ -66,6 +66,9 @@ public:
 		SubjectKey = InKey;
 	}
 
+	/** Clear the frame history. Useful when using Latest/Nearest evaluation type and wanting to accept timecoded frames older than what's currently in the buffer. */
+	UE_API void ClearFrameHistory();
+
 private:
 	UE_API FQualifiedFrameTime ConvertTo(FQualifiedFrameTime Value) const;
 
diff --git a/Engine/Plugins/Animation/LiveLinkHub/Source/LiveLinkHub/Private/Recording/LiveLinkHubPlaybackController.cpp b/Engine/Plugins/Animation/LiveLinkHub/Source/LiveLinkHub/Private/Recording/LiveLinkHubPlaybackController.cpp
--- a/Engine/Plugins/Animation/LiveLinkHub/Source/LiveLinkHub/Private/Recording/LiveLinkHubPlaybackController.cpp
+++ b/Engine/Plugins/Animation/LiveLinkHub/Source/LiveLinkHub/Private/Recording/LiveLinkHubPlaybackController.cpp
@@ -7,6 +7,7 @@
 #include "CoreGlobals.h"
 #include "ILiveLinkClient.h"
 #include "Implementations/LiveLinkUAssetRecordingPlayer.h"
+#include "LiveLinkTimecodeProvider.h"
 #include "LiveLinkHubCreatorAppMode.h"
 #include "LiveLinkHubLog.h"
 #include "LiveLinkHubPlaybackSourceSettings.h"
@@ -27,6 +28,17 @@
 
 #define LOCTEXT_NAMESPACE "LiveLinkHubPlayback"
 
+namespace UE::LiveLinkHub 
+{
+	static void ClearTimecodeProviderFrameHistory()
+	{
+		if (ULiveLinkTimecodeProvider* Provider = Cast< ULiveLinkTimecodeProvider>(GEngine->GetTimecodeProvider()))
+		{
+			Provider->ClearFrameHistory();
+		}
+	}
+}
+
 /**
  * Thread safe way to load and store a FQualifiedFrameTime. This is necessary because atomic<FQualifiedFrameTime> isn't
  * necessarily cross platform compatible when performing copy operations.
@@ -373,6 +385,9 @@ void FLiveLinkHubPlaybackController::BeginPlayback(bool bInReverse)
 		}
 		
 		ResumePlayback();
+
+		// Note: We don't want to clear frame history when pausing because timecode would become invalid.
+		UE::LiveLinkHub::ClearTimecodeProviderFrameHistory();
 	}
 	else if (bIsPlaying)
 	{
@@ -391,6 +406,8 @@ void FLiveLinkHubPlaybackController::RestartPlayback()
 	RecordingPlayer->RestartPlayback(GetCurrentFrame().Value);
 	bIsPlaying = true;
 	bIsReverse = bOldReverse;
+
+	UE::LiveLinkHub::ClearTimecodeProviderFrameHistory();
 }
 
 void FLiveLinkHubPlaybackController::PausePlayback()
@@ -545,6 +562,8 @@ void FLiveLinkHubPlaybackController::OnSelectRecording()
 
 void FLiveLinkHubPlaybackController::GoToTime(FQualifiedFrameTime InTime)
 {
+	UE::LiveLinkHub::ClearTimecodeProviderFrameHistory();
+
 	// Stop needs to occur to restart playback.
 	StopPlayback();
 
diff --git a/Engine/Plugins/Animation/LiveLinkHub/Source/LiveLinkHubMessaging/Private/LiveLinkHubMessages.cpp b/Engine/Plugins/Animation/LiveLinkHub/Source/LiveLinkHubMessaging/Private/LiveLinkHubMessages.cpp
--- a/Engine/Plugins/Animation/LiveLinkHub/Source/LiveLinkHubMessaging/Private/LiveLinkHubMessages.cpp
+++ b/Engine/Plugins/Animation/LiveLinkHub/Source/LiveLinkHubMessaging/Private/LiveLinkHubMessages.cpp
@@ -113,22 +113,26 @@ void FLiveLinkHubTimecodeSettings::AssignTimecodeSettingsAsProviderToEngine() co
 	{
 		TOptional<FLiveLinkSubjectKey> Target = UE::LiveLinkHubMessages::Private::FindSubjectKey(LiveLinkClient, SubjectName);
 
+		FName ObjectName = MakeUniqueObjectName(GEngine, ULiveLinkTimecodeProvider::StaticClass(), "DefaultLiveLinkTimecodeProvider");
+		ULiveLinkTimecodeProvider* LiveLinkProvider = NewObject<ULiveLinkTimecodeProvider>(GEngine, ObjectName);
+
 		if (Target)
 		{
-			FName ObjectName = MakeUniqueObjectName(GEngine, ULiveLinkTimecodeProvider::StaticClass(), "DefaultLiveLinkTimecodeProvider");
-			ULiveLinkTimecodeProvider* LiveLinkProvider = NewObject<ULiveLinkTimecodeProvider>(GEngine, ObjectName);
 			LiveLinkProvider->SetTargetSubjectKey(*Target);
-			LiveLinkProvider->OverrideFrameRate = DesiredFrameRate;
-			LiveLinkProvider->FrameDelay = FrameDelay;
-			LiveLinkProvider->BufferSize = BufferSize;
-			LiveLinkProvider->Evaluation = EvaluationType;
-			GEngine->SetTimecodeProvider(LiveLinkProvider);
-			UE_LOG(LogLiveLinkHubMessages, Display, TEXT("Live Link Timecode provider assigned to %s."), *SubjectName.ToString());
 		}
 		else
 		{
-			UE_LOG(LogLiveLinkHubMessages, Warning, TEXT("Failed to assign Live Link Timecode provider to %s."), *SubjectName.ToString());
+			// Create a mock subject key in order to match a subject when it comes online.
+			LiveLinkProvider->SetTargetSubjectKey(FLiveLinkSubjectKey{ FGuid{}, SubjectName });
+			UE_LOG(LogLiveLinkHubMessages, Warning, TEXT("Assigned tiemcode provider with invalid subject %s."), *SubjectName.ToString());
 		}
+
+		LiveLinkProvider->OverrideFrameRate = DesiredFrameRate;
+		LiveLinkProvider->FrameDelay = FrameDelay;
+		LiveLinkProvider->BufferSize = BufferSize;
+		LiveLinkProvider->Evaluation = EvaluationType;
+		GEngine->SetTimecodeProvider(LiveLinkProvider);
+		UE_LOG(LogLiveLinkHubMessages, Display, TEXT("Live Link Timecode provider assigned to %s."), *SubjectName.ToString());
 	}
 	else
 	{
