LiveLinkHub - Add option for ticking the Live Link Client whenever it receives data.
This should allow VirtualSubject's sync subjects to have a better frame pacing that matches the incoming data.

LiveLinkHub.EnableTriggeredTick 1 to enable this feature

LiveLinkHub Virtual Subjects are usually rebroadcast whenever LiveLinkTicker has waited for a full tick duration, but this can cause uneven frame pacing because it won't match the exact framerate of sources that make up the Virtual Subject.

diff --git a/Engine/Plugins/Animation/LiveLink/Source/LiveLink/Private/LiveLinkClient.cpp b/Engine/Plugins/Animation/LiveLink/Source/LiveLink/Private/LiveLinkClient.cpp
--- a/Engine/Plugins/Animation/LiveLink/Source/LiveLink/Private/LiveLinkClient.cpp
+++ b/Engine/Plugins/Animation/LiveLink/Source/LiveLink/Private/LiveLinkClient.cpp
@@ -65,12 +65,19 @@ FLiveLinkClient::FLiveLinkClient()
 FLiveLinkClient::FLiveLinkClient(FSimpleMulticastDelegate& InTickingDelegate)
 	: Collection(MakeUnique<FLiveLinkSourceCollection>())
 {
-	InTickingDelegate.AddRaw(this, &FLiveLinkClient::Tick);
+	InTickingDelegate.AddRaw(this, &FLiveLinkClient::Tick, ELiveLinkTickType::Default);
 	Initialize();
 }
 
 FLiveLinkClient::FLiveLinkClient(FTSSimpleMulticastDelegate& InTickingDelegate)
 	: Collection(MakeUnique<FLiveLinkSourceCollection>())
+{
+	InTickingDelegate.AddRaw(this, &FLiveLinkClient::Tick, ELiveLinkTickType::Default);
+	Initialize();
+}
+
+FLiveLinkClient::FLiveLinkClient(FLiveLinkTickDelegate& InTickingDelegate)
+	: Collection(MakeUnique<FLiveLinkSourceCollection>())
 {
 	InTickingDelegate.AddRaw(this, &FLiveLinkClient::Tick);
 	Initialize();
@@ -82,15 +89,19 @@ FLiveLinkClient::~FLiveLinkClient()
 	Shutdown();
 }
 
-void FLiveLinkClient::Tick()
+void FLiveLinkClient::Tick(ELiveLinkTickType TickType)
 {
 	SCOPE_CYCLE_COUNTER(STAT_LiveLink_Client_Tick);
 
 	TSet<FLiveLinkSubjectKey> ReceivedSubjects;
 
 	DoPendingWork(ReceivedSubjects);
-	CacheValues();
-	UpdateSources();
+
+	if (TickType == ELiveLinkTickType::Scheduled || TickType == ELiveLinkTickType::Default)
+	{
+		CacheValues();
+		UpdateSources();
+	}
 
 	TSet<FLiveLinkSubjectKey> UpdatedSubjects = BuildThisTicksSubjectSnapshot(ReceivedSubjects);
 	RebroadcastSubjects(UpdatedSubjects);
@@ -117,7 +128,7 @@ void FLiveLinkClient::Initialize()
 	bEnableParentSubjects = GConfig->GetBoolOrDefault(TEXT("LiveLink"), TEXT("bEnableParentSubjects"), false, GEngineIni);
 }
 
-void FLiveLinkClient::DoPendingWork(TSet<FLiveLinkSubjectKey>& OutReceivedSubjects)
+void FLiveLinkClient::DoPendingWork(TSet<FLiveLinkSubjectKey>& OutUpdatedSubjects)
 {
 	check(Collection);
 
@@ -132,7 +143,7 @@ void FLiveLinkClient::DoPendingWork(TSet<FLiveLinkSubjectKey>& OutReceivedSubjec
 		// Add new Subject static data
 		for (FPendingSubjectStatic& SubjectStatic : SubjectStaticToPush)
 		{
-			OutReceivedSubjects.Add(SubjectStatic.SubjectKey);
+			OutUpdatedSubjects.Add(SubjectStatic.SubjectKey);
 			PushSubjectStaticData_Internal(MoveTemp(SubjectStatic));
 		}
 		SubjectStaticToPush.Reset();
@@ -140,7 +151,7 @@ void FLiveLinkClient::DoPendingWork(TSet<FLiveLinkSubjectKey>& OutReceivedSubjec
 		// Add new Subject frame data
 		for (FPendingSubjectFrame& SubjectFrame : SubjectFrameToPush)
 		{
-			OutReceivedSubjects.Add(SubjectFrame.SubjectKey);
+			OutUpdatedSubjects.Add(SubjectFrame.SubjectKey);
 			PushSubjectFrameData_Internal(MoveTemp(SubjectFrame));
 		}
 		SubjectFrameToPush.Reset();
@@ -312,7 +323,7 @@ void FLiveLinkClient::RebroadcastSubjects(const TSet<FLiveLinkSubjectKey>& Updat
 				{
 					FLiveLinkFrameIdentifier SnapshotFrameId = LiveSubject->GetSnapshotFrameId();
 
-					if (LiveSubject->GetSnapshotFrameId() != LiveSubject->GetLastRebroadcastedFrameId())
+					if (LiveSubject->GetSnapshotFrameId() != LiveSubject->GetLastRebroadcastedFrameId() || SnapshotFrameId == -1)
 					{
 						FLiveLinkSubjectFrameData FrameData;
 
@@ -339,7 +350,6 @@ void FLiveLinkClient::RebroadcastSubjects(const TSet<FLiveLinkSubjectKey>& Updat
 		});
 }
 
-
 void FLiveLinkClient::CacheValues()
 {
 #if WITH_EDITOR
@@ -1025,7 +1035,7 @@ void FLiveLinkClient::PushSubjectStaticData_Internal(FPendingSubjectStatic&& Sub
 			Handles->OnStaticDataAdded.Broadcast(SubjectStaticData.SubjectKey, SubjectStaticData.Role, SubjectStaticData.StaticData);
 			Handles->OnUnmappedStaticDataAdded.Broadcast(SubjectStaticData.SubjectKey, SubjectStaticData.Role, *UnmappedStaticData);
 		}
-		else if (const FSubjectFramesAddedHandles* AllSubjectsHandler = SubjectFrameAddedHandles.Find(ALL_SUBJECTS_DELEGATE_TOKEN))
+		else if (const FSubjectFramesAddedHandles* AllSubjectsHandler = SubjectFrameAddedHandles.Find(UE::LiveLink::Private::ALL_SUBJECTS_DELEGATE_TOKEN))
 		{
 			AllSubjectsHandler->OnStaticDataAdded.Broadcast(SubjectStaticData.SubjectKey, SubjectStaticData.Role, SubjectStaticData.StaticData);
 			AllSubjectsHandler->OnUnmappedStaticDataAdded.Broadcast(SubjectStaticData.SubjectKey, SubjectStaticData.Role, *UnmappedStaticData);
@@ -1043,6 +1053,11 @@ void FLiveLinkClient::BroadcastFrameDataUpdate(const FLiveLinkSubjectKey& InSubj
 	{
 		Handles->OnFrameDataReceived.Broadcast(InFrameData);
 	}
+
+	if (const FSubjectFramesReceivedHandles* Handles = SubjectFrameReceivedHandles.Find(FLiveLinkSubjectKey{ FGuid(), UE::LiveLink::Private::ALL_SUBJECTS_DELEGATE_TOKEN}))
+	{
+		Handles->OnFrameDataReceived.Broadcast(InFrameData);
+	}
 }
 
 TSharedPtr<ILiveLinkProvider> FLiveLinkClient::GetRebroadcastLiveLinkProvider() const
@@ -1189,7 +1204,7 @@ void FLiveLinkClient::PushSubjectFrameData_Internal(FPendingSubjectFrame&& Subje
 		Handles->OnFrameDataAdded.Broadcast(SubjectItem->Key, Role, SubjectFrameData.FrameData);
 		Handles->OnUnmappedFrameDataAdded.Broadcast(SubjectItem->Key, Role, UnmappedFrameData.IsValid() ? UnmappedFrameData : SubjectFrameData.FrameData);
 	}
-	else if (const FSubjectFramesAddedHandles* AllSubjectsHandler = SubjectFrameAddedHandles.Find(ALL_SUBJECTS_DELEGATE_TOKEN))
+	else if (const FSubjectFramesAddedHandles* AllSubjectsHandler = SubjectFrameAddedHandles.Find(UE::LiveLink::Private::ALL_SUBJECTS_DELEGATE_TOKEN))
 	{
 		// NAME_None means we registered for all subjects update.
 		AllSubjectsHandler->OnFrameDataAdded.Broadcast(SubjectItem->Key, Role, SubjectFrameData.FrameData);
@@ -2199,7 +2214,7 @@ FOnLiveLinkSubjectEvaluated& FLiveLinkClient::OnLiveLinkSubjectEvaluated()
 
 void FLiveLinkClient::UnregisterGlobalSubjectFramesDelegate(FDelegateHandle& InStaticDataAddedHandle, FDelegateHandle& InFrameDataAddedHandle, bool bUseUnmappedData)
 {
-	if (FSubjectFramesAddedHandles* Handles = SubjectFrameAddedHandles.Find(ALL_SUBJECTS_DELEGATE_TOKEN))
+	if (FSubjectFramesAddedHandles* Handles = SubjectFrameAddedHandles.Find(UE::LiveLink::Private::ALL_SUBJECTS_DELEGATE_TOKEN))
 	{
 		if (bUseUnmappedData)
 		{
@@ -2218,7 +2233,7 @@ bool FLiveLinkClient::RegisterGlobalSubjectFramesDelegate(const FOnLiveLinkSubje
 	const FOnLiveLinkSubjectFrameDataAdded::FDelegate& InOnFrameDataAdded, FDelegateHandle& OutStaticDataAddedHandle,
 	FDelegateHandle& OutFrameDataAddedHandle, bool bUseUnmappedData)
 {
-	FSubjectFramesAddedHandles& Handles = SubjectFrameAddedHandles.FindOrAdd(ALL_SUBJECTS_DELEGATE_TOKEN);
+	FSubjectFramesAddedHandles& Handles = SubjectFrameAddedHandles.FindOrAdd(UE::LiveLink::Private::ALL_SUBJECTS_DELEGATE_TOKEN);
 	OutStaticDataAddedHandle = bUseUnmappedData ? Handles.OnUnmappedStaticDataAdded.Add(InOnStaticDataAdded) : Handles.OnStaticDataAdded.Add(InOnStaticDataAdded);
 	OutFrameDataAddedHandle = bUseUnmappedData ? Handles.OnUnmappedFrameDataAdded.Add(InOnFrameDataAdded) : Handles.OnFrameDataAdded.Add(InOnFrameDataAdded);
 
diff --git a/Engine/Plugins/Animation/LiveLink/Source/LiveLink/Public/LiveLinkClient.h b/Engine/Plugins/Animation/LiveLink/Source/LiveLink/Public/LiveLinkClient.h
--- a/Engine/Plugins/Animation/LiveLink/Source/LiveLink/Public/LiveLinkClient.h
+++ b/Engine/Plugins/Animation/LiveLink/Source/LiveLink/Public/LiveLinkClient.h
@@ -3,6 +3,7 @@
 #pragma once
 
 #include "ILiveLinkClient.h"
+#include "Containers/Set.h"
 
 #define UE_API LIVELINK_API
 
@@ -28,6 +29,20 @@ struct FLiveLinkSubjectTimeSyncData
 	FFrameRate SampleFrameRate;
 };
 
+/** Describes what triggered the LiveLinkClient Tick. */
+enum class ELiveLinkTickType
+{
+	Default, // Default tick path, used in UE. 
+	Scheduled, // Scheduled tick that's triggered by the LiveLinkHub ticker.
+	Triggered  // Tick triggered by frame data being received outside of the Game Thread.
+};
+
+namespace UE::LiveLink::Private
+{
+	/** Name token  used to register to all subject updates. */
+	static const FName ALL_SUBJECTS_DELEGATE_TOKEN = "__Internal_AllSubjects_Update";
+}
+
 class FLiveLinkClient : public ILiveLinkClient
 {
 public:
@@ -38,6 +53,9 @@ public:
 	UE_API FLiveLinkClient(FSimpleMulticastDelegate& InTickingDelegate);
 	UE_API FLiveLinkClient(FTSSimpleMulticastDelegate& InTickingDelegate);
 
+	DECLARE_TS_MULTICAST_DELEGATE_OneParam(FLiveLinkTickDelegate, ELiveLinkTickType)
+	UE_API FLiveLinkClient(FLiveLinkTickDelegate& InTickingDelegate);
+
 	UE_API virtual ~FLiveLinkClient();
 
 	//~ Begin ILiveLinkClient implementation
@@ -145,7 +163,7 @@ public:
 
 
 	/** The tick callback to update the pending work and clear the subject's snapshot*/
-	UE_API void Tick();
+	UE_API void Tick(ELiveLinkTickType TickType = ELiveLinkTickType::Default);
 
 	/** Remove all sources from the live link client */
 	UE_API void RemoveAllSources();
@@ -208,16 +226,19 @@ private:
 	void Initialize();
 
 	/** Remove old sources & subject,  */
-	void DoPendingWork(TSet<FLiveLinkSubjectKey>& OutReceivedSubjects);
+	void DoPendingWork(TSet<FLiveLinkSubjectKey>& OutUpdatedSubjects);
 
 	/** Update the added sources */
 	void UpdateSources();
 
 	/**
-	 * Build subject data so that during the rest of the tick it can be read without thread locking or mem copying
+	 * Build subject data so that during the rest of the tick it can be read without
+	 * thread locking or mem copying
+	 * @return Subjects that were updated by this method.
 	 */
 	TSet<FLiveLinkSubjectKey> BuildThisTicksSubjectSnapshot(const TSet<FLiveLinkSubjectKey>& ReceivedSubjects);
 
+
 	/** Rebroadcast live (if using evaluated data) and virtual subjects. */
 	void RebroadcastSubjects(const TSet<FLiveLinkSubjectKey>& UpdatedSubjects);
 
@@ -306,8 +327,6 @@ private:
 	FString RebroadcastLiveLinkProviderName;
 	TSet<FLiveLinkSubjectKey> RebroadcastedSubjects;
 
-	/** Name token  used to register to all subject updates. */
-	const FName ALL_SUBJECTS_DELEGATE_TOKEN = "__Internal_AllSubjects_Update";
 
 	/** Whether to Preprocess frames before rebroadcasting them. */
 	bool bPreProcessRebroadcastFrames = false;
diff --git a/Engine/Plugins/Animation/LiveLinkHub/Source/LiveLinkHub/Private/LiveLinkHub.cpp b/Engine/Plugins/Animation/LiveLinkHub/Source/LiveLinkHub/Private/LiveLinkHub.cpp
--- a/Engine/Plugins/Animation/LiveLinkHub/Source/LiveLinkHub/Private/LiveLinkHub.cpp
+++ b/Engine/Plugins/Animation/LiveLinkHub/Source/LiveLinkHub/Private/LiveLinkHub.cpp
@@ -41,6 +41,11 @@
 
 #define LOCTEXT_NAMESPACE "LiveLinkHub"
 
+static TAutoConsoleVariable<bool> CVarLiveLinkHubEnableTriggeredTick(
+	TEXT("LiveLinkHub.EnableTriggeredTick"), false,
+	TEXT("Whether to allow LiveLinkHub's LiveLink client to tick whenever it receives new data. (Experimental)"),
+	ECVF_RenderThreadSafe);
+
 namespace LiveLinkHubInternalUtils
 {
 	/** Is this a standalone distributed build on the EGS? */
@@ -186,6 +191,11 @@ void FLiveLinkHub::InitializePostEngineInit()
 	LiveLinkHubClient->OnFrameDataReceived_AnyThread().AddSP(this, &FLiveLinkHub::OnFrameDataReceived_AnyThread);
 	LiveLinkHubClient->OnSubjectMarkedPendingKill_AnyThread().AddSP(this, &FLiveLinkHub::OnSubjectMarkedPendingKill_AnyThread);
 
+	FDelegateHandle Handle1;
+	FDelegateHandle Handle2;
+	FLiveLinkSubjectKey GlobalKey = FLiveLinkSubjectKey{ FGuid(), FLiveLinkSubjectName{ UE::LiveLink::Private::ALL_SUBJECTS_DELEGATE_TOKEN }};
+	LiveLinkHubClient->RegisterForFrameDataReceived(GlobalKey, FOnLiveLinkSubjectStaticDataReceived::FDelegate(), FOnLiveLinkSubjectFrameDataReceived::FDelegate::CreateSP(this, &FLiveLinkHub::OnLiveLinkFrameDataReceived), Handle1, Handle2);
+
 	RegisterLiveLinkHubSettings();
 
 	PlaybackController->Start();
@@ -424,6 +434,15 @@ void FLiveLinkHub::OnStaticDataReceived_AnyThread(const FLiveLinkSubjectKey& InS
 	}
 }
 
+void FLiveLinkHub::OnLiveLinkFrameDataReceived(const FLiveLinkFrameDataStruct&) const
+{
+	if (CVarLiveLinkHubEnableTriggeredTick.GetValueOnAnyThread())
+	{
+		OnFrameReceivedDelegate.Execute();
+	}
+}
+
+
 void FLiveLinkHub::OnFrameDataReceived_AnyThread(const FLiveLinkSubjectKey& InSubjectKey, const FLiveLinkFrameDataStruct& InFrameDataStruct) const
 {
 	if (RecordingController->IsRecording())
diff --git a/Engine/Plugins/Animation/LiveLinkHub/Source/LiveLinkHub/Private/LiveLinkHub.h b/Engine/Plugins/Animation/LiveLinkHub/Source/LiveLinkHub/Private/LiveLinkHub.h
--- a/Engine/Plugins/Animation/LiveLinkHub/Source/LiveLinkHub/Private/LiveLinkHub.h
+++ b/Engine/Plugins/Animation/LiveLinkHub/Source/LiveLinkHub/Private/LiveLinkHub.h
@@ -102,6 +102,9 @@ public:
 	/** Save an existing recording as a new file. */
 	void UpgradeAndSaveRecording(ULiveLinkUAssetRecording* InRecording);
 	
+	/** Delegate triggered whenever the Live Link Client receives frame data. */
+	FSimpleDelegate& OnFrameReceived() { return OnFrameReceivedDelegate; }
+
 	//~ FAssetEditorToolkit interface
 	virtual void MapToolkitCommands() override;
 	virtual void OnClose() override;
@@ -117,6 +120,9 @@ private:
 	void OnStaticDataReceived_AnyThread(const FLiveLinkSubjectKey& InSubjectKey, TSubclassOf<ULiveLinkRole> InRole, const FLiveLinkStaticDataStruct& InStaticDataStruct) const;
 	void OnFrameDataReceived_AnyThread(const FLiveLinkSubjectKey& InSubjectKey, const FLiveLinkFrameDataStruct& InFrameDataStruct) const;
 	void OnSubjectMarkedPendingKill_AnyThread(const FLiveLinkSubjectKey& InSubjectKey) const;
+
+	/** Called whenever LiveLinkClient receives data to wake up the LiveLink Ticker if LiveLinkHub.EnableTriggeredTick is enabled. */
+	void OnLiveLinkFrameDataReceived(const FLiveLinkFrameDataStruct&) const;
 	//~ LiveLink Client delegates
 
 	/** Register settings relevant to the livelink hub. */
@@ -187,6 +193,8 @@ private:
 	TOptional<ELiveLinkTopologyMode> OverrideTopologyMode;
 	/** Delegate called when the topology mode changed. */
 	FOnTopologyModeChanged TopologyModeChangedDelegate;
+	/** Delegate called whenever LL Client receives a frame. */
+	FSimpleDelegate OnFrameReceivedDelegate;
 
 	friend class FLiveLinkHubModule;
 };
diff --git a/Engine/Plugins/Animation/LiveLinkHub/Source/LiveLinkHub/Private/LiveLinkHubModule.cpp b/Engine/Plugins/Animation/LiveLinkHub/Source/LiveLinkHub/Private/LiveLinkHubModule.cpp
--- a/Engine/Plugins/Animation/LiveLinkHub/Source/LiveLinkHub/Private/LiveLinkHubModule.cpp
+++ b/Engine/Plugins/Animation/LiveLinkHub/Source/LiveLinkHub/Private/LiveLinkHubModule.cpp
@@ -47,6 +47,7 @@ void FLiveLinkHubModule::StartupModule()
 
 		LiveLinkHub = MakeShared<FLiveLinkHub>();
 		LiveLinkHub->Initialize(Ticker);
+		LiveLinkHub->OnFrameReceived().BindRaw(&Ticker, &FLiveLinkHubTicker::TriggerUpdate);
 	}
 
 	FPropertyEditorModule& PropertyModule = FModuleManager::LoadModuleChecked<FPropertyEditorModule>("PropertyEditor");
@@ -80,6 +81,8 @@ void FLiveLinkHubModule::ShutdownModule()
 
 	if (bCreateLiveLinkHubInstance)
 	{
+		LiveLinkHub->OnFrameReceived().Unbind();
+
 		Ticker.Stop();
 		Ticker.Exit();
 		// Application modes keep a shared reference to LiveLinkHub so we have to clear them before resetting the livelinkhub shared ptr.
diff --git a/Engine/Plugins/Animation/LiveLinkHub/Source/LiveLinkHub/Private/LiveLinkHubTicker.h b/Engine/Plugins/Animation/LiveLinkHub/Source/LiveLinkHub/Private/LiveLinkHubTicker.h
--- a/Engine/Plugins/Animation/LiveLinkHub/Source/LiveLinkHub/Private/LiveLinkHubTicker.h
+++ b/Engine/Plugins/Animation/LiveLinkHub/Source/LiveLinkHub/Private/LiveLinkHubTicker.h
@@ -32,11 +32,13 @@ public:
 		while (bIsRunning)
 		{
 			check(TickEvent);
-			TickEvent->Wait(TickTimeSpan);
+
+			bool bWasTriggeredTick = TickEvent->Wait(TickTimeSpan);
+			TRACE_CPUPROFILER_EVENT_SCOPE_CONDITIONAL(FLiveLinkHubTicker::TriggeredTick, bWasTriggeredTick);
+			TRACE_CPUPROFILER_EVENT_SCOPE_CONDITIONAL(FLiveLinkHubTicker::ScheduledTick, !bWasTriggeredTick);
 
 			if (bIsRunning) // make sure we were not told to exit during the wait
 			{
-				TRACE_CPUPROFILER_EVENT_SCOPE(FLiveLinkHubTicker::Tick);
 				OnTickDelegate.Broadcast();
 			}
 		}
@@ -67,6 +69,15 @@ public:
 	}
 	//~ End FRunnable Interface
 
+	void TriggerUpdate()
+	{
+		if (TickEvent)
+		{
+			TRACE_CPUPROFILER_EVENT_SCOPE(FLiveLinkHubTicker::TriggerUpdate);
+			TickEvent->Trigger();
+		}
+	}
+
 public:
 
 	/** Get the delegate called whenever this object ticks. */
