Only load the Sprite if we are not a Template object.  This saves 256 KB in non-Editor builds where the library might be present but we don't instantiatet he visualizers.

[RN]

diff --git a/Engine/Plugins/Animation/LiveLink/Source/LiveLink/Private/Visualizers/LiveLinkDataPreview.cpp b/Engine/Plugins/Animation/LiveLink/Source/LiveLink/Private/Visualizers/LiveLinkDataPreview.cpp
--- a/Engine/Plugins/Animation/LiveLink/Source/LiveLink/Private/Visualizers/LiveLinkDataPreview.cpp
+++ b/Engine/Plugins/Animation/LiveLink/Source/LiveLink/Private/Visualizers/LiveLinkDataPreview.cpp
@@ -38,13 +38,17 @@ ALiveLinkDataPreview::ALiveLinkDataPreview()
 	SetRootComponent(SceneRoot);
 	BillboardComponent = CreateDefaultSubobject<UBillboardComponent>("Sprite");
 	BillboardComponent->SetupAttachment(SceneRoot);
-	UTexture2D* SpriteTex = SpriteTexture.LoadSynchronous();
-	if(SpriteTex)
+
+	if (!IsTemplate())
 	{
-		BillboardComponent->SetSprite(SpriteTex);
-		BillboardComponent->bIsScreenSizeScaled = true;
-		BillboardComponent->ScreenSize = 0.0006;
-		BillboardComponent->SetRelativeLocation(FVector(0,0,10));
+		UTexture2D* SpriteTex = SpriteTexture.LoadSynchronous();
+		if (SpriteTex)
+		{
+			BillboardComponent->SetSprite(SpriteTex);
+			BillboardComponent->bIsScreenSizeScaled = true;
+			BillboardComponent->ScreenSize = 0.0006;
+			BillboardComponent->SetRelativeLocation(FVector(0, 0, 10));
+		}
 	}
 }
 
