[LiveLink]
- Fixes bug in the Live Link broadcast component where last bone in a skeleton is not present in the live link data.

diff --git a/Engine/Plugins/Animation/LiveLink/Source/LiveLink/Private/LiveLinkBroadcastComponent.cpp b/Engine/Plugins/Animation/LiveLink/Source/LiveLink/Private/LiveLinkBroadcastComponent.cpp
--- a/Engine/Plugins/Animation/LiveLink/Source/LiveLink/Private/LiveLinkBroadcastComponent.cpp
+++ b/Engine/Plugins/Animation/LiveLink/Source/LiveLink/Private/LiveLinkBroadcastComponent.cpp
@@ -166,7 +166,7 @@ TArray<FName> ULiveLinkBroadcastComponent::GetSkeletalBoneNames() const
 		}
 		else
 		{
-			for (int32 Index = 0; Index < SourceSkelMesh->GetNumBones() - 1; Index++)
+			for (int32 Index = 0; Index < SourceSkelMesh->GetNumBones(); Index++)
 			{
 				FName BoneName = SourceSkelMesh->GetBoneName(Index);
 				BoneNames.Emplace(BoneName);
@@ -195,7 +195,7 @@ TArray<int32> ULiveLinkBroadcastComponent::GetSkeletalBoneParents() const
 		}
 		else
 		{
-			for (int32 Index = 0; Index < SourceSkelMesh->GetNumBones() - 1; Index++)
+			for (int32 Index = 0; Index < SourceSkelMesh->GetNumBones(); Index++)
 			{
 				FName BoneName = SourceSkelMesh->GetBoneName(Index);
 				FName ParentBoneName = SourceSkelMesh->GetParentBone(BoneName);
@@ -224,7 +224,7 @@ TArray<FTransform> ULiveLinkBroadcastComponent::GetBoneTransforms() const
 		}
 		else
 		{
-			for (int32 Index = 0; Index < SourceSkelMesh->GetNumBones() - 1; Index++)
+			for (int32 Index = 0; Index < SourceSkelMesh->GetNumBones(); Index++)
 			{
 				FTransform BoneTransform = SourceSkelMesh->GetSocketTransform(SourceSkelMesh->GetBoneName(Index), RTS_ParentBoneSpace);
 				BoneTransforms.Add(BoneTransform);
