LiveLinkHub - Fix animation virtual subject not setting root bone parent to -1

diff --git a/Engine/Plugins/Animation/LiveLink/Source/LiveLink/Private/LiveLinkAnimationVirtualSubject.cpp b/Engine/Plugins/Animation/LiveLink/Source/LiveLink/Private/LiveLinkAnimationVirtualSubject.cpp
--- a/Engine/Plugins/Animation/LiveLink/Source/LiveLink/Private/LiveLinkAnimationVirtualSubject.cpp
+++ b/Engine/Plugins/Animation/LiveLink/Source/LiveLink/Private/LiveLinkAnimationVirtualSubject.cpp
@@ -49,7 +49,14 @@ namespace LiveLinkAnimationVirtualSubjectUtils
 			// 2) Bones with a parent need and offset based on the current size of the buffer
 			if (BoneParent == INDEX_NONE)
 			{
-				BoneParents.Add(0);
+				if (BoneParents.Num())
+				{
+					BoneParents.Add(0);
+				}
+				else
+				{
+					BoneParents.Add(INDEX_NONE);
+				}
 			}
 			else
 			{
