LiveLinkHub - Allow using a virtual subject as a timecode source

diff --git a/Engine/Plugins/Animation/LiveLink/Source/LiveLink/Private/LiveLinkAnimationVirtualSubject.cpp b/Engine/Plugins/Animation/LiveLink/Source/LiveLink/Private/LiveLinkAnimationVirtualSubject.cpp
--- a/Engine/Plugins/Animation/LiveLink/Source/LiveLink/Private/LiveLinkAnimationVirtualSubject.cpp
+++ b/Engine/Plugins/Animation/LiveLink/Source/LiveLink/Private/LiveLinkAnimationVirtualSubject.cpp
@@ -4,6 +4,7 @@
 
 #include "Algo/TopologicalSort.h"
 #include "ILiveLinkClient.h"
+#include "Misc/App.h"
 #include "Roles/LiveLinkAnimationRole.h"
 #include "Roles/LiveLinkAnimationTypes.h"
 
@@ -264,6 +265,20 @@ void ULiveLinkAnimationVirtualSubject::BuildFrame(const TArray<FLiveLinkSubjectF
 		{
 			NewSnapshotFrameData->MetaData.StringMetaData.Emplace(Subjects[Index], MetaDatum.Value);
 		}
+
+		if (SyncSubject == Subjects[Index])
+		{
+			// Assign timecode from the sync subject.
+			NewSnapshotFrameData->MetaData.SceneTime = SubjectSnapShotData.FrameData.GetBaseData()->MetaData.SceneTime;
+		}
+	}
+
+	if (NewSnapshotFrameData->MetaData.SceneTime.Time == 0)
+	{	
+		if (TOptional<FQualifiedFrameTime> FrameTime = FApp::GetCurrentFrameTime())
+		{
+			NewSnapshotFrameData->MetaData.SceneTime = *FrameTime;
+		}
 	}
 
 	ProcessAttachmentsForFrameData(NewSnapshotFrameData);
diff --git a/Engine/Plugins/Animation/LiveLink/Source/LiveLink/Private/LiveLinkClient.cpp b/Engine/Plugins/Animation/LiveLink/Source/LiveLink/Private/LiveLinkClient.cpp
--- a/Engine/Plugins/Animation/LiveLink/Source/LiveLink/Private/LiveLinkClient.cpp
+++ b/Engine/Plugins/Animation/LiveLink/Source/LiveLink/Private/LiveLinkClient.cpp
@@ -286,6 +286,13 @@ TSet<FLiveLinkSubjectKey> FLiveLinkClient::BuildThisTicksSubjectSnapshot(const T
 				{
 					VSubject->Update();
 					UpdatedSubjects.Add(SubjectItem.Key);
+				
+					// If we want the LiveLinkTimecodeProvider to pick up on the timecode from the Virtual Subject, then we 
+					// need to broadcast this delegate as if it were any other subject that received data.
+					if (VSubject->HasValidFrameSnapshot())
+					{
+						BroadcastFrameDataUpdate(VSubject->GetSubjectKey(), VSubject->GetFrameData());
+					}
 				}
 
 				EnabledSubjects.Add(SubjectItem.Key.SubjectName, SubjectItem.Key);
