LiveLinkHub - Fix Virtual Subjects's subject list not updating automatically

diff --git a/Engine/Plugins/Animation/LiveLink/Source/LiveLinkEditor/Private/LiveLinkVirtualSubjectDetailCustomization.cpp b/Engine/Plugins/Animation/LiveLink/Source/LiveLinkEditor/Private/LiveLinkVirtualSubjectDetailCustomization.cpp
--- a/Engine/Plugins/Animation/LiveLink/Source/LiveLinkEditor/Private/LiveLinkVirtualSubjectDetailCustomization.cpp
+++ b/Engine/Plugins/Animation/LiveLink/Source/LiveLinkEditor/Private/LiveLinkVirtualSubjectDetailCustomization.cpp
@@ -51,30 +51,7 @@ void FLiveLinkVirtualSubjectDetailCustomization::CustomizeDetails(IDetailLayoutB
 
 	DetailBuilder.HideProperty(SubjectsPropertyHandle);
 
-	SubjectsListItems.Reset();
-
-	if (Client)
-	{
-		TArray<FLiveLinkSubjectKey> SubjectKeys = Client->GetSubjects(/*bIncludeDisabledSubject*/ false, /*bIncludeVirtualSubject*/ false);
-		for (const FLiveLinkSubjectKey& SubjectKey : SubjectKeys)
-		{
-			TSubclassOf<ULiveLinkRole> LiveLinkRole = Client->GetSubjectRole_AnyThread(SubjectKey);
-			if (LiveLinkRole && (LiveLinkRole->IsChildOf(ULiveLinkAnimationRole::StaticClass()) || LiveLinkRole->IsChildOf(ULiveLinkBasicRole::StaticClass())))
-			{
-				SubjectsListItems.AddUnique(MakeShared<FName>(SubjectKey.SubjectName.Name));
-			}
-		}
-	}
-
-	// In case one of the associated subject linked to this virtual one doesn't exist anymore, add it to the list to display it red
-	for (const FLiveLinkSubjectName& SelectedSubject : Subject->GetSubjects())
-	{
-		if (false == SubjectsListItems.ContainsByPredicate([SelectedSubject](const FSubjectEntryPtr& Other) { return *Other == SelectedSubject.Name; }))
-		{
-			SubjectsListItems.AddUnique(MakeShared<FName>(SelectedSubject.Name));
-		}
-	}
-
+	UpdateSubjectList();
 	UpdateSelectedSubjects();
 	
 	TSharedRef<SComboButton> ComboButton = SNew(SComboButton)
@@ -314,8 +291,41 @@ FText FLiveLinkVirtualSubjectDetailCustomization::HandleSubjectItemToolTip(FSubj
 	return Result;
 }
 
+void FLiveLinkVirtualSubjectDetailCustomization::UpdateSubjectList()
+{
+	SubjectsListItems.Reset();
+
+	if (Client)
+	{
+		TArray<FLiveLinkSubjectKey> SubjectKeys = Client->GetSubjects(/*bIncludeDisabledSubject*/ false, /*bIncludeVirtualSubject*/ false);
+		for (const FLiveLinkSubjectKey& SubjectKey : SubjectKeys)
+		{
+			TSubclassOf<ULiveLinkRole> LiveLinkRole = Client->GetSubjectRole_AnyThread(SubjectKey);
+			if (LiveLinkRole && (LiveLinkRole->IsChildOf(ULiveLinkAnimationRole::StaticClass()) || LiveLinkRole->IsChildOf(ULiveLinkBasicRole::StaticClass())))
+			{
+				SubjectsListItems.AddUnique(MakeShared<FName>(SubjectKey.SubjectName.Name));
+			}
+		}
+	}
+
+	if (ULiveLinkVirtualSubject* Subject = SubjectPtr.Get())
+	{
+		// In case one of the associated subject linked to this virtual one doesn't exist anymore, add it to the list to display it red
+		for (const FLiveLinkSubjectName& SelectedSubject : Subject->GetSubjects())
+		{
+			if (!SubjectsListItems.ContainsByPredicate([SelectedSubject](const FSubjectEntryPtr& Other) { return *Other == SelectedSubject.Name; }))
+			{
+				SubjectsListItems.AddUnique(MakeShared<FName>(SelectedSubject.Name));
+			}
+		}
+	}
+}
+
 TSharedRef<SWidget> FLiveLinkVirtualSubjectDetailCustomization::OnGetVirtualSubjectsMenu()
 {
+	UpdateSubjectList();
+	UpdateSelectedSubjects();
+
 	return SNew(SListView<FSubjectEntryPtr>)
 		.ListItemsSource(&SubjectsListItems)
 		.OnGenerateRow(this, &FLiveLinkVirtualSubjectDetailCustomization::OnGenerateWidgetForSubjectItem);
diff --git a/Engine/Plugins/Animation/LiveLink/Source/LiveLinkEditor/Private/LiveLinkVirtualSubjectDetailCustomization.h b/Engine/Plugins/Animation/LiveLink/Source/LiveLinkEditor/Private/LiveLinkVirtualSubjectDetailCustomization.h
--- a/Engine/Plugins/Animation/LiveLink/Source/LiveLinkEditor/Private/LiveLinkVirtualSubjectDetailCustomization.h
+++ b/Engine/Plugins/Animation/LiveLink/Source/LiveLinkEditor/Private/LiveLinkVirtualSubjectDetailCustomization.h
@@ -53,6 +53,9 @@ private:
 	// Returns whether an entry is currently selected.
 	bool IsEntrySelected(FSubjectEntryPtr EntryPtr);
 
+	// Update the list of LiveLink subjects.
+	void UpdateSubjectList();
+
 	// Update the list of selected subjects displayed above the subject picker.
 	void UpdateSelectedSubjects();
 
@@ -64,7 +67,6 @@ private:
 
 	// Sets the SyncSubject property on the virtual subject,
 	void SetSyncSubject(SLiveLinkSubjectRepresentationPicker::FLiveLinkSourceSubjectRole Role);
-
 private:
 	// The tile set being edited
 	TWeakObjectPtr<ULiveLinkVirtualSubject> SubjectPtr;
