LiveLink - Fix possible ensure caused by accessing Subject dropdwon with no valid objects

diff --git a/Engine/Plugins/Animation/LiveLink/Source/LiveLinkEditor/Private/LiveLinkSubjectNameDetailCustomization.cpp b/Engine/Plugins/Animation/LiveLink/Source/LiveLinkEditor/Private/LiveLinkSubjectNameDetailCustomization.cpp
--- a/Engine/Plugins/Animation/LiveLink/Source/LiveLinkEditor/Private/LiveLinkSubjectNameDetailCustomization.cpp
+++ b/Engine/Plugins/Animation/LiveLink/Source/LiveLinkEditor/Private/LiveLinkSubjectNameDetailCustomization.cpp
@@ -58,6 +58,11 @@ SLiveLinkSubjectRepresentationPicker::FLiveLinkSourceSubjectRole FLiveLinkSubjec
 
 void FLiveLinkSubjectNameDetailCustomization::SetValue(SLiveLinkSubjectRepresentationPicker::FLiveLinkSourceSubjectRole NewValue)
 {
+	if (!StructPropertyHandle->IsValidHandle() || StructPropertyHandle->GetNumOuterObjects() == 0)
+	{
+		return;
+	}
+
 	FStructProperty* StructProperty = CastFieldChecked<FStructProperty>(StructPropertyHandle->GetProperty());
 
 	TArray<void*> RawData;
