LiveLink - Fix Settings and Visibility icons

#rnx

diff --git a/Engine/Plugins/Animation/LiveLink/Source/LiveLinkEditor/Private/LiveLinkClientPanelToolbar.cpp b/Engine/Plugins/Animation/LiveLink/Source/LiveLinkEditor/Private/LiveLinkClientPanelToolbar.cpp
--- a/Engine/Plugins/Animation/LiveLink/Source/LiveLinkEditor/Private/LiveLinkClientPanelToolbar.cpp
+++ b/Engine/Plugins/Animation/LiveLink/Source/LiveLinkEditor/Private/LiveLinkClientPanelToolbar.cpp
@@ -146,9 +146,9 @@ void SLiveLinkClientPanelToolbar::Construct(const FArguments& Args, FLiveLinkCli
 					.IsChecked_Lambda([]() { return ECheckBoxState::Unchecked; })
 					.OnCheckStateChanged_Lambda([](ECheckBoxState CheckState){ FModuleManager::LoadModuleChecked<ISettingsModule>("Settings").ShowViewer("Project", "Plugins", "LiveLink"); })
 					[
-						SNew(STextBlock)
-						.Font(FAppStyle::Get().GetFontStyle("FontAwesome.14"))
-						.Text(FEditorFontGlyphs::Cogs)
+						SNew(SImage)
+						.ColorAndOpacity(FSlateColor::UseForeground())
+						.Image(FSlateIcon(FAppStyle::Get().GetStyleSetName(), "Icons.Toolbar.Settings").GetIcon())
 					]
 				]
 			]
diff --git a/Engine/Plugins/Animation/LiveLink/Source/LiveLinkEditor/Private/SLiveLinkDataView.cpp b/Engine/Plugins/Animation/LiveLink/Source/LiveLinkEditor/Private/SLiveLinkDataView.cpp
--- a/Engine/Plugins/Animation/LiveLink/Source/LiveLinkEditor/Private/SLiveLinkDataView.cpp
+++ b/Engine/Plugins/Animation/LiveLink/Source/LiveLinkEditor/Private/SLiveLinkDataView.cpp
@@ -137,10 +137,12 @@ void SLiveLinkDataView::Construct(const FArguments& Args, FLiveLinkClient* InCli
 					SNew(SHorizontalBox)
 					+ SHorizontalBox::Slot()
 					.AutoWidth()
+					.Padding(0, 0, 2, 0)
 					.VAlign(VAlign_Center)
 					[
 						SNew(SImage)
-						.Image(FAppStyle::GetBrush("GenericViewButton"))
+						.ColorAndOpacity(FSlateColor::UseForeground())
+						.Image(FSlateIcon(FAppStyle::Get().GetStyleSetName(), "Icons.Visibility").GetIcon())
 					]
 					+ SHorizontalBox::Slot()
 					.AutoWidth()
