LiveLink - Use FPlatformTime::Seconds directly when invoking the OnLiveLinkSubjectEvaluated delegate
Fix callsites of FApp::GetCurrentFrameTime to use the deprecated non-thread safe method that doesn't trigger an ensure.

A more comprehensive fix will be required to get the current frame time in a thread safe way.

#Synced-CL 48456341

diff --git a/Engine/Plugins/Animation/LiveLink/Source/LiveLink/Private/LiveLinkClient.cpp b/Engine/Plugins/Animation/LiveLink/Source/LiveLink/Private/LiveLinkClient.cpp
--- a/Engine/Plugins/Animation/LiveLink/Source/LiveLink/Private/LiveLinkClient.cpp
+++ b/Engine/Plugins/Animation/LiveLink/Source/LiveLink/Private/LiveLinkClient.cpp
@@ -113,10 +113,6 @@ void FLiveLinkClient::Initialize()
 {
 	LLM_SCOPE_BYTAG(LiveLink_LiveLinkClient);
 
-#if WITH_EDITOR
-	CachedEngineTime = 0.0;
-#endif
-
 	OnLiveLinkSubjectRemoved().AddRaw(this, &FLiveLinkClient::OnSubjectRemovedCallback);
 	FCoreDelegates::OnPreExit.AddRaw(this, &FLiveLinkClient::Shutdown);
 
@@ -362,11 +358,13 @@ void FLiveLinkClient::CacheValues()
 #if WITH_EDITOR
 
 PRAGMA_DISABLE_DEPRECATION_WARNINGS
-	const FAppTime& Time = FAppTime::Get_UnsafeNoCheck();
+	FAppTime Time = FAppTime::Get_UnsafeNoCheck();
 PRAGMA_ENABLE_DEPRECATION_WARNINGS
 
-	CachedEngineTime      = Time.CurrentTime;
-	CachedEngineFrameTime = Time.CurrentFrameTime;
+	if (Time.CurrentFrameTime.IsSet())
+	{
+		CachedEngineFrameTime = Time.CurrentFrameTime;
+	}
 #endif
 }
 
@@ -1201,7 +1199,13 @@ void FLiveLinkClient::PushSubjectFrameData_Internal(FPendingSubjectFrame&& Subje
 	if (SubjectFrameData.FrameData.GetBaseData())
 	{
 		SubjectFrameData.FrameData.GetBaseData()->ArrivalTime.WorldTime = FPlatformTime::Seconds();
-		const TOptional<FQualifiedFrameTime>& CurrentTime = FApp::GetCurrentFrameTime();
+
+		PRAGMA_DISABLE_DEPRECATION_WARNINGS
+		// Have to use the unsafe method since we're calling this outside the game thread.
+		FAppTime Time = FAppTime::Get_UnsafeNoCheck();
+		PRAGMA_ENABLE_DEPRECATION_WARNINGS
+
+		const TOptional<FQualifiedFrameTime>& CurrentTime = Time.CurrentFrameTime;
 		if (CurrentTime.IsSet())
 		{
 			SubjectFrameData.FrameData.GetBaseData()->ArrivalTime.SceneTime = *CurrentTime;
@@ -1832,7 +1836,7 @@ bool FLiveLinkClient::EvaluateFrame_AnyThread(FLiveLinkSubjectName InSubjectName
 #if WITH_EDITOR
 		if (OnLiveLinkSubjectEvaluated().IsBound())
 		{
-			FLiveLinkTime RequestedTime = FLiveLinkTime(CachedEngineTime, CachedEngineFrameTime.Get(FQualifiedFrameTime()));
+			FLiveLinkTime RequestedTime = FLiveLinkTime(FPlatformTime::Seconds(), CachedEngineFrameTime.Get(FQualifiedFrameTime()));
 			FLiveLinkTime ResultTime;
 			if (bResult)
 			{
diff --git a/Engine/Plugins/Animation/LiveLink/Source/LiveLink/Private/LiveLinkSubject.cpp b/Engine/Plugins/Animation/LiveLink/Source/LiveLink/Private/LiveLinkSubject.cpp
--- a/Engine/Plugins/Animation/LiveLink/Source/LiveLink/Private/LiveLinkSubject.cpp
+++ b/Engine/Plugins/Animation/LiveLink/Source/LiveLink/Private/LiveLinkSubject.cpp
@@ -188,8 +188,12 @@ void FLiveLinkSubject::Update()
 
 		case ELiveLinkSourceMode::EngineTime:
 		{
+			PRAGMA_DISABLE_DEPRECATION_WARNINGS
+			FAppTime Time = FAppTime::Get_UnsafeNoCheck();
+			PRAGMA_ENABLE_DEPRECATION_WARNINGS
+
 			//Compute the evaluation time using all offsets (user offset + clock offset + smooth offset)
-			const double ReadTime = FApp::GetCurrentTime() - CachedSettings.BufferSettings.EngineTimeOffset - CachedSettings.BufferSettings.EngineTimeClockOffset - CachedSettings.BufferSettings.SmoothEngineTimeOffset;
+			const double ReadTime = Time.CurrentTime - CachedSettings.BufferSettings.EngineTimeOffset - CachedSettings.BufferSettings.EngineTimeClockOffset - CachedSettings.BufferSettings.SmoothEngineTimeOffset;
 
 			UE_LOG(LogLiveLink, Verbose, TEXT("Subject '%s' Read Time: %lf    User Offset: %.4lf    Clock Offset: %.4lf    Smooth Offset: %.4lf")
 				, *SubjectKey.SubjectName.ToString()
@@ -477,7 +481,10 @@ void FLiveLinkSubject::AddFrameData(FLiveLinkFrameDataStruct&& InFrameData)
 		IncreaseFrameDroppedStat();
 	}
 
-	LastPushTime = FApp::GetCurrentTime();
+	PRAGMA_DISABLE_DEPRECATION_WARNINGS
+	FAppTime Time = FAppTime::Get_UnsafeNoCheck();
+	PRAGMA_ENABLE_DEPRECATION_WARNINGS
+	LastPushTime = Time.CurrentTime;
 }
 
 void FLiveLinkSubject::PreprocessFrame(FLiveLinkFrameDataStruct& InOutFrameData)
@@ -1093,7 +1100,11 @@ void FLiveLinkSubject::UpdateState()
 {
 	const ELiveLinkSubjectState PreviousState = State;
 
-	const bool bHasValidFrame = (FApp::GetCurrentTime() - GetLastPushTime() < GetDefault<ULiveLinkSettings>()->GetTimeWithoutFrameToBeConsiderAsInvalid());
+	PRAGMA_DISABLE_DEPRECATION_WARNINGS
+	FAppTime Time = FAppTime::Get_UnsafeNoCheck();
+	PRAGMA_ENABLE_DEPRECATION_WARNINGS
+
+	const bool bHasValidFrame = (Time.CurrentTime - GetLastPushTime() < GetDefault<ULiveLinkSettings>()->GetTimeWithoutFrameToBeConsiderAsInvalid());
 	State = (bHasValidFrame && HasValidFrameSnapshot()) ? ELiveLinkSubjectState::Connected : ELiveLinkSubjectState::Unresponsive;
 
 	if (State != PreviousState)
diff --git a/Engine/Plugins/Animation/LiveLink/Source/LiveLink/Public/LiveLinkClient.h b/Engine/Plugins/Animation/LiveLink/Source/LiveLink/Public/LiveLinkClient.h
--- a/Engine/Plugins/Animation/LiveLink/Source/LiveLink/Public/LiveLinkClient.h
+++ b/Engine/Plugins/Animation/LiveLink/Source/LiveLink/Public/LiveLinkClient.h
@@ -342,7 +342,6 @@ private:
 	FOnLiveLinkSubjectEvaluated OnLiveLinkSubjectEvaluatedDelegate;
 
 	/** Cached value of the engine timecode and frame rate*/
-	double CachedEngineTime;
 	TOptional<FQualifiedFrameTime> CachedEngineFrameTime;
 #endif
 };
