LiveLink - Fix assertion causd by calling GetCurrentTime outside the game thread

#rnx
#Synced-CL 48513711

diff --git a/Engine/Plugins/Animation/LiveLink/Source/LiveLink/Private/LiveLinkMessageBusSource.cpp b/Engine/Plugins/Animation/LiveLink/Source/LiveLink/Private/LiveLinkMessageBusSource.cpp
--- a/Engine/Plugins/Animation/LiveLink/Source/LiveLink/Private/LiveLinkMessageBusSource.cpp
+++ b/Engine/Plugins/Animation/LiveLink/Source/LiveLink/Private/LiveLinkMessageBusSource.cpp
@@ -101,7 +101,7 @@ void FLiveLinkMessageBusSource::Update()
 	else
 	{
 		const double HeartbeatTimeout = GetDefault<ULiveLinkSettings>()->GetMessageBusHeartbeatTimeout();
-		const double CurrentTime = FApp::GetCurrentTime();
+		const double CurrentTime = FPlatformTime::Seconds();
 
 		float BytesLastSecond = 0.f;
 		if (TOptional<FMessageTransportStatistics> Statistics = UE::LiveLinkHub::FClientNetworkStatisticsModel::GetLatestNetworkStatistics(ConnectionAddress))
