LiveLinkHub - Compress LiveLink track data before saving it to disk

LiveLink.CompressSubSectionData 0 to disable compression

#Synced-CL 48456341

diff --git a/Engine/Plugins/Animation/LiveLink/Source/LiveLinkMovieScene/Private/MovieScene/MovieSceneLiveLinkStructProperties.cpp b/Engine/Plugins/Animation/LiveLink/Source/LiveLinkMovieScene/Private/MovieScene/MovieSceneLiveLinkStructProperties.cpp
--- a/Engine/Plugins/Animation/LiveLink/Source/LiveLinkMovieScene/Private/MovieScene/MovieSceneLiveLinkStructProperties.cpp
+++ b/Engine/Plugins/Animation/LiveLink/Source/LiveLinkMovieScene/Private/MovieScene/MovieSceneLiveLinkStructProperties.cpp
@@ -3,16 +3,26 @@
 #include "MovieScene/MovieSceneLiveLinkStructProperties.h"
 
 #include "Channels/MovieSceneBoolChannel.h"
-#include "MovieScene/MovieSceneLiveLinkEnumHandler.h"
 #include "Channels/MovieSceneByteChannel.h"
-#include "MovieScene/MovieSceneLiveLinkPropertyHandler.h"
 #include "Channels/MovieSceneIntegerChannel.h"
-#include "MovieScene/MovieSceneLiveLinkTransformHandler.h"
 #include "Channels/MovieSceneStringChannel.h"
+#include "HAL/IConsoleManager.h"
+#include "MovieScene/MovieSceneLiveLinkEnumHandler.h"
+#include "MovieScene/MovieSceneLiveLinkPropertyHandler.h"
+#include "MovieScene/MovieSceneLiveLinkTransformHandler.h"
+#include "Serialization/MemoryReader.h"
+#include "Serialization/MemoryWriter.h"
+#include "Serialization/ObjectAndNameAsStringProxyArchive.h"
 #include "UObject/EnumProperty.h"
+#include "UObject/FortniteMainBranchObjectVersion.h"
 
 #include UE_INLINE_GENERATED_CPP_BY_NAME(MovieSceneLiveLinkStructProperties)
 
+static TAutoConsoleVariable<bool> CVarLiveLinkCompressSubSectionData(
+	TEXT("LiveLink.CompressSubSectionData"), 1,
+	TEXT("When enabled, compresses the LiveLink subsection data before serializing it to disk."),
+	ECVF_Default);
+
 FLiveLinkPropertyData::FLiveLinkPropertyData() = default;
 FLiveLinkPropertyData::~FLiveLinkPropertyData() = default;
 
@@ -62,7 +72,7 @@ namespace LiveLinkPropertiesUtils
 				}
 			}
 		}
-		
+
 		return TSharedPtr<IMovieSceneLiveLinkPropertyHandler>();
 	}
 
@@ -87,3 +97,61 @@ namespace LiveLinkPropertiesUtils
 
 }
 
+bool FLiveLinkSubSectionData::Serialize(FArchive& Ar)
+{
+	Ar.UsingCustomVersion(FFortniteMainBranchObjectVersion::GUID);
+
+	if (Ar.IsPersistent() && !Ar.IsObjectReferenceCollector())
+	{
+		TArray<uint8> DataBuffer;
+		int32 UncompressedBufferSize = 0;
+		const UScriptStruct* Struct = FLiveLinkSubSectionData::StaticStruct();
+
+		if (Ar.IsLoading())
+		{
+			if (Ar.CustomVer(FFortniteMainBranchObjectVersion::GUID) < FFortniteMainBranchObjectVersion::LiveLinkCompressSubSectionData)
+			{
+				// Fallback to tagged property serialization.
+				return false;
+			}
+
+			// Read uncompressed buffer size
+			Ar << UncompressedBufferSize;
+			DataBuffer.SetNumUninitialized(UncompressedBufferSize);
+
+			if (UncompressedBufferSize > 0)
+			{
+				Ar.SerializeCompressedNew(DataBuffer.GetData(), DataBuffer.Num(), NAME_Oodle, NAME_Oodle, COMPRESS_BiasMemory);
+			}
+
+			FMemoryReader Reader(DataBuffer);
+			FObjectAndNameAsStringProxyArchive StructAr(Reader, false);
+
+			Struct->SerializeTaggedProperties(StructAr, (uint8*)this, Struct, nullptr);
+			return true;
+		}
+		else if (CVarLiveLinkCompressSubSectionData.GetValueOnAnyThread())
+		{
+			FMemoryWriter MemoryWriter(DataBuffer);
+
+			FObjectAndNameAsStringProxyArchive StructAr(MemoryWriter, false);
+
+			Struct->SerializeTaggedProperties(StructAr, (uint8*)this, Struct, nullptr);
+
+			// Store uncompressed buffer size
+			UncompressedBufferSize = DataBuffer.Num();
+
+			Ar << UncompressedBufferSize;
+
+			if (UncompressedBufferSize > 0)
+			{
+				Ar.SerializeCompressedNew(DataBuffer.GetData(), DataBuffer.Num(), NAME_Oodle, NAME_Oodle, COMPRESS_BiasMemory);
+			}
+
+			return true;
+		}
+	}
+
+	// Use tagged property serialization
+	return false;
+}
\ No newline at end of file
diff --git a/Engine/Plugins/Animation/LiveLink/Source/LiveLinkMovieScene/Public/MovieScene/MovieSceneLiveLinkStructProperties.h b/Engine/Plugins/Animation/LiveLink/Source/LiveLinkMovieScene/Public/MovieScene/MovieSceneLiveLinkStructProperties.h
--- a/Engine/Plugins/Animation/LiveLink/Source/LiveLinkMovieScene/Public/MovieScene/MovieSceneLiveLinkStructProperties.h
+++ b/Engine/Plugins/Animation/LiveLink/Source/LiveLinkMovieScene/Public/MovieScene/MovieSceneLiveLinkStructProperties.h
@@ -59,4 +59,17 @@ struct FLiveLinkSubSectionData
 public:
 	UPROPERTY()
 	TArray<FLiveLinkPropertyData> Properties;
+
+	/** Custom serializer that compresses property data before serializing it when doing persistent serialization. */
+	bool Serialize(FArchive& Archive);
+};
+
+
+template<>
+struct TStructOpsTypeTraits<FLiveLinkSubSectionData> : public TStructOpsTypeTraitsBase2<FLiveLinkSubSectionData>
+{
+	enum
+	{
+		WithSerializer = true
+	};
 };
diff --git a/Engine/Source/Runtime/Core/Public/UObject/FortniteMainBranchObjectVersions.inl b/Engine/Source/Runtime/Core/Public/UObject/FortniteMainBranchObjectVersions.inl
--- a/Engine/Source/Runtime/Core/Public/UObject/FortniteMainBranchObjectVersions.inl
+++ b/Engine/Source/Runtime/Core/Public/UObject/FortniteMainBranchObjectVersions.inl
@@ -712,5 +712,10 @@ FORTNITE_MAIN_VERSION(MediaColorSettingsUnification, 232)
 // Changed behavior of the String Operation metadata node in PCG for Substring and Matches
 FORTNITE_MAIN_VERSION(PCGMetadataOpStringTypeChange, 233)
 
+// Compress LiveLink SubSection data
+FORTNITE_MAIN_VERSION(LiveLinkCompressSubSectionData, 234)
+
+
+
 // -----<new versions can be added above this line>-------------------------------------------------
 #undef FORTNITE_MAIN_VERSION
