LiveLink - In playback, remove "Playback" from the source name, as it will be moved to its own column

#Synced-CL 48735856

diff --git a/Engine/Plugins/Animation/LiveLink/Source/LiveLink/Private/LiveLinkClient.cpp b/Engine/Plugins/Animation/LiveLink/Source/LiveLink/Private/LiveLinkClient.cpp
--- a/Engine/Plugins/Animation/LiveLink/Source/LiveLink/Private/LiveLinkClient.cpp
+++ b/Engine/Plugins/Animation/LiveLink/Source/LiveLink/Private/LiveLinkClient.cpp
@@ -701,7 +701,7 @@ bool FLiveLinkClient::CreateSource(const FLiveLinkSourcePreset& InSourcePreset)
 			return false;
 		}
 
-		Source = InSourcePreset.Settings->Factory.Get()->GetDefaultObject<ULiveLinkSourceFactory>()->CreateSource(InSourcePreset.Settings->ConnectionString);
+		Source = InSourcePreset.Settings->Factory.Get()->GetDefaultObject<ULiveLinkSourceFactory>()->CreateSource(InSourcePreset);
 		if (!Source.IsValid())
 		{
 			FLiveLinkLog::Warning(TEXT("Create Source Failure: The source couldn't be created by the factory."));
@@ -1738,13 +1738,12 @@ FText FLiveLinkClient::GetSourceNameOverride(const FLiveLinkSubjectKey& SubjectK
 	FText SourceType = GetSourceType(SubjectKey.Source);
 	FText SourceNameOverride = SourceType;
 
-	UObject* Settings = GetSubjectSettings(SubjectKey);
-	if (ULiveLinkSubjectSettings* SubjectSettings = Cast<ULiveLinkSubjectSettings>(Settings))
+	ULiveLinkSubjectSettings* SubjectSettings = Cast<ULiveLinkSubjectSettings>(GetSubjectSettings(SubjectKey));
+	ULiveLinkSourceSettings* SourceSettings = GetSourceSettings(SubjectKey.Source);
+
+	if (SubjectSettings && SourceSettings)
 	{
-		if (!SubjectSettings->OriginalSourceName.IsNone())
-		{
-			SourceNameOverride = FText::Format(INVTEXT("{0} ({1})"), FText::FromName(SubjectSettings->OriginalSourceName), SourceType);
-		}
+		SourceNameOverride = SourceSettings->GetSourceNameOverride(SubjectSettings, SourceType);
 	}
 
 	return SourceNameOverride;
diff --git a/Engine/Plugins/Animation/LiveLinkHub/Source/LiveLinkHub/Private/Recording/LiveLinkHubPlaybackSourceFactory.h b/Engine/Plugins/Animation/LiveLinkHub/Source/LiveLinkHub/Private/Recording/LiveLinkHubPlaybackSourceFactory.h
--- a/Engine/Plugins/Animation/LiveLinkHub/Source/LiveLinkHub/Private/Recording/LiveLinkHubPlaybackSourceFactory.h
+++ b/Engine/Plugins/Animation/LiveLinkHub/Source/LiveLinkHub/Private/Recording/LiveLinkHubPlaybackSourceFactory.h
@@ -5,6 +5,7 @@
 #include "LiveLinkSourceFactory.h"
 
 #include "LiveLinkPlaybackSource.h"
+#include "LiveLinkPresetTypes.h"
 
 #include "LiveLinkHubPlaybackSourceFactory.generated.h"
 
@@ -32,6 +33,11 @@ public:
 	{
 		return MakeShared<FLiveLinkPlaybackSource>(ConnectionString);
 	}
+
+	virtual TSharedPtr<ILiveLinkSource> CreateSource(const FLiveLinkSourcePreset& SourcePreset)
+	{
+		return MakeShared<FLiveLinkPlaybackSource>(*SourcePreset.SourceType.ToString());
+	}
 	//~ End ULiveLinkSourceFactory interface
 };
 
diff --git a/Engine/Plugins/Animation/LiveLinkHub/Source/LiveLinkHub/Private/Recording/LiveLinkHubPlaybackSourceSettings.cpp b/Engine/Plugins/Animation/LiveLinkHub/Source/LiveLinkHub/Private/Recording/LiveLinkHubPlaybackSourceSettings.cpp
--- a/Engine/Plugins/Animation/LiveLinkHub/Source/LiveLinkHub/Private/Recording/LiveLinkHubPlaybackSourceSettings.cpp
+++ b/Engine/Plugins/Animation/LiveLinkHub/Source/LiveLinkHub/Private/Recording/LiveLinkHubPlaybackSourceSettings.cpp
@@ -1,11 +1,25 @@
 // Copyright Epic Games, Inc. All Rights Reserved.
 
 #include "LiveLinkHubPlaybackSourceSettings.h"
+
 #include "LiveLinkHubPlaybackSourceFactory.h"
 #include "LiveLinkSourceFactory.h"
+#include "LiveLinkSubjectSettings.h"
 
 
 ULiveLinkHubPlaybackSourceSettings::ULiveLinkHubPlaybackSourceSettings()
 {
 	Factory = ULiveLinkHubPlaybackSourceFactory::StaticClass();
 }
+
+FText ULiveLinkHubPlaybackSourceSettings::GetSourceNameOverride(ULiveLinkSubjectSettings* SubjectSettings, FText SourceType)
+{
+	check(SubjectSettings);
+
+	if (!SubjectSettings->OriginalSourceName.IsNone())
+	{
+		return FText::FromName(SubjectSettings->OriginalSourceName);
+	}
+
+	return SourceType;
+}
diff --git a/Engine/Plugins/Animation/LiveLinkHub/Source/LiveLinkHub/Private/Recording/LiveLinkHubPlaybackSourceSettings.h b/Engine/Plugins/Animation/LiveLinkHub/Source/LiveLinkHub/Private/Recording/LiveLinkHubPlaybackSourceSettings.h
--- a/Engine/Plugins/Animation/LiveLinkHub/Source/LiveLinkHub/Private/Recording/LiveLinkHubPlaybackSourceSettings.h
+++ b/Engine/Plugins/Animation/LiveLinkHub/Source/LiveLinkHub/Private/Recording/LiveLinkHubPlaybackSourceSettings.h
@@ -14,4 +14,6 @@ public:
 	GENERATED_BODY()
 
 	ULiveLinkHubPlaybackSourceSettings();
+
+	virtual FText GetSourceNameOverride(class ULiveLinkSubjectSettings* SubjectSettings, FText SourceType) override;
 };
diff --git a/Engine/Plugins/Animation/LiveLinkHub/Source/LiveLinkHub/Private/Recording/LiveLinkPlaybackSource.h b/Engine/Plugins/Animation/LiveLinkHub/Source/LiveLinkHub/Private/Recording/LiveLinkPlaybackSource.h
--- a/Engine/Plugins/Animation/LiveLinkHub/Source/LiveLinkHub/Private/Recording/LiveLinkPlaybackSource.h
+++ b/Engine/Plugins/Animation/LiveLinkHub/Source/LiveLinkHub/Private/Recording/LiveLinkPlaybackSource.h
@@ -18,7 +18,30 @@ public:
 
 	FLiveLinkPlaybackSource(const FString& OriginalSourceName)
 	{
-		SourceName = *FString::Printf(TEXT("Playback (%s)"), *OriginalSourceName);
+		// Upgrading a recording has a side effect of re-saving the source with "Playback (SourceName)"
+		// We want to strip the Playback prefix and only retain the source name, since "Playback" will be displayed in a different column.
+
+		static const FString PlaybackPrefix = TEXT("Playback (");
+		const int32 NameStartIndex = PlaybackPrefix.Len();
+
+		if (OriginalSourceName.StartsWith(PlaybackPrefix))
+		{
+			int32 LastParenthesisIndex = 0;
+			OriginalSourceName.FindLastChar(TEXT(')'), LastParenthesisIndex);
+			
+			if (LastParenthesisIndex - NameStartIndex > 0)
+			{
+				SourceName = *OriginalSourceName.Mid(NameStartIndex, LastParenthesisIndex - NameStartIndex);
+			}
+			else
+			{
+				SourceName = "Unnamed Source";
+			}
+		}
+		else
+		{
+			SourceName = *OriginalSourceName;
+		}
 	}
 
 	virtual ~FLiveLinkPlaybackSource() = default;
diff --git a/Engine/Source/Runtime/LiveLinkInterface/Private/LiveLinkSourceFactory.cpp b/Engine/Source/Runtime/LiveLinkInterface/Private/LiveLinkSourceFactory.cpp
--- a/Engine/Source/Runtime/LiveLinkInterface/Private/LiveLinkSourceFactory.cpp
+++ b/Engine/Source/Runtime/LiveLinkInterface/Private/LiveLinkSourceFactory.cpp
@@ -1,6 +1,8 @@
 // Copyright Epic Games, Inc. All Rights Reserved.
 
 #include "LiveLinkSourceFactory.h"
+#include "LiveLinkPresetTypes.h"
+#include "LiveLinkSourceSettings.h"
 
 #include UE_INLINE_GENERATED_CPP_BY_NAME(LiveLinkSourceFactory)
 
@@ -20,3 +22,8 @@ TSharedPtr<ILiveLinkSource> ULiveLinkSourceFactory::OnSourceCreationPanelClosed(
 	return TSharedPtr<ILiveLinkSource>();
 }
 
+TSharedPtr<ILiveLinkSource> ULiveLinkSourceFactory::CreateSource(const FLiveLinkSourcePreset& SourcePreset)
+{
+	check(SourcePreset.Settings);
+	return CreateSource(SourcePreset.Settings->ConnectionString);
+}
\ No newline at end of file
diff --git a/Engine/Source/Runtime/LiveLinkInterface/Private/LiveLinkSourceSettings.cpp b/Engine/Source/Runtime/LiveLinkInterface/Private/LiveLinkSourceSettings.cpp
--- a/Engine/Source/Runtime/LiveLinkInterface/Private/LiveLinkSourceSettings.cpp
+++ b/Engine/Source/Runtime/LiveLinkInterface/Private/LiveLinkSourceSettings.cpp
@@ -1,6 +1,7 @@
 // Copyright Epic Games, Inc. All Rights Reserved.
 
 #include "LiveLinkSourceSettings.h"
+#include "LiveLinkSubjectSettings.h"
 #include "UObject/EnterpriseObjectVersion.h"
 
 #include UE_INLINE_GENERATED_CPP_BY_NAME(LiveLinkSourceSettings)
@@ -13,6 +14,15 @@ ULiveLinkSourceSettings::ULiveLinkSourceSettings()
 	}
 }
 
+FText ULiveLinkSourceSettings::GetSourceNameOverride(ULiveLinkSubjectSettings* SubjectSettings, FText SourceType)
+{
+	if (!SubjectSettings->OriginalSourceName.IsNone())
+	{
+		return FText::Format(INVTEXT("{0} ({1})"), FText::FromName(SubjectSettings->OriginalSourceName), SourceType);
+	}
+	return SourceType;
+}
+
 void ULiveLinkSourceSettings::Serialize(FArchive& Ar)
 {
 	Super::Serialize(Ar);
diff --git a/Engine/Source/Runtime/LiveLinkInterface/Public/LiveLinkSourceFactory.h b/Engine/Source/Runtime/LiveLinkInterface/Public/LiveLinkSourceFactory.h
--- a/Engine/Source/Runtime/LiveLinkInterface/Public/LiveLinkSourceFactory.h
+++ b/Engine/Source/Runtime/LiveLinkInterface/Public/LiveLinkSourceFactory.h
@@ -15,6 +15,7 @@
 
 #include "LiveLinkSourceFactory.generated.h"
 
+struct FLiveLinkSourcePreset;
 class ILiveLinkClient;
 class ILiveLinkSource;
 class SWidget;
@@ -65,6 +66,9 @@ public:
 	/** Create a new source from a ConnectionString */
 	virtual TSharedPtr<ILiveLinkSource> CreateSource(const FString& ConnectionString) const PURE_VIRTUAL(ULiveLinkSourceFactory::CreateSource, return TSharedPtr<ILiveLinkSource>(); );
 
+	/** Override this method in child classes to allow creating a source with more context from the source preset. */
+	LIVELINKINTERFACE_API virtual TSharedPtr<ILiveLinkSource> CreateSource(const FLiveLinkSourcePreset& SourcePreset);
+
 	/** Whether this factory is enabled. Can be overriden to disable it in certain configurations. */
 	virtual bool IsEnabled() const { return true; }
 };
diff --git a/Engine/Source/Runtime/LiveLinkInterface/Public/LiveLinkSourceSettings.h b/Engine/Source/Runtime/LiveLinkInterface/Public/LiveLinkSourceSettings.h
--- a/Engine/Source/Runtime/LiveLinkInterface/Public/LiveLinkSourceSettings.h
+++ b/Engine/Source/Runtime/LiveLinkInterface/Public/LiveLinkSourceSettings.h
@@ -202,6 +202,9 @@ public:
 	TArray<FLiveLinkSourceDebugInfo> SourceDebugInfos_DEPRECATED;
 #endif
 
+	/** Get the source name that should be displayed in the UI, based on the subject's source name override and the source type. */
+	LIVELINKINTERFACE_API virtual FText GetSourceNameOverride(class ULiveLinkSubjectSettings* SubjectSettings, FText SourceType);
+
 	LIVELINKINTERFACE_API virtual void Serialize(FArchive& Ar) override;
 
 #if WITH_EDITOR
