LiveLink - Create combined view of subjects, sources and devices

#Synced-CL 48735856

diff --git a/Engine/Plugins/Animation/LiveLink/LiveLink.uplugin b/Engine/Plugins/Animation/LiveLink/LiveLink.uplugin
--- a/Engine/Plugins/Animation/LiveLink/LiveLink.uplugin
+++ b/Engine/Plugins/Animation/LiveLink/LiveLink.uplugin
@@ -16,6 +16,10 @@
 	"Installed": false,
 	"SupportedPrograms": [ "LiveLinkHub" ],
 	"Plugins": [
+		{
+			"Name": "LiveLinkDevice",
+			"Enabled": true
+		},
 		{
 			"Name": "Takes",
 			"Enabled": true
@@ -27,7 +31,7 @@
 		},
 		{
 			"Name": "ContentBrowserAssetDataSource",
-			"Enabled":  true
+			"Enabled": true
 		}
 	],
 	"Modules": [
diff --git a/Engine/Plugins/Animation/LiveLink/Source/LiveLinkEditor/LiveLinkEditor.Build.cs b/Engine/Plugins/Animation/LiveLink/Source/LiveLinkEditor/LiveLinkEditor.Build.cs
--- a/Engine/Plugins/Animation/LiveLink/Source/LiveLinkEditor/LiveLinkEditor.Build.cs
+++ b/Engine/Plugins/Animation/LiveLink/Source/LiveLinkEditor/LiveLinkEditor.Build.cs
@@ -33,6 +33,7 @@ namespace UnrealBuildTool.Rules
 					"KismetCompiler",
 					"GraphEditor",
 					"LiveLinkComponents",
+					"LiveLinkDevice",
 					"LiveLinkGraphNode",
 					"LiveLinkMessageBusFramework",
 					"LiveLinkMovieScene",
diff --git a/Engine/Plugins/Animation/LiveLink/Source/LiveLinkEditor/Private/LiveLinkClientCommands.cpp b/Engine/Plugins/Animation/LiveLink/Source/LiveLinkEditor/Private/LiveLinkClientCommands.cpp
--- a/Engine/Plugins/Animation/LiveLink/Source/LiveLinkEditor/Private/LiveLinkClientCommands.cpp
+++ b/Engine/Plugins/Animation/LiveLink/Source/LiveLinkEditor/Private/LiveLinkClientCommands.cpp
@@ -14,6 +14,7 @@ void FLiveLinkClientCommands::RegisterCommands()
 	UI_COMMAND(RemoveSource, "Remove Selected Source(s)", "Remove selected live link source", EUserInterfaceActionType::Button, FInputChord());
 	UI_COMMAND(RemoveAllSources, "Remove All Sources", "Remove all live link sources", EUserInterfaceActionType::Button, FInputChord());
 	UI_COMMAND(RemoveSubject, "Remove Subject", "Remove selected live link subject", EUserInterfaceActionType::Button, FInputChord());
+	UI_COMMAND(RemoveDevice, "Remove Device", "Remove selected live link device", EUserInterfaceActionType::Button, FInputChord());
 	UI_COMMAND(PauseSubject, "Pause/Unpause Subject", "Pause or unpause selected live link subject", EUserInterfaceActionType::Button, FInputChord());
 }
 
diff --git a/Engine/Plugins/Animation/LiveLink/Source/LiveLinkEditor/Private/LiveLinkClientCommands.h b/Engine/Plugins/Animation/LiveLink/Source/LiveLinkEditor/Private/LiveLinkClientCommands.h
--- a/Engine/Plugins/Animation/LiveLink/Source/LiveLinkEditor/Private/LiveLinkClientCommands.h
+++ b/Engine/Plugins/Animation/LiveLink/Source/LiveLinkEditor/Private/LiveLinkClientCommands.h
@@ -13,6 +13,7 @@ public:
 	TSharedPtr<FUICommandInfo> RemoveAllSources;
 	TSharedPtr<FUICommandInfo> RemoveSubject;
 	TSharedPtr<FUICommandInfo> PauseSubject;
+	TSharedPtr<FUICommandInfo> RemoveDevice;
 
 	/** Initialize commands */
 	virtual void RegisterCommands() override;
diff --git a/Engine/Plugins/Animation/LiveLink/Source/LiveLinkEditor/Private/LiveLinkClientPanel.cpp b/Engine/Plugins/Animation/LiveLink/Source/LiveLinkEditor/Private/LiveLinkClientPanel.cpp
--- a/Engine/Plugins/Animation/LiveLink/Source/LiveLinkEditor/Private/LiveLinkClientPanel.cpp
+++ b/Engine/Plugins/Animation/LiveLink/Source/LiveLinkEditor/Private/LiveLinkClientPanel.cpp
@@ -61,7 +61,6 @@ void SLiveLinkClientPanel::Construct(const FArguments& Args, FLiveLinkClient* In
 	[
 		SNew(SBorder)
 		.BorderImage(FAppStyle::GetBrush("MessageLog.ListBorder")) // set panel background color to same color as message log at the bottom
-		.Padding(FMargin(4.f, 4.f, 4.f, 4.f))
 		[
 			SNew(SVerticalBox)
 			+SVerticalBox::Slot()
@@ -72,61 +71,27 @@ void SLiveLinkClientPanel::Construct(const FArguments& Args, FLiveLinkClient* In
 			]
 			+SVerticalBox::Slot()
 			.FillHeight(1.f)
-			.Padding(FMargin(0.0f, 4.0f, 0.0f, 0.0f))
 			[
 				SNew(SSplitter)
 				.Orientation(EOrientation::Orient_Vertical)
 				+SSplitter::Slot()
-				.Value(0.8f)
+				.Value(0.6f)
 				[
-					SNew(SSplitter)
-					.Orientation(EOrientation::Orient_Horizontal)
-					+SSplitter::Slot()
-					.Value(0.5f)
+					SNew(SBorder)
+					.Clipping(EWidgetClipping::ClipToBounds)
+					.BorderImage(FAppStyle::GetBrush("ToolPanel.GroupBorder"))
+					.Padding(FMargin(4.0f, 4.0f))
 					[
-						SNew(SSplitter)
-						.Orientation(EOrientation::Orient_Vertical)
-						+ SSplitter::Slot()
-						.Value(0.5f)
-						[
-							SNew(SBorder)
-							.Clipping(EWidgetClipping::ClipToBounds)
-							.BorderImage(FAppStyle::GetBrush("ToolPanel.GroupBorder"))
-							.Padding(FMargin(4.0f, 4.0f))
-							[
-								PanelController->SourcesView->GetWidget()
-							]
-						]
-						+ SSplitter::Slot()
-						[
-							SNew(SBorder)
-							.Clipping(EWidgetClipping::ClipToBounds)
-							.BorderImage(FAppStyle::GetBrush("ToolPanel.GroupBorder"))
-							.Padding(FMargin(4.0f, 4.0f))
-							[
-								PanelController->SubjectsView->GetWidget()
-							]
-						]
-					]
-					+SSplitter::Slot()
-					.Value(0.5f)
-					[
-						SNew(SWidgetSwitcher)
-						.WidgetIndex(this, &SLiveLinkClientPanel::GetDetailWidgetIndex)
-						+SWidgetSwitcher::Slot()
-						[
-							//[0] Detail view for Source
-							PanelController->SourcesDetailsView.ToSharedRef()
-						]
-						+SWidgetSwitcher::Slot()
-						[
-							// [1] Detail view for Subject, Frame data & Static data
-							PanelController->SubjectsDetailsView.ToSharedRef()
-						]
+						PanelController->SubjectsView->GetWidget()
 					]
 				]
+				+ SSplitter::Slot()
+				.Value(0.3f)
+				[
+					PanelController->GetCombinedDetailsWidget()
+				]
 				+SSplitter::Slot()
-				.Value(0.2f)
+				.Value(0.1f)
 				[
 					SNew(SOverlay)
 					+SOverlay::Slot()
@@ -181,11 +146,6 @@ void SLiveLinkClientPanel::AddReferencedObjects(FReferenceCollector& Collector)
 	Collector.AddReferencedObjects(DetailsPanelEditorObjects);
 }
 
-int32 SLiveLinkClientPanel::GetDetailWidgetIndex() const
-{
-	return PanelController->SubjectsDetailsView->GetSubjectKey().Source.IsValid() && !PanelController->SubjectsDetailsView->GetSubjectKey().SubjectName.IsNone() ? 1 : 0;
-}
-
 bool SLiveLinkClientPanel::IsInReadOnlyMode() const
 {
 	return GetDefault<ULiveLinkEditorSettings>()->bReadOnly;
diff --git a/Engine/Plugins/Animation/LiveLink/Source/LiveLinkEditor/Private/LiveLinkClientPanel.h b/Engine/Plugins/Animation/LiveLink/Source/LiveLinkEditor/Private/LiveLinkClientPanel.h
--- a/Engine/Plugins/Animation/LiveLink/Source/LiveLinkEditor/Private/LiveLinkClientPanel.h
+++ b/Engine/Plugins/Animation/LiveLink/Source/LiveLinkEditor/Private/LiveLinkClientPanel.h
@@ -50,27 +50,22 @@ class SLiveLinkClientPanel : public SCompoundWidget, public FGCObject, public FE
 
 private:
 	// Controls whether the editor performance throttling warning should be visible
-	UE_API EVisibility ShowEditorPerformanceThrottlingWarning() const;
+	EVisibility ShowEditorPerformanceThrottlingWarning() const;
 
 	// Handle disabling of editor performance throttling
-	UE_API FReply DisableEditorPerformanceThrottling();
+	FReply DisableEditorPerformanceThrottling();
 
 	// Return the message count text
-	UE_API FText GetMessageCountText() const;
+	FText GetMessageCountText() const;
 
 	// Return the occurrence count and last time occurred text
-	UE_API FText GetSelectedMessageOccurrenceText() const;
-
-	UE_API int32 GetDetailWidgetIndex() const;
-	
-	// Handler for the subject tree context menu opening
-	UE_API TSharedPtr<SWidget> OnOpenVirtualSubjectContextMenu();
+	FText GetSelectedMessageOccurrenceText() const;
 
 	// Returns whether the panel is in read only mode or not.
-	UE_API bool IsInReadOnlyMode() const;
+	bool IsInReadOnlyMode() const;
 
 	/** Get widget visibility according to whether the panel is in read-only mode. */
-	UE_API EVisibility GetVisibilityBasedOnReadOnly() const;
+	EVisibility GetVisibilityBasedOnReadOnly() const;
 
 private:
 	// Handles callback connections between the sources, subjects and details views
diff --git a/Engine/Plugins/Animation/LiveLink/Source/LiveLinkEditor/Private/LiveLinkClientPanelViews.cpp b/Engine/Plugins/Animation/LiveLink/Source/LiveLinkEditor/Private/LiveLinkClientPanelViews.cpp
--- a/Engine/Plugins/Animation/LiveLink/Source/LiveLinkEditor/Private/LiveLinkClientPanelViews.cpp
+++ b/Engine/Plugins/Animation/LiveLink/Source/LiveLinkEditor/Private/LiveLinkClientPanelViews.cpp
@@ -12,9 +12,14 @@
 #include "IDetailsView.h"
 #include "ILiveLinkClient.h"
 #include "ILiveLinkModule.h"
+#include "ILiveLinkDeviceModule.h"
 #include "Internationalization/Text.h"
 #include "LiveLinkClient.h"
 #include "LiveLinkClientCommands.h"
+#include "LiveLinkDefaultColumns.h"
+#include "LiveLinkDevice.h"
+#include "LiveLinkDeviceCapability.h"
+#include "LiveLinkDeviceSubsystem.h"
 #include "LiveLinkSettings.h"
 #include "LiveLinkSubjectSettings.h"
 #include "LiveLinkTypes.h"
@@ -45,46 +50,6 @@ namespace SourceListUI
 	static const FName ActionsColumnName(TEXT("Action"));
 };
 
-// Static Subject UI FNames
-namespace SubjectTreeUI
-{
-	static const FName EnabledColumnName(TEXT("Enabled"));
-	static const FName NameColumnName(TEXT("Name"));
-	static const FName RoleColumnName(TEXT("Role"));
-	static const FName TranslatedRoleColumnName(TEXT("TranslatedRole"));
-	static const FName ActionsColumnName(TEXT("Action"));
-};
-
-namespace UE::LiveLink
-{
-TSharedPtr<IDetailsView> CreateSourcesDetailsView(const TSharedPtr<FLiveLinkSourcesView>& InSourcesView, const TAttribute<bool>& bInReadOnly)
-{
-	FPropertyEditorModule& PropertyEditorModule = FModuleManager::GetModuleChecked<FPropertyEditorModule>("PropertyEditor");
-
-	FDetailsViewArgs DetailsViewArgs;
-	DetailsViewArgs.bUpdatesFromSelection = false;
-	DetailsViewArgs.bLockable = false;
-	DetailsViewArgs.bShowPropertyMatrixButton = false;
-	DetailsViewArgs.NameAreaSettings = FDetailsViewArgs::HideNameArea;
-	DetailsViewArgs.ViewIdentifier = NAME_None;
-	TSharedPtr<IDetailsView> SettingsDetailsView = PropertyEditorModule.CreateDetailView(DetailsViewArgs);
-	// todo: use controller here instead of view widget
-	SettingsDetailsView->OnFinishedChangingProperties().AddRaw(InSourcesView.Get(), &FLiveLinkSourcesView::OnPropertyChanged);
-	SettingsDetailsView->SetIsPropertyEditingEnabledDelegate(FIsPropertyEditingEnabled::CreateLambda(
-		[bInReadOnly](){
-		return !bInReadOnly.Get();
-	}));
-
-	return SettingsDetailsView;
-}
-
-TSharedPtr<SLiveLinkDataView> CreateSubjectsDetailsView(FLiveLinkClient* InLiveLinkClient, const TAttribute<bool>& bInReadOnly)
-{
-	return SNew(SLiveLinkDataView, InLiveLinkClient)
-		.ReadOnly(bInReadOnly);
-}
-} // namespace UE::LiveLink
-
 
 /** Dialog to create a new virtual subject */
 class SVirtualSubjectCreateDialog : public SCompoundWidget
@@ -136,60 +101,60 @@ public:
 		ChildSlot
 			[
 				SNew(SBorder)
-					.BorderImage(FAppStyle::GetBrush("Menu.Background"))
+				.BorderImage(FAppStyle::GetBrush("Menu.Background"))
+				[
+					SNew(SBox)
 					[
-						SNew(SBox)
-							[
-								SNew(SVerticalBox)
-								+SVerticalBox::Slot()
-								.HAlign(HAlign_Fill)
-								.AutoHeight()
+						SNew(SVerticalBox)
+						+ SVerticalBox::Slot()
+						.HAlign(HAlign_Fill)
+						.AutoHeight()
+						[
+							TextEntry->AsShared()
+						]
+
+						+ SVerticalBox::Slot()
+						.FillHeight(1.0)
+						[
+							SNew(SBorder)
+								.BorderImage(FAppStyle::GetBrush("ToolPanel.GroupBorder"))
+								.Content()
 								[
-									TextEntry->AsShared()
-								]
-
-								+ SVerticalBox::Slot()
-								.FillHeight(1.0)
-								[
-									SNew(SBorder)
-										.BorderImage(FAppStyle::GetBrush("ToolPanel.GroupBorder"))
-										.Content()
-										[
-											SAssignNew(RoleClassPicker, SVerticalBox)
-										]
-								]
-
-								// Ok/Cancel buttons
-								+ SVerticalBox::Slot()
-								.AutoHeight()
-								.HAlign(HAlign_Right)
-								.VAlign(VAlign_Bottom)
-								.Padding(8)
-								[
-									SNew(SUniformGridPanel)
-										.SlotPadding(FAppStyle::GetMargin("StandardDialog.SlotPadding"))
-										.MinDesiredSlotWidth(FAppStyle::GetFloat("StandardDialog.MinDesiredSlotWidth"))
-										.MinDesiredSlotHeight(FAppStyle::GetFloat("StandardDialog.MinDesiredSlotHeight"))
-										+ SUniformGridPanel::Slot(0, 0)
-										[
-											SNew(SButton)
-												.HAlign(HAlign_Center)
-												.ContentPadding(FAppStyle::GetMargin("StandardDialog.ContentPadding"))
-												.OnClicked(this, &SVirtualSubjectCreateDialog::OkClicked)
-												.Text(LOCTEXT("AddVirtualSubjectAdd", "Add"))
-												.IsEnabled(this, &SVirtualSubjectCreateDialog::IsVirtualSubjectClassSelected)
-										]
-										+ SUniformGridPanel::Slot(1, 0)
-										[
-											SNew(SButton)
-												.HAlign(HAlign_Center)
-												.ContentPadding(FAppStyle::GetMargin("StandardDialog.ContentPadding"))
-												.OnClicked(this, &SVirtualSubjectCreateDialog::CancelClicked)
-												.Text(LOCTEXT("AddVirtualSubjectCancel", "Cancel"))
-										]
+									SAssignNew(RoleClassPicker, SVerticalBox)
 								]
+						]
+
+						// Ok/Cancel buttons
+						+ SVerticalBox::Slot()
+						.AutoHeight()
+						.HAlign(HAlign_Right)
+						.VAlign(VAlign_Bottom)
+						.Padding(8)
+						[
+							SNew(SUniformGridPanel)
+							.SlotPadding(FAppStyle::GetMargin("StandardDialog.SlotPadding"))
+							.MinDesiredSlotWidth(FAppStyle::GetFloat("StandardDialog.MinDesiredSlotWidth"))
+							.MinDesiredSlotHeight(FAppStyle::GetFloat("StandardDialog.MinDesiredSlotHeight"))
+							+ SUniformGridPanel::Slot(0, 0)
+							[
+								SNew(SButton)
+									.HAlign(HAlign_Center)
+									.ContentPadding(FAppStyle::GetMargin("StandardDialog.ContentPadding"))
+									.OnClicked(this, &SVirtualSubjectCreateDialog::OkClicked)
+									.Text(LOCTEXT("AddVirtualSubjectAdd", "Add"))
+									.IsEnabled(this, &SVirtualSubjectCreateDialog::IsVirtualSubjectClassSelected)
+							]
+							+ SUniformGridPanel::Slot(1, 0)
+							[
+								SNew(SButton)
+									.HAlign(HAlign_Center)
+									.ContentPadding(FAppStyle::GetMargin("StandardDialog.ContentPadding"))
+									.OnClicked(this, &SVirtualSubjectCreateDialog::CancelClicked)
+									.Text(LOCTEXT("AddVirtualSubjectCancel", "Cancel"))
 							]
+						]
 					]
+				]
 			];
 
 		MakeRoleClassPicker();
@@ -368,6 +333,247 @@ private:
 	bool bOkClicked;
 };
 
+namespace UE::LiveLink
+{
+// Source creation window is static so that the code in this namespace can be used by the subjects and the sources view. 
+// This code could be moved to the subjects view once the SourcesView is deprecated in 5.10
+static TSharedPtr<SWindow> SourceCreationWindow;
+
+PRAGMA_DISABLE_DEPRECATION_WARNINGS
+TSharedPtr<IDetailsView> CreateSourcesDetailsView(const TSharedPtr<FLiveLinkSourcesView>& InSourcesView, const TAttribute<bool>& bInReadOnly)
+{
+	FPropertyEditorModule& PropertyEditorModule = FModuleManager::GetModuleChecked<FPropertyEditorModule>("PropertyEditor");
+
+	FDetailsViewArgs DetailsViewArgs;
+	DetailsViewArgs.bUpdatesFromSelection = false;
+	DetailsViewArgs.bLockable = false;
+	DetailsViewArgs.bShowPropertyMatrixButton = false;
+	DetailsViewArgs.NameAreaSettings = FDetailsViewArgs::HideNameArea;
+	DetailsViewArgs.ViewIdentifier = NAME_None;
+	TSharedPtr<IDetailsView> SettingsDetailsView = PropertyEditorModule.CreateDetailView(DetailsViewArgs);
+	// todo: use controller here instead of view widget
+	SettingsDetailsView->OnFinishedChangingProperties().AddRaw(InSourcesView.Get(), &FLiveLinkSourcesView::OnPropertyChanged);
+	SettingsDetailsView->SetIsPropertyEditingEnabledDelegate(FIsPropertyEditingEnabled::CreateLambda(
+		[bInReadOnly]() {
+			return !bInReadOnly.Get();
+		}));
+
+	return SettingsDetailsView;
+}
+PRAGMA_ENABLE_DEPRECATION_WARNINGS
+
+TSharedPtr<IDetailsView> CreateSourcesDetailsView(FOnFinishedChangingProperties InChangedPropertiesDelegate, const TAttribute<bool>& bInReadOnly)
+{
+	FPropertyEditorModule& PropertyEditorModule = FModuleManager::GetModuleChecked<FPropertyEditorModule>("PropertyEditor");
+
+	FDetailsViewArgs DetailsViewArgs;
+	DetailsViewArgs.bUpdatesFromSelection = false;
+	DetailsViewArgs.bLockable = false;
+	DetailsViewArgs.bShowPropertyMatrixButton = false;
+	DetailsViewArgs.NameAreaSettings = FDetailsViewArgs::HideNameArea;
+	DetailsViewArgs.ViewIdentifier = NAME_None;
+	TSharedPtr<IDetailsView> SettingsDetailsView = PropertyEditorModule.CreateDetailView(DetailsViewArgs);
+
+	SettingsDetailsView->OnFinishedChangingProperties() = InChangedPropertiesDelegate;
+	SettingsDetailsView->SetIsPropertyEditingEnabledDelegate(FIsPropertyEditingEnabled::CreateLambda(
+		[bInReadOnly](){
+		return !bInReadOnly.Get();
+	}));
+
+	return SettingsDetailsView;
+}
+
+TSharedPtr<SLiveLinkDataView> CreateSubjectsDetailsView(FLiveLinkClient* InLiveLinkClient, const TAttribute<bool>& bInReadOnly)
+{
+	return SNew(SLiveLinkDataView, InLiveLinkClient)
+		.ReadOnly(bInReadOnly);
+}
+
+TSharedPtr<IDetailsView> CreateDevicesDetailsView(FOnFinishedChangingProperties OnFinishedChangingProperties, const TAttribute<bool>& bInReadOnly)
+{
+	FPropertyEditorModule& PropertyEditorModule = FModuleManager::GetModuleChecked<FPropertyEditorModule>("PropertyEditor");
+
+	FDetailsViewArgs DetailsViewArgs;
+	DetailsViewArgs.bUpdatesFromSelection = false;
+	DetailsViewArgs.bLockable = false;
+	DetailsViewArgs.bShowPropertyMatrixButton = false;
+	DetailsViewArgs.NameAreaSettings = FDetailsViewArgs::HideNameArea;
+	DetailsViewArgs.ViewIdentifier = NAME_None;
+	TSharedPtr<IDetailsView> SettingsDetailsView = PropertyEditorModule.CreateDetailView(DetailsViewArgs);
+
+	SettingsDetailsView->OnFinishedChangingProperties() = OnFinishedChangingProperties;
+	SettingsDetailsView->SetIsPropertyEditingEnabledDelegate(FIsPropertyEditingEnabled::CreateLambda(
+		[bInReadOnly]() {
+			return !bInReadOnly.Get();
+		}));
+
+	return SettingsDetailsView;
+}
+
+void OnSourceCreated(TSharedPtr<ILiveLinkSource> NewSource, FString ConnectionString, TSubclassOf<ULiveLinkSourceFactory> Factory)
+{
+	FLiveLinkClient* LiveLinkClient = (FLiveLinkClient*)&IModularFeatures::Get().GetModularFeature<ILiveLinkClient>(ILiveLinkClient::ModularFeatureName);
+
+	if (NewSource.IsValid())
+	{
+		FGuid NewSourceGuid = LiveLinkClient->AddSource(NewSource);
+		if (NewSourceGuid.IsValid())
+		{
+			if (ULiveLinkSourceSettings* Settings = LiveLinkClient->GetSourceSettings(NewSourceGuid))
+			{
+				Settings->ConnectionString = ConnectionString;
+				Settings->Factory = Factory;
+			}
+		}
+	}
+
+	if (SourceCreationWindow)
+	{
+		SourceCreationWindow->RequestDestroyWindow();
+		SourceCreationWindow.Reset();
+	}
+}
+
+void OpenCreateMenuWindow(TArray<TObjectPtr<ULiveLinkSourceFactory>> Factories, int32 FactoryIndex)
+{
+	if (Factories.IsValidIndex(FactoryIndex))
+	{
+		if (ULiveLinkSourceFactory* FactoryInstance = Factories[FactoryIndex])
+		{
+			TSharedPtr<SWidget> Widget = FactoryInstance->BuildCreationPanel(ULiveLinkSourceFactory::FOnLiveLinkSourceCreated::CreateStatic(&OnSourceCreated, TSubclassOf<ULiveLinkSourceFactory>(FactoryInstance->GetClass())));
+			if (Widget.IsValid())
+			{
+				FText CreateText = FText::Format(LOCTEXT("CreateSourceLabel", "Create {0} connection"), FactoryInstance->GetSourceDisplayName());
+				SourceCreationWindow = SNew(SWindow)
+					.Title(CreateText)
+					.ClientSize(FVector2D(400, 150))
+					.SizingRule(ESizingRule::Autosized)
+					.SupportsMinimize(false)
+					.SupportsMaximize(false)
+					[
+						Widget.ToSharedRef()
+					];
+
+				GEditor->EditorAddModalWindow(SourceCreationWindow.ToSharedRef());
+			}
+		}
+	}
+}
+
+void ExecuteCreateSource(TArray<TObjectPtr<ULiveLinkSourceFactory>> Factories, int32 FactoryIndex)
+{
+	if (Factories.IsValidIndex(FactoryIndex))
+	{
+		if (ULiveLinkSourceFactory* FactoryInstance = Factories[FactoryIndex])
+		{
+			OnSourceCreated(FactoryInstance->CreateSource(FString()), FString(), FactoryInstance->GetClass());
+		}
+	}
+}
+
+void AddVirtualSubject()
+{
+	FLiveLinkClient* LiveLinkClient = (FLiveLinkClient*)&IModularFeatures::Get().GetModularFeature<ILiveLinkClient>(ILiveLinkClient::ModularFeatureName);
+
+	TSharedRef<SVirtualSubjectCreateDialog> Dialog =
+		SNew(SVirtualSubjectCreateDialog)
+		.LiveLinkClient(LiveLinkClient);
+
+	Dialog->ConfigureVirtualSubject();
+}
+
+TArray<TObjectPtr<ULiveLinkSourceFactory>> GetAllSourceFactories()
+{
+	TArray<TObjectPtr<ULiveLinkSourceFactory>> Factories;
+
+	TArray<UClass*> Results;
+	GetDerivedClasses(ULiveLinkSourceFactory::StaticClass(), Results, true);
+	for (UClass* SourceFactory : Results)
+	{
+		if (!SourceFactory->HasAllClassFlags(CLASS_Abstract | CLASS_Deprecated | CLASS_NewerVersionExists) && CastChecked<ULiveLinkSourceFactory>(SourceFactory->GetDefaultObject())->IsEnabled())
+		{
+			Factories.Add(NewObject<ULiveLinkSourceFactory>(GetTransientPackage(), SourceFactory));
+		}
+	}
+
+	Algo::StableSort(Factories, [](ULiveLinkSourceFactory* LHS, ULiveLinkSourceFactory* RHS)
+	{
+		return LHS->GetSourceDisplayName().CompareTo(RHS->GetSourceDisplayName()) <= 0;
+	});
+
+	return Factories;
+}
+
+void GenerateSourceMenuEntries(const TArray<TObjectPtr<ULiveLinkSourceFactory>>& Factories, FMenuBuilder& MenuBuilder)
+{
+	MenuBuilder.BeginSection("SourceSection", LOCTEXT("Sources", "Live Link Sources"));
+
+	for (int32 FactoryIndex = 0; FactoryIndex < Factories.Num(); ++FactoryIndex)
+	{
+		ULiveLinkSourceFactory* FactoryInstance = Factories[FactoryIndex];
+		if (FactoryInstance)
+		{
+			ULiveLinkSourceFactory::EMenuType MenuType = FactoryInstance->GetMenuType();
+
+			if (MenuType == ULiveLinkSourceFactory::EMenuType::SubPanel)
+			{
+				MenuBuilder.AddMenuEntry(
+					FactoryInstance->GetSourceDisplayName(),
+					FactoryInstance->GetSourceTooltip(),
+					FSlateIcon(),
+					FUIAction(
+						FExecuteAction::CreateStatic(OpenCreateMenuWindow, Factories, FactoryIndex)
+					),
+					NAME_None,
+					EUserInterfaceActionType::Button);
+			}
+			else if (MenuType == ULiveLinkSourceFactory::EMenuType::MenuEntry)
+			{
+				MenuBuilder.AddMenuEntry(
+					FactoryInstance->GetSourceDisplayName(),
+					FactoryInstance->GetSourceTooltip(),
+					FSlateIcon(),
+					FUIAction(
+						FExecuteAction::CreateStatic(ExecuteCreateSource, Factories, FactoryIndex)
+					),
+					NAME_None,
+					EUserInterfaceActionType::Button);
+			}
+			else
+			{
+				MenuBuilder.AddMenuEntry(
+					FactoryInstance->GetSourceDisplayName(),
+					FactoryInstance->GetSourceTooltip(),
+					FSlateIcon(),
+					FUIAction(
+						FExecuteAction(),
+						FCanExecuteAction::CreateLambda([]() { return false; })
+					),
+					NAME_None,
+					EUserInterfaceActionType::Button);
+			}
+		}
+	}
+
+	MenuBuilder.EndSection();
+
+	MenuBuilder.BeginSection("VirtualSourceSection", LOCTEXT("VirtualSources", "Live Link VirtualSubject Sources"));
+
+	MenuBuilder.AddMenuEntry(
+		LOCTEXT("AddVirtualSubject", "Add Virtual Subject"),
+		LOCTEXT("AddVirtualSubject_Tooltip", "Adds a new virtual subject to Live Link. Instead of coming from a source a virtual subject is a combination of 2 or more real subjects"),
+		FSlateIcon(),
+		FUIAction(
+			FExecuteAction::CreateStatic(&AddVirtualSubject)
+		),
+		NAME_None,
+		EUserInterfaceActionType::Button);
+
+	MenuBuilder.EndSection();
+}
+
+} // namespace UE::LiveLink
+
 FGuid FLiveLinkSourceUIEntry::GetGuid() const
 {
 	return EntryGuid;
@@ -404,22 +610,38 @@ FText FLiveLinkSourceUIEntry::GetToolTip() const
 FLiveLinkSubjectUIEntry::FLiveLinkSubjectUIEntry(const FLiveLinkSubjectKey& InSubjectKey, FLiveLinkClient* InClient, bool bInIsSource)
 	: SubjectKey(InSubjectKey)
 	, Client(InClient)
-	, bIsSource(bInIsSource)
 {
+	Type = bInIsSource ? ELiveLinkEntryType::Source : ELiveLinkEntryType::Subject;
+
 	if (InClient)
 	{
 		bIsVirtualSubject = InClient->IsVirtualSubject(InSubjectKey);
+
+		// LiveLinkClient doesn't have a concept of playback sources since they're a LiveLinkHub concept.
+		// We could replace this by allowing sources to modify their appearance in the source list.
+		bIsPlaybackSource = IsSource() && Client->GetSourceSettings(SubjectKey.Source)->GetName().Contains(TEXT("Playback"));
 	}
 }
 
+FLiveLinkSubjectUIEntry::FLiveLinkSubjectUIEntry(FLiveLinkClient* InClient, ELiveLinkEntryType InType)
+	: Client(InClient)
+	, Type(InType)
+{
+}
+
 bool FLiveLinkSubjectUIEntry::IsSubject() const
 {
-	return !bIsSource;
+	return Type == ELiveLinkEntryType::Subject;
 }
 
 bool FLiveLinkSubjectUIEntry::IsSource() const
 {
-	return bIsSource;
+	return Type == ELiveLinkEntryType::Source;
+}
+
+bool FLiveLinkSubjectUIEntry::IsDevice() const
+{
+	return Type == ELiveLinkEntryType::Device;
 }
 
 bool FLiveLinkSubjectUIEntry::IsVirtualSubject() const
@@ -433,10 +655,14 @@ UObject* FLiveLinkSubjectUIEntry::GetSettings() const
 	{
 		return Client->GetSourceSettings(SubjectKey.Source);
 	}
-	else
+	else if (IsSubject())
 	{
 		return Client->GetSubjectSettings(SubjectKey);
 	}
+	else
+	{
+		return WeakDevice.Get();
+	}
 }
 
 bool FLiveLinkSubjectUIEntry::IsSubjectEnabled() const
@@ -449,6 +675,11 @@ bool FLiveLinkSubjectUIEntry::IsSubjectValid() const
 	return IsSubject() ? Client->IsSubjectValid(SubjectKey) : false;
 }
 
+bool FLiveLinkSubjectUIEntry::IsPlaybackSource() const
+{
+	return bIsPlaybackSource;
+}
+
 void FLiveLinkSubjectUIEntry::SetSubjectEnabled(bool bIsEnabled)
 {
 	if (IsSubject())
@@ -459,14 +690,39 @@ void FLiveLinkSubjectUIEntry::SetSubjectEnabled(bool bIsEnabled)
 
 FText FLiveLinkSubjectUIEntry::GetItemText() const
 {
-	if (IsSubject())
+	switch (Type)
 	{
+	case ELiveLinkEntryType::SourcesHeader:
+		return LOCTEXT("SourcesLabel", "Sources");
+	case ELiveLinkEntryType::DevicesHeader:
+		return LOCTEXT("DevicesLabel", "Devices");
+	case ELiveLinkEntryType::Subject:
 		return Client->GetSubjectDisplayName(SubjectKey);
+	case ELiveLinkEntryType::Source:
+		return Client->GetSourceNameOverride(SubjectKey);
+	case ELiveLinkEntryType::Device:
+	{
+		FText Name;
+		if (TStrongObjectPtr<ULiveLinkDevice> Device = WeakDevice.Pin())
+		{
+			Name = Device->GetDisplayName();
+		}
+		return Name;
 	}
-	else
+	default:
+		checkNoEntry();
+		return FText::GetEmpty();
+	}
+}
+
+FText FLiveLinkSubjectUIEntry::GetMachineName() const
+{
+	if (IsSource())
 	{
-		return Client->GetSourceNameOverride(SubjectKey);
+		return Client->GetSourceMachineName(SubjectKey.Source);
 	}
+
+	return FText::GetEmpty();
 }
 
 TSubclassOf<ULiveLinkRole> FLiveLinkSubjectUIEntry::GetItemRole() const
@@ -481,7 +737,24 @@ TSubclassOf<ULiveLinkRole> FLiveLinkSubjectUIEntry::GetItemTranslatedRole() cons
 
 void FLiveLinkSubjectUIEntry::RemoveFromClient() const
 {
-	Client->RemoveSubject_AnyThread(SubjectKey);
+	switch (Type)
+	{
+	case ELiveLinkEntryType::Subject:
+		Client->RemoveSubject_AnyThread(SubjectKey);
+		break;
+	case ELiveLinkEntryType::Source:
+		Client->RemoveSource(SubjectKey.Source);
+		break;
+	case ELiveLinkEntryType::Device:
+	{
+		if (ULiveLinkDevice* Device = WeakDevice.Get())
+		{
+			ULiveLinkDeviceSubsystem* Subsystem = GEngine->GetEngineSubsystem<ULiveLinkDeviceSubsystem>();
+			Subsystem->RemoveDevice(Device);
+		}
+		break;
+	}
+	}
 }
 
 void FLiveLinkSubjectUIEntry::PauseSubject()
@@ -533,7 +806,17 @@ public:
 	/** Overridden from SMultiColumnTableRow.  Generates a widget for this column of the list view. */
 	virtual TSharedRef<SWidget> GenerateWidgetForColumn(const FName& ColumnName) override
 	{
-		if (ColumnName == SubjectTreeUI::EnabledColumnName)
+		TDelegate<bool()> IsRowSelected = TDelegate<bool()>::CreateLambda([WeakRow = TWeakPtr<SLiveLinkClientPanelSubjectRow>(StaticCastSharedRef<SLiveLinkClientPanelSubjectRow>(AsShared()))]() -> bool
+		{
+			if (TSharedPtr<SLiveLinkClientPanelSubjectRow> Row = WeakRow.Pin())
+			{
+				return Row->IsSelected();
+			}
+			return false;
+		});
+
+		constexpr double BaseIndent = 12.0f;
+		if (ColumnName == UE::LiveLink::DefaultColumn::Action)
 		{
 			if (EntryPtr->IsSubject())
 			{
@@ -542,38 +825,107 @@ public:
 					.IsChecked(MakeAttributeSP(this, &SLiveLinkClientPanelSubjectRow::GetSubjectEnabled))
 					.OnCheckStateChanged(this, &SLiveLinkClientPanelSubjectRow::OnEnabledChanged);
 			}
+			else if (EntryPtr->IsDevice())
+			{
+				if (TStrongObjectPtr<ULiveLinkDevice> Device = EntryPtr->WeakDevice.Pin())
+				{
+					FLiveLinkDeviceWidgetArguments Args;
+					Args.IsRowSelected = IsRowSelected;
+					return Device->GenerateWidgetForColumn(UE::LiveLink::DefaultColumn::Action, Args);
+				}
+			}
+			else
+			{
+				const float IndentAmount = EntryPtr->IsSource() ? BaseIndent : 0.f;
+				return SNew(SBox)
+				.Padding(2.f, 0.f, 0.f, 0.f)
+				.HAlign(HAlign_Left)
+				[
+					SNew(SExpanderArrow, SharedThis(this)).IndentAmount(IndentAmount)
+				];
+			}
 		}
-		else if (ColumnName == SubjectTreeUI::NameColumnName)
+		else if (ColumnName == UE::LiveLink::DefaultColumn::Name)
 		{
+			FMargin Padding = EntryPtr->IsSource() || EntryPtr->IsDevice() ? FMargin(BaseIndent, 0, 0, 0) : FMargin(0);
+			if (EntryPtr->IsSubject())
+			{
+				Padding = FMargin(BaseIndent * 2, 0, 0, 0);
+			}
+
+			TAttribute<FSlateFontInfo> Font;
+			if (EntryPtr->Type == ELiveLinkEntryType::SourcesHeader || EntryPtr->Type == ELiveLinkEntryType::DevicesHeader)
+			{
+				Font = FCoreStyle::GetDefaultFontStyle("Bold", 10);
+			}
+
+			if (EntryPtr->IsPlaybackSource())
+			{
+				Font = FCoreStyle::GetDefaultFontStyle("Italic", 10);
+			}
+			
 			return SNew(SHorizontalBox)
 				+ SHorizontalBox::Slot()
-				.AutoWidth()
-				.Padding(6, 0, 0, 0)
-				[
-					SNew(SExpanderArrow, SharedThis(this)).IndentAmount(12)
-				]
-				+ SHorizontalBox::Slot()
+				.Padding(Padding)
 				.FillWidth(1.0f)
 				[
 					SNew(STextBlock)
+					.Font(Font)
 					.ColorAndOpacity(this, &SLiveLinkClientPanelSubjectRow::GetSubjectTextColor)
 					.Text(MakeAttributeSP(this, &SLiveLinkClientPanelSubjectRow::GetItemText))
 				];
 		}
-		else if (ColumnName == SubjectTreeUI::RoleColumnName)
+		else if (ColumnName == UE::LiveLink::DefaultColumn::Machine)
+		{
+			if (EntryPtr->IsSource())
+			{
+				return SNew(STextBlock)
+					.ColorAndOpacity(this, &SLiveLinkClientPanelSubjectRow::GetSubjectTextColor)
+					.Text(MakeAttributeSP(this, &SLiveLinkClientPanelSubjectRow::GetItemMachineName));
+			}
+			else if (EntryPtr->IsDevice()) 
+			{
+				if (TStrongObjectPtr<ULiveLinkDevice> Device = EntryPtr->WeakDevice.Pin())
+				{
+					FLiveLinkDeviceWidgetArguments Args;
+					Args.IsRowSelected = IsRowSelected;
+					return Device->GenerateWidgetForColumn(UE::LiveLink::DefaultColumn::Machine, Args);
+				}
+			}
+		}
+		else if (ColumnName == UE::LiveLink::DefaultColumn::Role)
 		{
-			TAttribute<FText> RoleAttribute = MakeAttributeSP(this, &SLiveLinkClientPanelSubjectRow::GetItemRole);
+			TAttribute<FText> RoleAttribute;
+			
+			if (EntryPtr->IsSubject())
+			{
+				RoleAttribute = MakeAttributeSP(this, &SLiveLinkClientPanelSubjectRow::GetItemRole);
+			}
+			else if (EntryPtr->IsPlaybackSource())
+			{
+				RoleAttribute = LOCTEXT("Playback", "Playback");
+			}
+			else
+			{
+				RoleAttribute = FText::GetEmpty();
+			}
+
+			TAttribute<FSlateColor> TextColorAttribute = MakeAttributeSP(this, &SLiveLinkClientPanelSubjectRow::GetSubjectTextColor);
+			if (EntryPtr->IsPlaybackSource())
+			{
+				TextColorAttribute = FSlateColor::UseSubduedForeground();
+			}
 
 			return SNew(SHorizontalBox)
 			+ SHorizontalBox::Slot()
 			.AutoWidth()
 			[
 				SNew(STextBlock)
-				.ColorAndOpacity(this, &SLiveLinkClientPanelSubjectRow::GetSubjectTextColor)
-				.Text(EntryPtr->IsSubject() ? RoleAttribute : FText::GetEmpty())
+				.ColorAndOpacity(TextColorAttribute)
+				.Text(RoleAttribute)
 			];
 		}
-		else if (ColumnName == SubjectTreeUI::TranslatedRoleColumnName)
+		else if (ColumnName == UE::LiveLink::DefaultColumn::TranslatedRole)
 		{
 			TAttribute<FText> TranslatedRoleAttribute = MakeAttributeSP(this, &SLiveLinkClientPanelSubjectRow::GetItemTranslatedRole);
 
@@ -586,16 +938,40 @@ public:
 					.Text(EntryPtr->IsSubject() ? TranslatedRoleAttribute : FText::GetEmpty())
 			];
 		}
-		else if (ColumnName == SubjectTreeUI::ActionsColumnName)
+		else if (ColumnName == UE::LiveLink::DefaultColumn::Status)
 		{
-			return SNew(SBox)
+			if (EntryPtr->IsDevice())
+			{
+				if (TStrongObjectPtr<ULiveLinkDevice> Device = EntryPtr->WeakDevice.Pin())
+				{
+					FLiveLinkDeviceWidgetArguments Args;
+					Args.IsRowSelected = IsRowSelected;
+					return Device->GenerateWidgetForColumn(UE::LiveLink::DefaultColumn::Status, Args);
+				}
+			}
+			else if (EntryPtr->IsPlaybackSource())
+			{
+				return SNew(SBox)
+					.WidthOverride(12.0)
+					.HeightOverride(12.0)
+					[
+						SNew(SImage)
+							.ColorAndOpacity(FSlateColor::UseSubduedForeground())
+							.ToolTipText(LOCTEXT("PlaybackSourceToolTip", "This source is being driven by a recording."))
+							.Image(FAppStyle::Get().GetBrush("Animation.Forward"))
+					];
+			}
+			else
+			{
+				return SNew(SBox)
 				.HAlign(HAlign_Center)
 				.VAlign(VAlign_Center)
 				[
 					SNew(SImage)
-					.Image(this, &SLiveLinkClientPanelSubjectRow::GetSubjectIcon)
-					.ToolTipText(this, &SLiveLinkClientPanelSubjectRow::GetSubjectIconToolTip)
+						.Image(this, &SLiveLinkClientPanelSubjectRow::GetSubjectIcon)
+						.ToolTipText(this, &SLiveLinkClientPanelSubjectRow::GetSubjectIconToolTip)
 				];
+			}
 		}
 
 		return SNullWidget::NullWidget;
@@ -605,6 +981,7 @@ private:
 	ECheckBoxState GetSubjectEnabled() const { return EntryPtr->IsSubjectEnabled() ? ECheckBoxState::Checked : ECheckBoxState::Unchecked; }
 	void OnEnabledChanged(ECheckBoxState NewState) { EntryPtr->SetSubjectEnabled(NewState == ECheckBoxState::Checked); }
 	FText GetItemText() const { return EntryPtr->GetItemText(); }
+	FText GetItemMachineName() const { return EntryPtr->GetMachineName(); }
 	FText GetItemRole() const
 	{
 		TSubclassOf<ULiveLinkRole> Role = EntryPtr->GetItemRole();
@@ -630,7 +1007,7 @@ private:
 	/** Get widget visibility according to whether or not the panel is in read-only mode. */
 	EVisibility GetVisibilityFromReadOnly() const
 	{
-		return bReadOnly.Get() ? EVisibility::Collapsed : EVisibility::Visible;
+		return bReadOnly.Get() ? EVisibility::Hidden : EVisibility::Visible;
 	}
 
 	/** Get the icon for a subject's status. */
@@ -793,28 +1170,14 @@ private:
 
 	/** Attribute used to query whether the panel is in read only mode or not. */
 	TAttribute<bool> bReadOnly;
-};
-
-FLiveLinkSourcesView::FLiveLinkSourcesView(FLiveLinkClient* InLiveLinkClient, TSharedPtr<FUICommandList> InCommandList, TAttribute<bool> bInReadOnly, FOnSourceSelectionChanged InOnSourceSelectionChanged)
-	: Client(InLiveLinkClient)
-	, OnSourceSelectionChangedDelegate(MoveTemp(InOnSourceSelectionChanged))
-	, bReadOnly(MoveTemp(bInReadOnly))
-{
-	TArray<UClass*> Results;
-	GetDerivedClasses(ULiveLinkSourceFactory::StaticClass(), Results, true);
-	for (UClass* SourceFactory : Results)
-	{
-		if (!SourceFactory->HasAllClassFlags(CLASS_Abstract | CLASS_Deprecated | CLASS_NewerVersionExists) && CastChecked<ULiveLinkSourceFactory>(SourceFactory->GetDefaultObject())->IsEnabled())
-		{
-			Factories.Add(NewObject<ULiveLinkSourceFactory>(GetTransientPackage(), SourceFactory));
-		}
-	}
-
-	Algo::StableSort(Factories, [](ULiveLinkSourceFactory* LHS, ULiveLinkSourceFactory* RHS)
-	{
-		return LHS->GetSourceDisplayName().CompareTo(RHS->GetSourceDisplayName()) <= 0;
-	});
+};
 
+FLiveLinkSourcesView::FLiveLinkSourcesView(FLiveLinkClient* InLiveLinkClient, TSharedPtr<FUICommandList> InCommandList, TAttribute<bool> bInReadOnly, FOnSourceSelectionChanged InOnSourceSelectionChanged)
+	: Client(InLiveLinkClient)
+	, OnSourceSelectionChangedDelegate(MoveTemp(InOnSourceSelectionChanged))
+	, bReadOnly(MoveTemp(bInReadOnly))
+{
+	Factories = UE::LiveLink::GetAllSourceFactories();
 	CreateSourcesListView(InCommandList);
 }
 
@@ -962,11 +1325,30 @@ bool FLiveLinkSourcesView::CanRemoveSource()
 	return SourcesListView->GetNumItemsSelected() > 0;
 }
 
+void FLiveLinkSourcesView::OnPropertyChanged(const FPropertyChangedEvent& InEvent)
+{
+	TArray<FLiveLinkSourceUIEntryPtr> Selected;
+	SourcesListView->GetSelectedItems(Selected);
+	for (FLiveLinkSourceUIEntryPtr Item : Selected)
+	{
+		Client->OnPropertyChanged(Item->GetGuid(), InEvent);
+	}
+}
+
 FLiveLinkSubjectsView::FLiveLinkSubjectsView(FOnSubjectSelectionChanged InOnSubjectSelectionChanged, const TSharedPtr<FUICommandList>& InCommandList, TAttribute<bool> bInReadOnly)
 	: SubjectSelectionChangedDelegate(InOnSubjectSelectionChanged)
 	, bReadOnly(MoveTemp(bInReadOnly))
+	, bShowDevices(false)
 {
-	CreateSubjectsTreeView(InCommandList);
+	CreateCombinedTreeView(InCommandList);
+}
+
+FLiveLinkSubjectsView::FLiveLinkSubjectsView(FOnSubjectSelectionChanged InOnSubjectSelectionChanged, const TSharedPtr<FUICommandList>& InCommandList, const FSubjectsViewArgs& InArgs)
+	: SubjectSelectionChangedDelegate(InOnSubjectSelectionChanged)
+	, bReadOnly(InArgs.bInReadOnly)
+	, bShowDevices(InArgs.bInShowDevices)
+{
+	CreateCombinedTreeView(InCommandList);
 }
 
 void FLiveLinkSubjectsView::OnSubjectSelectionChanged(FLiveLinkSubjectUIEntryPtr SubjectEntry, ESelectInfo::Type SelectInfo)
@@ -974,14 +1356,46 @@ void FLiveLinkSubjectsView::OnSubjectSelectionChanged(FLiveLinkSubjectUIEntryPtr
 	SubjectSelectionChangedDelegate.Execute(SubjectEntry, SelectInfo);
 }
 
-void FLiveLinkSourcesView::OnPropertyChanged(const FPropertyChangedEvent& InEvent)
+TSharedPtr<SWidget> FLiveLinkSubjectsView::OnCombinedListOpenContextMenu(TSharedPtr<FUICommandList> InCommandList)
 {
-	TArray<FLiveLinkSourceUIEntryPtr> Selected;
-	SourcesListView->GetSelectedItems(Selected);
-	for (FLiveLinkSourceUIEntryPtr Item : Selected)
+	if (bReadOnly.Get())
 	{
-		Client->OnPropertyChanged(Item->GetGuid(), InEvent);
+		return nullptr;
+	}
+
+	constexpr bool bShouldCloseWindowAfterMenuSelection = true;
+	FMenuBuilder MenuBuilder(bShouldCloseWindowAfterMenuSelection, InCommandList);
+
+	MenuBuilder.BeginSection(TEXT("Remove"));
+	{
+		if (CanRemoveSubject())
+		{
+			MenuBuilder.AddMenuEntry(FLiveLinkClientCommands::Get().RemoveSubject);
+		}
+
+		if (CanRemoveDevice())
+		{
+			MenuBuilder.AddMenuEntry(FLiveLinkClientCommands::Get().RemoveDevice);
+		}
+
+		if (CanRemoveSource())
+		{
+			MenuBuilder.AddMenuEntry(FLiveLinkClientCommands::Get().RemoveSource);
+		}
+
+		if (HasAnySources())
+		{
+			MenuBuilder.AddMenuEntry(FLiveLinkClientCommands::Get().RemoveAllSources);
+		}
+	}
+	MenuBuilder.EndSection();
+
+	if (CanPauseSubject())
+	{
+		MenuBuilder.AddMenuEntry(FLiveLinkClientCommands::Get().PauseSubject, NAME_None, TAttribute<FText>::CreateSP(this, &FLiveLinkSubjectsView::GetPauseSubjectLabel), TAttribute<FText>::CreateSP(this, &FLiveLinkSubjectsView::GetPauseSubjectToolTip));
 	}
+
+	return MenuBuilder.MakeWidget();
 }
 
 TSharedRef<ITableRow> FLiveLinkSubjectsView::MakeTreeRowWidget(FLiveLinkSubjectUIEntryPtr InInfo, const TSharedRef<STableViewBase>& OwnerTable)
@@ -1036,6 +1450,53 @@ bool FLiveLinkSubjectsView::CanRemoveSubject() const
 	return false;
 }
 
+bool FLiveLinkSubjectsView::CanRemoveDevice() const
+{
+	TArray<FLiveLinkSubjectUIEntryPtr> Selected;
+	SubjectsTreeView->GetSelectedItems(Selected);
+
+	for (const FLiveLinkSubjectUIEntryPtr& Entry : Selected)
+	{
+		if (Entry && Entry->IsDevice())
+		{
+			return true;
+		}
+	}
+
+	return false;
+}
+
+bool FLiveLinkSubjectsView::CanRemoveSource() const
+{
+	TArray<FLiveLinkSubjectUIEntryPtr> Selected;
+	SubjectsTreeView->GetSelectedItems(Selected);
+
+	for (const FLiveLinkSubjectUIEntryPtr& Entry : Selected)
+	{
+		// Allow removing souce when right clicking a subject.
+		if (Entry && (Entry->IsSubject() || Entry->IsSource()))
+		{
+			return true;
+		}
+	}
+
+	return false;
+}
+
+
+bool FLiveLinkSubjectsView::HasAnySources() const
+{
+	for (const FLiveLinkSubjectUIEntryPtr& Entry : FilteredList)
+	{
+		if (Entry && Entry->IsSource())
+		{
+			return true;
+		}
+	}
+
+	return false;
+}
+
 void FLiveLinkSubjectsView::RefreshSubjects()
 {
 	TArray<TPair<FLiveLinkSubjectKey, bool>> SavedSelection;
@@ -1051,45 +1512,109 @@ void FLiveLinkSubjectsView::RefreshSubjects()
 	{
 		if (ILiveLinkClient* Client = &IModularFeatures::Get().GetModularFeature<ILiveLinkClient>(ILiveLinkClient::ModularFeatureName))
 		{
+			FLiveLinkSubjectUIEntryPtr SourcesHeader;
+			FLiveLinkSubjectUIEntryPtr DevicesHeader;
+
 			TArray<FLiveLinkSubjectKey> SubjectKeys = Client->GetSubjects(true, true);
-			SubjectData.Reset();
+			TArray<FGuid> Sources = Client->GetSources();
 
-			TMap<FName, FLiveLinkSubjectUIEntryPtr> SourceItems;
+			TMap<FGuid, FLiveLinkSubjectUIEntryPtr> SourceItems;
 
 			TArray<FLiveLinkSubjectUIEntryPtr> AllItems;
 			AllItems.Reserve(SubjectKeys.Num());
 
+			SubjectData.Reset();
+			
+			// Add Subjects & Sources
 			for (const FLiveLinkSubjectKey& SubjectKey : SubjectKeys)
 			{
-				FLiveLinkSubjectUIEntryPtr Source;
-
 				FName SourceNameOverride = *Client->GetSourceNameOverride(SubjectKey).ToString();
 
-				if (FLiveLinkSubjectUIEntryPtr* SourcePtr = SourceItems.Find(*SourceNameOverride.ToString()))
-				{
-					Source = *SourcePtr;
-				}
-				else
+				FLiveLinkSubjectUIEntryPtr Source = SourceItems.FindRef(SubjectKey.Source);
+
+				if (!Source)
 				{
 					constexpr bool bIsSource = true;
 					Source = MakeShared<FLiveLinkSubjectUIEntry>(SubjectKey, static_cast<FLiveLinkClient*>(Client), bIsSource);
-					SubjectData.Add(Source);
-					SourceItems.Add(SourceNameOverride) = Source;
+
+					if (!SourcesHeader)
+					{
+						// It's possible that no source were created such as in Playback.
+						SourcesHeader = MakeShared<FLiveLinkSubjectUIEntry>(static_cast<FLiveLinkClient*>(Client), ELiveLinkEntryType::SourcesHeader);
+						SubjectData.Add(SourcesHeader);
+					}
+
+					SourcesHeader->Children.Add(Source);
+					SourceItems.Add(SubjectKey.Source) = Source;
 
 					SubjectsTreeView->SetItemExpansion(Source, true);
 					AllItems.Add(Source);
 				}
-				
+
 				FLiveLinkSubjectUIEntryPtr SubjectEntry = MakeShared<FLiveLinkSubjectUIEntry>(SubjectKey, static_cast<FLiveLinkClient*>(Client));
 				Source->Children.Add(SubjectEntry);
 				AllItems.Add(SubjectEntry);
 			}
 
+			// Add sources that haven't been added through subjects.
+			// @Note: This pass must come after subjects were added, because subjects can override their source name, so we add them 
+			// in the previous pass to avoid having to rename them.
+			for (const FGuid& SourceId : Sources)
+			{
+				if (SourceItems.Contains(SourceId))
+				{
+					continue;
+				}
+
+				constexpr bool bIsSource = true;
+				FLiveLinkSubjectUIEntryPtr Source = MakeShared<FLiveLinkSubjectUIEntry>(FLiveLinkSubjectKey{ SourceId, NAME_None }, static_cast<FLiveLinkClient*>(Client), bIsSource);
+
+				if (!SourcesHeader)
+				{
+					SourcesHeader = MakeShared<FLiveLinkSubjectUIEntry>(static_cast<FLiveLinkClient*>(Client), ELiveLinkEntryType::SourcesHeader);
+					SubjectData.Add(SourcesHeader);
+				}
+
+				SourcesHeader->Children.Add(Source);
+
+				SubjectsTreeView->SetItemExpansion(Source, true);
+				FText SourceType = Client->GetSourceType(SourceId);
+
+				SourceItems.Add(SourceId, Source);
+				AllItems.Add(Source);
+			}
+
+			// Add Devices (Only visible in LiveLinkHub)
+			if (bShowDevices)
+			{
+				if (ULiveLinkDeviceSubsystem* Subsystem = GEngine->GetEngineSubsystem<ULiveLinkDeviceSubsystem>())
+				{
+					const TMap<FGuid, TObjectPtr<ULiveLinkDevice>>& DeviceMap = Subsystem->GetDeviceMap();
+					for (const TPair<FGuid, TObjectPtr<ULiveLinkDevice>>& DevicePair : DeviceMap)
+					{
+						if (!DevicesHeader)
+						{
+							DevicesHeader = MakeShared<FLiveLinkSubjectUIEntry>(static_cast<FLiveLinkClient*>(Client), ELiveLinkEntryType::DevicesHeader);
+							SubjectData.Add(DevicesHeader);
+						}
+
+						FLiveLinkSubjectUIEntryPtr Entry = MakeShared<FLiveLinkSubjectUIEntry>(static_cast<FLiveLinkClient*>(Client), ELiveLinkEntryType::Device);
+						Entry->WeakDevice = DevicePair.Value;
+
+						DevicesHeader->Children.Add(Entry);
+					}
+				}
+			}
+
 			auto SortPredicate = [](const FLiveLinkSubjectUIEntryPtr& LHS, const FLiveLinkSubjectUIEntryPtr& RHS) {return LHS->GetItemText().CompareTo(RHS->GetItemText()) < 0; };
 			SubjectData.Sort(SortPredicate);
 			for (FLiveLinkSubjectUIEntryPtr& Subject : SubjectData)
 			{
 				Subject->Children.Sort(SortPredicate);
+				for (const FLiveLinkSubjectUIEntryPtr& Children : Subject->Children)
+				{
+					Children->Children.Sort(SortPredicate);
+				}
 			}
 
 			for (const FLiveLinkSubjectUIEntryPtr& Item : AllItems)
@@ -1151,34 +1676,44 @@ TSharedRef<SWidget> FLiveLinkSubjectsView::GetWidget()
 		.AutoHeight()
 		[
 			SNew(SHorizontalBox)
-				+ SHorizontalBox::Slot()
-				.Padding(FMargin(4.0, 2.0))
-				.FillWidth(1.0f)
-				[
-					SAssignNew(FilterSearchBox, SLiveLinkFilterSearchBox<FLiveLinkSubjectUIEntryPtr>)
-						.ItemSource(&SubjectData)
-						.OnGatherItems_Lambda([this](TArray<FLiveLinkSubjectUIEntryPtr>& Items)
+			+ SHorizontalBox::Slot()
+			.Padding(FMargin(4.0, 2.0))
+			.AutoWidth()
+			[
+				SNew(SPositiveActionButton)
+					.OnGetMenuContent(this, &FLiveLinkSubjectsView::OnGenerateAddMenuEntries)
+					.Icon(FAppStyle::Get().GetBrush("Icons.Plus"))
+					.Text(LOCTEXT("AddButtonLabel", "Add"))
+					.ToolTipText(LOCTEXT("AddButton_ToolTip", "Add a new Live Link data device."))
+			]
+			+ SHorizontalBox::Slot()
+			.Padding(FMargin(4.0, 2.0))
+			.FillWidth(1.0f)
+			[
+				SAssignNew(FilterSearchBox, SLiveLinkFilterSearchBox<FLiveLinkSubjectUIEntryPtr>)
+					.ItemSource(&SubjectData)
+					.OnGatherItems_Lambda([this](TArray<FLiveLinkSubjectUIEntryPtr>& Items)
+						{
+							for (FLiveLinkSubjectUIEntryPtr Source : SubjectData)
 							{
-								for (FLiveLinkSubjectUIEntryPtr Source : SubjectData)
+								for (FLiveLinkSubjectUIEntryPtr Subject : Source->Children)
 								{
-									for (FLiveLinkSubjectUIEntryPtr Subject : Source->Children)
-									{
-										Items.Add(Subject);
-									}
+									Items.Add(Subject);
 								}
-							})
-						.OnUpdateFilteredList_Lambda([this](const TArray<FLiveLinkSubjectUIEntryPtr>& FilteredItems) 
-							{
-								FilteredList = FilteredItems;
+							}
+						})
+					.OnUpdateFilteredList_Lambda([this](const TArray<FLiveLinkSubjectUIEntryPtr>& FilteredItems)
+						{
+							FilteredList = FilteredItems;
 
-								for (const FLiveLinkSubjectUIEntryPtr& Item : FilteredList)
-								{
-									SubjectsTreeView->SetItemExpansion(Item, true);
-								}
+							for (const FLiveLinkSubjectUIEntryPtr& Item : FilteredList)
+							{
+								SubjectsTreeView->SetItemExpansion(Item, true);
+							}
 
-								SubjectsTreeView->RequestListRefresh();
-							})
-				]
+							SubjectsTreeView->RequestListRefresh();
+						})
+			]
 		]
 		+ SVerticalBox::Slot()
 		.VAlign(VAlign_Fill)
@@ -1187,38 +1722,44 @@ TSharedRef<SWidget> FLiveLinkSubjectsView::GetWidget()
 		];
 }
 
-void FLiveLinkSubjectsView::CreateSubjectsTreeView(const TSharedPtr<FUICommandList>& InCommandList)
+void FLiveLinkSubjectsView::CreateCombinedTreeView(const TSharedPtr<FUICommandList>& InCommandList)
 {
 	SAssignNew(SubjectsTreeView, SLiveLinkSubjectsTreeView, bReadOnly)
 		.TreeItemsSource(&FilteredList)
 		.OnGenerateRow_Raw(this, &FLiveLinkSubjectsView::MakeTreeRowWidget)
 		.OnGetChildren_Raw(this, &FLiveLinkSubjectsView::GetChildrenForInfo)
 		.OnSelectionChanged_Raw(this, &FLiveLinkSubjectsView::OnSubjectSelectionChanged)
-		.OnContextMenuOpening_Raw(this, &FLiveLinkSubjectsView::OnOpenVirtualSubjectContextMenu, InCommandList)
+		.OnContextMenuOpening_Raw(this, &FLiveLinkSubjectsView::OnCombinedListOpenContextMenu, InCommandList)
 		.SelectionMode(ESelectionMode::Multi)
 		.HeaderRow
 		(
 			SNew(SHeaderRow)
 			.CanSelectGeneratedColumn(true)
-			.HiddenColumnsList({ SubjectTreeUI::TranslatedRoleColumnName })
-			+ SHeaderRow::Column(SubjectTreeUI::EnabledColumnName)
-			.DefaultTooltip(LOCTEXT("EnabledToolTip", "Whether this subject should be evaluated."))
+			.HiddenColumnsList({ UE::LiveLink::DefaultColumn::TranslatedRole })
+			+ SHeaderRow::Column(UE::LiveLink::DefaultColumn::Action)
+			.DefaultTooltip(LOCTEXT("ActionWidgetToolTip", "Action that can be applied to the device or subject."))
 			.DefaultLabel(FText())
-			.FixedWidth(22)
-			+ SHeaderRow::Column(SubjectTreeUI::NameColumnName)
-			.DefaultLabel(LOCTEXT("SubjectItemName", "Subject Name"))
-			.DefaultTooltip(LOCTEXT("NameToolTip", "Unique name for this subject."))
+			.FixedWidth(28)
+			+ SHeaderRow::Column(UE::LiveLink::DefaultColumn::Name)
+			.DefaultLabel(LOCTEXT("ItemName", "Name"))
+			.DefaultTooltip(LOCTEXT("NameToolTip", "Name of this Subject/Source/Device."))
 			.FillWidth(0.60f)
-			+ SHeaderRow::Column(SubjectTreeUI::RoleColumnName)
-			.DefaultLabel(LOCTEXT("RoleName", "Role"))
+			+ SHeaderRow::Column(UE::LiveLink::DefaultColumn::Machine)
+			.DefaultLabel(LOCTEXT("MachineName", "Machine"))
+			.DefaultTooltip(LOCTEXT("NameToolTip", "Name of the machine that hosts this Subject/Source/Device."))
+			.HAlignCell(HAlign_Center)
+			.FillWidth(0.30f)
+			+ SHeaderRow::Column(UE::LiveLink::DefaultColumn::Role)
+			.DefaultLabel(LOCTEXT("StatusRoleName", "Status / Role"))
 			.ManualWidth(90.f)
 			.DefaultTooltip(LOCTEXT("RoleToolTip", "Role of this subject."))
-			+ SHeaderRow::Column(SubjectTreeUI::TranslatedRoleColumnName)
+			+ SHeaderRow::Column(UE::LiveLink::DefaultColumn::TranslatedRole)
 			.DefaultLabel(LOCTEXT("TranslatedRoleName", "Translated"))
 			.DefaultTooltip(LOCTEXT("TranslatedRoleToolTip", "Get the translated role of a subject if it's being translated before being rebroadcast. This should only be relevant in Live Link Hub."))
 			.ManualWidth(70.f)
-			+ SHeaderRow::Column(SubjectTreeUI::ActionsColumnName)
+			+ SHeaderRow::Column(UE::LiveLink::DefaultColumn::Status)
 			.ManualWidth(20.f)
+			.HAlignCell(HAlign_Center)
 			.DefaultLabel(FText())
 		);
 }
@@ -1263,145 +1804,33 @@ bool FLiveLinkSubjectsView::IsSelectedSubjectPaused() const
 	return true;
 }
 
-TSharedRef<SWidget> FLiveLinkSourcesView::OnGenerateSourceMenu()
+TSharedRef<SWidget> FLiveLinkSubjectsView::OnGenerateAddMenuEntries()
 {
-	const bool CloseAfterSelection = true;
-	FMenuBuilder MenuBuilder(CloseAfterSelection, NULL);
+	bool bShouldCloseWindowAfterMenuSelection = true;
+	FMenuBuilder MenuBuilder(bShouldCloseWindowAfterMenuSelection, nullptr);
 
-	MenuBuilder.BeginSection("SourceSection", LOCTEXT("Sources", "Live Link Sources"));
+	ILiveLinkDeviceModule& DeviceModule = ILiveLinkDeviceModule::Get();
 
-	for (int32 FactoryIndex = 0; FactoryIndex < Factories.Num(); ++FactoryIndex)
+	if (bShowDevices)
 	{
-		ULiveLinkSourceFactory* FactoryInstance = Factories[FactoryIndex];
-		if (FactoryInstance)
-		{
-			ULiveLinkSourceFactory::EMenuType MenuType = FactoryInstance->GetMenuType();
-
-			if (MenuType == ULiveLinkSourceFactory::EMenuType::SubPanel)
-			{
-				MenuBuilder.AddMenuEntry(
-					FactoryInstance->GetSourceDisplayName(),
-					FactoryInstance->GetSourceTooltip(),
-					FSlateIcon(),
-					FUIAction(
-						FExecuteAction::CreateSP(this, &FLiveLinkSourcesView::OpenCreateMenuWindow, FactoryIndex)
-					),
-					NAME_None,
-					EUserInterfaceActionType::Button);
-			}
-			else if (MenuType == ULiveLinkSourceFactory::EMenuType::MenuEntry)
-			{
-				MenuBuilder.AddMenuEntry(
-					FactoryInstance->GetSourceDisplayName(),
-					FactoryInstance->GetSourceTooltip(),
-					FSlateIcon(),
-					FUIAction(
-						FExecuteAction::CreateSP(this, &FLiveLinkSourcesView::ExecuteCreateSource, FactoryIndex)
-					),
-					NAME_None,
-					EUserInterfaceActionType::Button);
-			}
-			else
-			{
-				MenuBuilder.AddMenuEntry(
-					FactoryInstance->GetSourceDisplayName(),
-					FactoryInstance->GetSourceTooltip(),
-					FSlateIcon(),
-					FUIAction(
-						FExecuteAction(),
-						FCanExecuteAction::CreateLambda([]() { return false; })
-					),
-					NAME_None,
-					EUserInterfaceActionType::Button);
-			}
-		}
+		DeviceModule.CreateDeviceMenuEntries(MenuBuilder);
 	}
 
-	MenuBuilder.EndSection();
-
-	MenuBuilder.BeginSection("VirtualSourceSection", LOCTEXT("VirtualSources", "Live Link VirtualSubject Sources"));
-
-	MenuBuilder.AddMenuEntry(
-		LOCTEXT("AddVirtualSubject", "Add Virtual Subject"),
-		LOCTEXT("AddVirtualSubject_Tooltip", "Adds a new virtual subject to Live Link. Instead of coming from a source a virtual subject is a combination of 2 or more real subjects"),
-		FSlateIcon(),
-		FUIAction(
-			FExecuteAction::CreateSP(this, &FLiveLinkSourcesView::AddVirtualSubject)
-		),
-		NAME_None,
-		EUserInterfaceActionType::Button);
-
-	MenuBuilder.EndSection();
+	TArray<TObjectPtr<ULiveLinkSourceFactory>> Factories = UE::LiveLink::GetAllSourceFactories();
+	UE::LiveLink::GenerateSourceMenuEntries(Factories, MenuBuilder);
 
 	return MenuBuilder.MakeWidget();
 }
 
-void FLiveLinkSourcesView::OpenCreateMenuWindow(int32 FactoryIndex)
-{
-	if (Factories.IsValidIndex(FactoryIndex))
-	{
-		if (ULiveLinkSourceFactory* FactoryInstance = Factories[FactoryIndex])
-		{
-			TSharedPtr<SWidget> Widget = FactoryInstance->BuildCreationPanel(ULiveLinkSourceFactory::FOnLiveLinkSourceCreated::CreateSP(this, &FLiveLinkSourcesView::OnSourceCreated, TSubclassOf<ULiveLinkSourceFactory>(FactoryInstance->GetClass())));
-			if (Widget.IsValid())
-			{
-				FText CreateText = FText::Format(LOCTEXT("CreateSourceLabel", "Create {0} connection"), FactoryInstance->GetSourceDisplayName());
-				SourceCreationWindow = SNew(SWindow)
-					.Title(CreateText)
-					.ClientSize(FVector2D(400, 150))
-					.SizingRule(ESizingRule::Autosized)
-					.SupportsMinimize(false)
-					.SupportsMaximize(false)
-					[
-						Widget.ToSharedRef()
-					];
-
-				GEditor->EditorAddModalWindow(SourceCreationWindow.ToSharedRef());
-			}
-		}
-	}
-}
-
-void FLiveLinkSourcesView::ExecuteCreateSource(int32 FactoryIndex)
-{
-	if (Factories.IsValidIndex(FactoryIndex))
-	{
-		if (ULiveLinkSourceFactory* FactoryInstance = Factories[FactoryIndex])
-		{
-			OnSourceCreated(FactoryInstance->CreateSource(FString()), FString(), FactoryInstance->GetClass());
-		}
-	}
-}
 
-void FLiveLinkSourcesView::OnSourceCreated(TSharedPtr<ILiveLinkSource> NewSource, FString ConnectionString, TSubclassOf<ULiveLinkSourceFactory> Factory)
+TSharedRef<SWidget> FLiveLinkSourcesView::OnGenerateSourceMenu()
 {
-	if (NewSource.IsValid())
-	{
-		FGuid NewSourceGuid = Client->AddSource(NewSource);
-		if (NewSourceGuid.IsValid())
-		{
-			if (ULiveLinkSourceSettings* Settings = Client->GetSourceSettings(NewSourceGuid))
-			{
-				Settings->ConnectionString = ConnectionString;
-				Settings->Factory = Factory;
-			}
-		}
-	}
-
-	if (SourceCreationWindow)
-	{
-		SourceCreationWindow->RequestDestroyWindow();
-		SourceCreationWindow.Reset();
-	}
-}
+	bool bShouldCloseWindowAfterMenuSelection = true;
+	FMenuBuilder MenuBuilder(bShouldCloseWindowAfterMenuSelection, nullptr);
 
-void FLiveLinkSourcesView::AddVirtualSubject()
-{
-	TSharedRef<SVirtualSubjectCreateDialog> Dialog =
-		SNew(SVirtualSubjectCreateDialog)
-		.LiveLinkClient(Client);
+	UE::LiveLink::GenerateSourceMenuEntries(Factories, MenuBuilder);
 
-	Dialog->ConfigureVirtualSubject();
+	return MenuBuilder.MakeWidget();
 }
 
 
diff --git a/Engine/Plugins/Animation/LiveLink/Source/LiveLinkEditor/Private/LiveLinkPanelController.cpp b/Engine/Plugins/Animation/LiveLink/Source/LiveLinkEditor/Private/LiveLinkPanelController.cpp
--- a/Engine/Plugins/Animation/LiveLink/Source/LiveLinkEditor/Private/LiveLinkPanelController.cpp
+++ b/Engine/Plugins/Animation/LiveLink/Source/LiveLinkEditor/Private/LiveLinkPanelController.cpp
@@ -3,19 +3,21 @@
 #include "LiveLinkPanelController.h"
 #include "Async/TaskGraphInterfaces.h"
 #include "Framework/Commands/UICommandList.h"
+#include "Engine/Engine.h"
 #include "Features/IModularFeatures.h"
 #include "IDetailsView.h"
 #include "LiveLinkClient.h"
 #include "LiveLinkClientCommands.h"
 #include "LiveLinkClientPanelViews.h"
+#include "LiveLinkDevice.h"
+#include "LiveLinkDeviceSubsystem.h"
 #include "Misc/ConfigCacheIni.h"
 #include "SLiveLinkDataView.h"
+#include "Widgets/Layout/SWidgetSwitcher.h"
 
 
 FLiveLinkPanelController::FLiveLinkPanelController(TAttribute<bool> bInReadOnly)
 {
-	bSeparateSourcesSubjects = GConfig->GetBoolOrDefault(TEXT("LiveLink"), TEXT("bPanelControllerSeparateSourcesSubjects"), false, GEngineIni);
-
 	Client = (FLiveLinkClient*)&IModularFeatures::Get().GetModularFeature<ILiveLinkClient>(ILiveLinkClient::ModularFeatureName);
 
 	OnSourcesChangedHandle = Client->OnLiveLinkSourcesChanged().AddRaw(this, &FLiveLinkPanelController::OnSourcesChangedHandler);
@@ -25,17 +27,55 @@ FLiveLinkPanelController::FLiveLinkPanelController(TAttribute<bool> bInReadOnly)
 	CommandList = MakeShared<FUICommandList>();
 	BindCommands();
 
-	SourcesView = MakeShared<FLiveLinkSourcesView>(Client, CommandList, bInReadOnly, FLiveLinkSourcesView::FOnSourceSelectionChanged::CreateRaw(this, &FLiveLinkPanelController::OnSourceSelectionChangedHandler));
-	SubjectsView = MakeShared<FLiveLinkSubjectsView>(FLiveLinkSubjectsView::FOnSubjectSelectionChanged::CreateRaw(this, &FLiveLinkPanelController::OnSubjectSelectionChangedHandler), CommandList, bInReadOnly);
-	SourcesDetailsView = UE::LiveLink::CreateSourcesDetailsView(SourcesView, bInReadOnly);
+	FLiveLinkSubjectsView::FSubjectsViewArgs Args;
+	Args.bInReadOnly = bInReadOnly;
+	Args.bInShowDevices = GConfig->GetBoolOrDefault(TEXT("LiveLink"), TEXT("bCreateLiveLinkHubInstance"), false, GEngineIni);
+
+	SubjectsView = MakeShared<FLiveLinkSubjectsView>(FLiveLinkSubjectsView::FOnSubjectSelectionChanged::CreateRaw(this, &FLiveLinkPanelController::OnSubjectSelectionChangedHandler), CommandList, Args);
+
+	if (Args.bInShowDevices)
+	{
+		ULiveLinkDeviceSubsystem* Subsystem = GEngine->GetEngineSubsystem<ULiveLinkDeviceSubsystem>();
+
+		TWeakPtr<FLiveLinkSubjectsView> WeakSubjectsView = SubjectsView;
+
+		auto OnDevicesChangedHandler = [WeakSubjectsView](FGuid InDeviceId, ULiveLinkDevice* InDevice)
+		{
+			if (TSharedPtr<FLiveLinkSubjectsView> View = WeakSubjectsView.Pin())
+			{
+				View->RefreshSubjects();
+			}
+		};
+
+		DeviceAddedHandle = Subsystem->OnDeviceAdded().AddLambda(OnDevicesChangedHandler);
+		DeviceRemovedHandle = Subsystem->OnDeviceRemoved().AddLambda(OnDevicesChangedHandler);
+	}
+
+	FOnFinishedChangingProperties SourcesDelegate;
+	SourcesDelegate.AddRaw(this, &FLiveLinkPanelController::OnFinishedChangingSourceProperties);
+
+	SourcesDetailsView = UE::LiveLink::CreateSourcesDetailsView(SourcesDelegate, bInReadOnly);
 	SubjectsDetailsView = UE::LiveLink::CreateSubjectsDetailsView(Client, bInReadOnly);
+	FOnFinishedChangingProperties Delegate;
+	Delegate.AddRaw(this, &FLiveLinkPanelController::OnFinishedChangingDeviceProperties);
+
+	DevicesDetailsView = UE::LiveLink::CreateDevicesDetailsView(MoveTemp(Delegate), bInReadOnly);
 
-	RebuildSourceList();
 	RebuildSubjectList();
 }
 
+PRAGMA_DISABLE_DEPRECATION_WARNINGS // Needed until SourcesView is removed in 5.10
 FLiveLinkPanelController::~FLiveLinkPanelController()
 {
+	if (GEngine)
+	{
+		if (ULiveLinkDeviceSubsystem* Subsystem = GEngine->GetEngineSubsystem<ULiveLinkDeviceSubsystem>())
+		{
+			Subsystem->OnDeviceAdded().Remove(DeviceAddedHandle);
+			Subsystem->OnDeviceRemoved().Remove(DeviceRemovedHandle);
+		}
+	}
+
 	if (IModularFeatures::Get().IsModularFeatureAvailable(ILiveLinkClient::ModularFeatureName))
 	{
 		ILiveLinkClient* LiveLinkClient = (FLiveLinkClient*)&IModularFeatures::Get().GetModularFeature<ILiveLinkClient>(ILiveLinkClient::ModularFeatureName);
@@ -49,41 +89,7 @@ FLiveLinkPanelController::~FLiveLinkPanelController()
 		}
 	}
 }
-
-void FLiveLinkPanelController::OnSourceSelectionChangedHandler(FLiveLinkSourceUIEntryPtr Entry, ESelectInfo::Type SelectionType) const
-{
-	if (bSelectionChangedGuard)
-	{
-		return;
-	}
-	
-	TGuardValue<bool> ReentrantGuard(bSelectionChangedGuard, true);
-	
-	SubjectsDetailsView->SetSubjectKey(FLiveLinkSubjectKey());
-
-	int32 FoundSubjectEntryIndex = INDEX_NONE;
-	if (Entry.IsValid())
-	{
-		SourcesDetailsView->SetObject(Entry->GetSourceSettings());
-
-		// Find the corresponding subject
-		FoundSubjectEntryIndex = SubjectsView->SubjectData.IndexOfByPredicate([Entry](const FLiveLinkSubjectUIEntryPtr& SubjectEntry) { return SubjectEntry->SubjectKey.Source == Entry->GetGuid() && SubjectEntry->IsSource(); });
-	}
-	else
-	{
-		SourcesDetailsView->SetObject(nullptr);
-	}
-
-	// Set the corresponding subject
-	if (FoundSubjectEntryIndex != INDEX_NONE)
-	{
-		SubjectsView->SubjectsTreeView->SetSelection(SubjectsView->SubjectData[FoundSubjectEntryIndex]);
-	}
-	else
-	{
-		SubjectsView->SubjectsTreeView->ClearSelection();
-	}
-}
+PRAGMA_ENABLE_DEPRECATION_WARNINGS
 
 void FLiveLinkPanelController::OnSubjectSelectionChangedHandler(FLiveLinkSubjectUIEntryPtr SubjectEntry, ESelectInfo::Type SelectInfo)
 {
@@ -93,59 +99,103 @@ void FLiveLinkPanelController::OnSubjectSelectionChangedHandler(FLiveLinkSubject
 	}
 	TGuardValue<bool> ReentrantGuard(bSelectionChangedGuard, true);
 
-	int32 FoundSourceIndex = INDEX_NONE;
-	bool bDetailViewSet = false;
+	SourcesDetailsView->SetObject(nullptr);
+	DevicesDetailsView->SetObject(nullptr);
+	SubjectsDetailsView->SetSubjectKey(FLiveLinkSubjectKey());
+
 	if (SubjectEntry.IsValid())
 	{
 		// Find the corresponding Source entry
 		FGuid SourceGuid = SubjectEntry->SubjectKey.Source;
-		FoundSourceIndex = SourcesView->SourceData.IndexOfByPredicate([SourceGuid](FLiveLinkSourceUIEntryPtr SourceEntry) { return SourceEntry->GetGuid() == SourceGuid; });
 
 		if (SubjectEntry->IsSource())
 		{
 			SourcesDetailsView->SetObject(SubjectEntry->GetSettings());
-			if (!bSeparateSourcesSubjects)
-			{
-				SubjectsDetailsView->SetSubjectKey(FLiveLinkSubjectKey());
-			}
 		}
-		else
+		else if (SubjectEntry->IsSubject())
 		{
-			if (!bSeparateSourcesSubjects)
+			SubjectsDetailsView->SetSubjectKey(SubjectEntry->SubjectKey);
+		}
+		else if (SubjectEntry->IsDevice())
+		{
+			if (TStrongObjectPtr<ULiveLinkDevice> Device = SubjectEntry->WeakDevice.Pin())
 			{
-				SourcesDetailsView->SetObject(nullptr);
+				DevicesDetailsView->SetObject(Device->GetDeviceSettings());
 			}
-			SubjectsDetailsView->SetSubjectKey(SubjectEntry->SubjectKey);
 		}
-		bDetailViewSet = true;
 
 		SubjectSelectionChangedDelegate.Broadcast(SubjectEntry->SubjectKey);
 	}
+}
 
-	if (!bDetailViewSet)
+void FLiveLinkPanelController::OnFinishedChangingDeviceProperties(const FPropertyChangedEvent& InPropertyChangedEvent)
+{
+	if (DevicesDetailsView)
 	{
-		if (!bSeparateSourcesSubjects)
+		ULiveLinkDevice* Device = nullptr;
+
+		const TArray<TWeakObjectPtr<UObject>>& SelectedObjects = DevicesDetailsView->GetSelectedObjects();
+		if (ensure(SelectedObjects.Num() > 0))
+		{
+			if (UObject* SelectedObject = SelectedObjects[0].Get(); ensure(SelectedObject))
+			{
+				Device = SelectedObject->GetTypedOuter<ULiveLinkDevice>();
+			}
+		}
+
+		if (ensure(Device))
 		{
-			SourcesDetailsView->SetObject(nullptr);
+			Device->OnSettingChanged(InPropertyChangedEvent);
 		}
-		SubjectsDetailsView->SetSubjectKey(FLiveLinkSubjectKey());
 	}
+}
 
-	// Select the corresponding Source entry
-	if (FoundSourceIndex != INDEX_NONE)
+void FLiveLinkPanelController::OnFinishedChangingSourceProperties(const FPropertyChangedEvent& InPropertyChangedEvent)
+{
+	TArray<FLiveLinkSubjectUIEntryPtr> SelectedItems = SubjectsView->SubjectsTreeView->GetSelectedItems();
+	for (const FLiveLinkSubjectUIEntryPtr& Item : SelectedItems)
 	{
-		SourcesView->SourcesListView->SetSelection(SourcesView->SourceData[FoundSourceIndex]);
-
-		if (bSeparateSourcesSubjects)
+		if (Item->IsSource())
 		{
-			// Update source details to the selected subject's source.
-			SourcesDetailsView->SetObject(SourcesView->SourceData[FoundSourceIndex]->GetSourceSettings());
+			Client->OnPropertyChanged(Item->SubjectKey.Source, InPropertyChangedEvent);
 		}
 	}
-	else
+}
+
+int32 FLiveLinkPanelController::GetCombinedDetailWidgetIndex() const
+{
+	if (SourcesDetailsView->GetSelectedObjects().Num() > 0)
 	{
-		SourcesView->SourcesListView->ClearSelection();
+		return 0;
 	}
+	
+	if (!SubjectsDetailsView->GetSubjectKey().SubjectName.IsNone())
+	{
+		return 1;
+	}
+
+	return 2;
+}
+
+TSharedRef<SWidget> FLiveLinkPanelController::GetCombinedDetailsWidget()
+{
+	return SNew(SWidgetSwitcher)
+		.WidgetIndex(this, &FLiveLinkPanelController::GetCombinedDetailWidgetIndex)
+		+ SWidgetSwitcher::Slot()
+		[
+			//[0] Detail view for Source
+			SourcesDetailsView.ToSharedRef()
+		]
+		+ SWidgetSwitcher::Slot()
+		[
+			// [1] Detail view for Subject, Frame data & Static data
+			SubjectsDetailsView.ToSharedRef()
+		]
+		+ SWidgetSwitcher::Slot()
+		[
+			// [2] Detail view for Devices
+			DevicesDetailsView.ToSharedRef()
+		];
 }
 
 void FLiveLinkPanelController::BindCommands()
@@ -165,6 +215,11 @@ void FLiveLinkPanelController::BindCommands()
 		FCanExecuteAction::CreateRaw(this, &FLiveLinkPanelController::CanRemoveSubject)
 	);
 
+	CommandList->MapAction(FLiveLinkClientCommands::Get().RemoveDevice,
+		FExecuteAction::CreateRaw(this, &FLiveLinkPanelController::HandleRemoveDevice),
+		FCanExecuteAction::CreateRaw(this, &FLiveLinkPanelController::CanRemoveDevice)
+	);
+
 	CommandList->MapAction(FLiveLinkClientCommands::Get().PauseSubject,
 		FExecuteAction::CreateRaw(this, &FLiveLinkPanelController::HandlePauseSubject),
 		FCanExecuteAction::CreateRaw(this, &FLiveLinkPanelController::CanPauseSubject)
@@ -176,7 +231,6 @@ void FLiveLinkPanelController::OnSourcesChangedHandler()
 	// Since this can be called from any thread, make sure we only update slate on the game thread.
 	FFunctionGraphTask::CreateAndDispatchWhenReady([this]()
 	{
-		SourcesView->RefreshSourceData(true);
 		RebuildSubjectList();
 	}, TStatId(), nullptr, ENamedThreads::GameThread);
 }
@@ -187,15 +241,9 @@ void FLiveLinkPanelController::OnSubjectsChangedHandler()
 	FFunctionGraphTask::CreateAndDispatchWhenReady([this]()
 	{
 		RebuildSubjectList();
-		SourcesDetailsView->ForceRefresh();
 	}, TStatId(), nullptr, ENamedThreads::GameThread);
 }
 
-void FLiveLinkPanelController::RebuildSourceList()
-{
-	SourcesView->RefreshSourceData(true);
-}
-
 void FLiveLinkPanelController::RebuildSubjectList()
 {
 	if (SubjectsView)
@@ -212,16 +260,19 @@ bool FLiveLinkPanelController::HasSource() const
 
 bool FLiveLinkPanelController::CanRemoveSource() const
 {
-	return SourcesView->SourcesListView->GetNumItemsSelected() > 0;
+	return true;
 }
 
 void FLiveLinkPanelController::HandleRemoveSource()
 {
-	TArray<FLiveLinkSourceUIEntryPtr> Selected;
-	SourcesView->SourcesListView->GetSelectedItems(Selected);
-	if (Selected.Num() > 0)
+	TArray<FLiveLinkSubjectUIEntryPtr> Selected;
+	SubjectsView->SubjectsTreeView->GetSelectedItems(Selected);
+	for (const FLiveLinkSubjectUIEntryPtr& Entry : Selected)
 	{
-		Selected[0]->RemoveFromClient();
+		if (Entry->IsSource())
+		{
+			Entry->RemoveFromClient();
+		}
 	}
 }
 
@@ -235,6 +286,11 @@ bool FLiveLinkPanelController::CanRemoveSubject() const
 	return SubjectsView->CanRemoveSubject();
 }
 
+bool FLiveLinkPanelController::CanRemoveDevice() const
+{
+	return true;
+}
+
 void FLiveLinkPanelController::HandleRemoveSubject()
 {
 	TArray<FLiveLinkSubjectUIEntryPtr> Selected;
@@ -245,6 +301,20 @@ void FLiveLinkPanelController::HandleRemoveSubject()
 	}
 }
 
+void FLiveLinkPanelController::HandleRemoveDevice()
+{
+	TArray<FLiveLinkSubjectUIEntryPtr> Selected;
+	SubjectsView->SubjectsTreeView->GetSelectedItems(Selected);
+	
+	for (const FLiveLinkSubjectUIEntryPtr& Entry : Selected)
+	{
+		if (Entry)
+		{
+			Entry->RemoveFromClient();
+		}
+	}
+}
+
 bool FLiveLinkPanelController::CanPauseSubject() const
 {
 	return SubjectsView->CanPauseSubject();
diff --git a/Engine/Plugins/Animation/LiveLink/Source/LiveLinkEditor/Public/LiveLinkClientPanelViews.h b/Engine/Plugins/Animation/LiveLink/Source/LiveLinkEditor/Public/LiveLinkClientPanelViews.h
--- a/Engine/Plugins/Animation/LiveLink/Source/LiveLinkEditor/Public/LiveLinkClientPanelViews.h
+++ b/Engine/Plugins/Animation/LiveLink/Source/LiveLinkEditor/Public/LiveLinkClientPanelViews.h
@@ -10,6 +10,7 @@
 #include "LiveLinkRole.h"
 #include "LiveLinkSourceSettings.h"
 #include "Misc/Guid.h"
+#include "PropertyEditorDelegates.h"
 #include "SLiveLinkFilterSearchBox.h"
 #include "Templates/SharedPointer.h"
 #include "UObject/GCObject.h"
@@ -37,20 +38,37 @@ typedef TSharedPtr<FLiveLinkSubjectUIEntry> FLiveLinkSubjectUIEntryPtr;
 
 namespace UE::LiveLink
 {
+	UE_DEPRECATED(5.8, "Use the other CreateSourcesDetailsView method")
 	TSharedPtr<IDetailsView> UE_API CreateSourcesDetailsView(const TSharedPtr<FLiveLinkSourcesView>& InSourcesView, const TAttribute<bool>& bInReadOnly);
+	TSharedPtr<IDetailsView> UE_API CreateSourcesDetailsView(FOnFinishedChangingProperties InChangedPropertiesDelegate, const TAttribute<bool>& bInReadOnly);
 	TSharedPtr<SLiveLinkDataView> UE_API CreateSubjectsDetailsView(FLiveLinkClient* InLiveLinkClient, const TAttribute<bool>& bInReadOnly);
+	TSharedPtr<IDetailsView> UE_API CreateDevicesDetailsView(FOnFinishedChangingProperties OnFinishedChangingProperties, const TAttribute<bool>& bInReadOnly);
 }
 
+/** Type of the subject entry in the list. */
+enum class ELiveLinkEntryType : uint8
+{
+	Source,
+	Subject,
+	Device,
+	DevicesHeader,
+	SourcesHeader
+};
+
 // Structure that defines a single entry in the subject UI
 struct FLiveLinkSubjectUIEntry
 {
 	FLiveLinkSubjectUIEntry(const FLiveLinkSubjectKey& InSubjectKey, FLiveLinkClient* InClient, bool bIsSource = false);
+	FLiveLinkSubjectUIEntry(FLiveLinkClient* InClient, ELiveLinkEntryType InType);
 
 	// Subject key
 	FLiveLinkSubjectKey SubjectKey;
 	// LiveLink Client
 	FLiveLinkClient* Client;
-
+	// Weak pointer to the device represented by this row.
+	TWeakObjectPtr<class ULiveLinkDevice> WeakDevice;
+	// What type of item this row represents.
+	ELiveLinkEntryType Type = ELiveLinkEntryType::Subject;
 	// Children (if this entry represents a source instead of a specific subject
 	TArray<TSharedPtr<FLiveLinkSubjectUIEntry>> Children;
 
@@ -58,6 +76,8 @@ struct FLiveLinkSubjectUIEntry
 	bool IsSubject() const;
 	// Whether the subject entry is a source
 	bool IsSource() const;
+	// Whether this entry is a device.
+	bool IsDevice() const;
 	// Whether the subject is virtual
 	bool IsVirtualSubject() const;
 	// Get the subject or source settings
@@ -66,10 +86,14 @@ struct FLiveLinkSubjectUIEntry
 	bool IsSubjectEnabled() const;
 	// Whether the subject is valid
 	bool IsSubjectValid() const;
+	// Whether the entry is a playback source.
+	bool IsPlaybackSource() const;
 	// Enable or disable a subject
 	void SetSubjectEnabled(bool bIsEnabled);
 	// Get a textual representation of the ui entry 
 	FText GetItemText() const;
+	// Get the name of the machine that hosts this source or device.
+	FText GetMachineName() const;
 	// Get the livelink role of this entry
 	TSubclassOf<ULiveLinkRole> GetItemRole() const;
 	// Get the translated role for this subject (if translating rebroadcast subjects).
@@ -86,9 +110,8 @@ struct FLiveLinkSubjectUIEntry
 private:
 	// Whether the subject is virtual
 	bool bIsVirtualSubject = false;
-
-	// Whether this represents a source.
-	bool bIsSource = false;
+	// Whether this is a playback source, used to display a different icon when the source is in playback.
+	bool bIsPlaybackSource = false;
 };
 
 // Structure that defines a single entry in the source UI
@@ -202,8 +225,16 @@ class FLiveLinkSubjectsView : public TSharedFromThis<FLiveLinkSubjectsView>
 public:
 	DECLARE_DELEGATE_TwoParams(FOnSubjectSelectionChanged, FLiveLinkSubjectUIEntryPtr, ESelectInfo::Type);
 
+	struct FSubjectsViewArgs
+	{
+		TAttribute<bool> bInReadOnly;
+		bool bInShowDevices = false;
+	};
+
 	UE_API FLiveLinkSubjectsView(FOnSubjectSelectionChanged InOnSubjectSelectionChanged, const TSharedPtr<FUICommandList>& InCommandList, TAttribute<bool> bInReadOnly);
 
+	UE_API FLiveLinkSubjectsView(FOnSubjectSelectionChanged InOnSubjectSelectionChanged, const TSharedPtr<FUICommandList>& InCommandList, const FSubjectsViewArgs& InArgs);
+
 	// Helper functions for building the subject tree UI
 	UE_API TSharedRef<ITableRow> MakeTreeRowWidget(FLiveLinkSubjectUIEntryPtr InInfo, const TSharedRef<STableViewBase>& OwnerTable);
 	// Get a subjectsview's children
@@ -224,14 +255,23 @@ public:
 	UE_API TSharedRef<SWidget> GetWidget();
 
 private:
-	// Create the subjects tree view
-	void CreateSubjectsTreeView(const TSharedPtr<FUICommandList>& InCommandList);
+	void CreateCombinedTreeView(const TSharedPtr<FUICommandList>& InCommandList);
 	// Get the label for the pause subject context menu entry.
 	FText GetPauseSubjectLabel() const;
 	// Get the tooltip for the pause subject context menu entry.
 	FText GetPauseSubjectToolTip() const;
 	// Returns whether the selected subject is paused.
 	bool IsSelectedSubjectPaused() const;
+	// Returns whether a device can be removed.
+	bool CanRemoveDevice() const;
+	// Returns whether a source can be removed.
+	bool CanRemoveSource() const;
+	// Returns whether there are any sources in the list.
+	bool HasAnySources() const;
+	// Opens a context menu with options relevant to the selected items.
+	TSharedPtr<SWidget> OnCombinedListOpenContextMenu(TSharedPtr<FUICommandList> InCommandList);
+	// Populates entries for the Add combo button.
+	TSharedRef<SWidget> OnGenerateAddMenuEntries();
 
 public:
 	// Subject tree widget
@@ -246,6 +286,10 @@ public:
 	TAttribute<bool> bReadOnly;
 	// Filter search box.
 	TSharedPtr<SLiveLinkFilterSearchBox<FLiveLinkSubjectUIEntryPtr>> FilterSearchBox;
+
+private:
+	// Whether to show devices in the menu. (Disabled by default in editor)
+	bool bShowDevices = false;
 };
 
 class FLiveLinkSourcesView : public TSharedFromThis<FLiveLinkSourcesView>, public FGCObject
@@ -289,14 +333,6 @@ private:
 	void CreateSourcesListView(const TSharedPtr<FUICommandList>& InCommandList);
 	// Populates the source combo menu widget content.
 	TSharedRef<SWidget> OnGenerateSourceMenu();
-	// Spawns a window upon selecting a LiveLink source.
-	void OpenCreateMenuWindow(int32 FactoryIndex);
-	// Invokes the selected source factory to create a livelink source.
-	void ExecuteCreateSource(int32 FactoryIndex);
-	// Completes the source creation process after a source factory was invoked.
-	void OnSourceCreated(TSharedPtr<ILiveLinkSource> NewSource, FString ConnectionString, TSubclassOf<ULiveLinkSourceFactory> Factory);
-	// Creates the virtual subject dialog box.
-	void AddVirtualSubject();
 
 public:
 	// Holds the data that will be displayed in the list view
@@ -311,8 +347,6 @@ public:
 	TAttribute<bool> bReadOnly;
 
 private:
-	// Source creation modal window.
-	TSharedPtr<SWindow> SourceCreationWindow;
 	// Filtered list resulting from applying the text filter.
 	TArray<FLiveLinkSourceUIEntryPtr> FilteredList;
 	// Holds the searchbox widget and the source list.
diff --git a/Engine/Plugins/Animation/LiveLink/Source/LiveLinkEditor/Public/LiveLinkPanelController.h b/Engine/Plugins/Animation/LiveLink/Source/LiveLinkEditor/Public/LiveLinkPanelController.h
--- a/Engine/Plugins/Animation/LiveLink/Source/LiveLinkEditor/Public/LiveLinkPanelController.h
+++ b/Engine/Plugins/Animation/LiveLink/Source/LiveLinkEditor/Public/LiveLinkPanelController.h
@@ -7,6 +7,7 @@
 #include "Delegates/DelegateCombinations.h"
 #include "Misc/Attribute.h"
 #include "Templates/SharedPointer.h"
+#include "UObject/UnrealType.h"
 
 #define UE_API LIVELINKEDITOR_API
 
@@ -15,11 +16,11 @@ class FLiveLinkClient;
 struct FLiveLinkSourceUIEntry;
 struct FLiveLinkSubjectKey;
 struct FLiveLinkSubjectUIEntry;
-class FLiveLinkSourcesView;
 class SLiveLinkSourceListView;
 class FLiveLinkSubjectsView;
 class FUICommandList;
 class SLiveLinkDataView;
+class SWidget;
 
 namespace ESelectInfo { enum Type : int; }
 typedef TSharedPtr<FLiveLinkSourceUIEntry> FLiveLinkSourceUIEntryPtr;
@@ -40,50 +41,59 @@ public:
 		return SubjectSelectionChangedDelegate;
 	}
 
+	/** Get the widget displaying details for the selected item, whether it's a Source, Subject or Device. */
+	UE_API TSharedRef<SWidget> GetCombinedDetailsWidget();
+
 private:
 	// Bind live link commands 
-	UE_API void BindCommands();
+	void BindCommands();
 	// Handles the source collection changing.
-	UE_API void OnSourcesChangedHandler();
+	void OnSourcesChangedHandler();
 	// Handles the subject collection changing.
-	UE_API void OnSubjectsChangedHandler();
-	// Remove all sources command handlers
-	UE_API bool HasSource() const;
+	void OnSubjectsChangedHandler();
+	// Returns whether there is a source in the list.
+	bool HasSource() const;
 	// Handles the remove source command.
-	UE_API void HandleRemoveSource();
+	void HandleRemoveSource();
 	// Handles the remove all sources command.
-	UE_API void HandleRemoveAllSources();
+	void HandleRemoveAllSources();
 	// Returns whether a source could be removed.
-	UE_API bool CanRemoveSource() const;
+	bool CanRemoveSource() const;
 	// Returns whether a subject can be removed.
-	UE_API bool CanRemoveSubject() const;
+	bool CanRemoveSubject() const;
+	// Returns whether a device can be removed.
+	bool CanRemoveDevice() const;
 	// Returns whethera subject can be paused
-	UE_API bool CanPauseSubject() const;
+	bool CanPauseSubject() const;
 	// Pauses/Unpauses a subject.
-	UE_API void HandlePauseSubject();
+	void HandlePauseSubject();
 	// Handles the remove subject command.
-	UE_API void HandleRemoveSubject();
-	// Recreates the source list data behind the list view.
-	UE_API void RebuildSourceList();
+	void HandleRemoveSubject();
+	// Handles the remove device command.
+	void HandleRemoveDevice();
 	// Recreates the subject list data behind the tree view.
-	UE_API void RebuildSubjectList();
-	// Handles source selection changing.
-	UE_API void OnSourceSelectionChangedHandler(FLiveLinkSourceUIEntryPtr Entry, ESelectInfo::Type SelectionType) const;
+	void RebuildSubjectList();
 	// Hadnles subject selection changing.
-	UE_API void OnSubjectSelectionChangedHandler(FLiveLinkSubjectUIEntryPtr SubjectEntry, ESelectInfo::Type SelectInfo);
-
-	// Sources and subjects live in different tabs in LiveLink Hub.
-	bool bSeparateSourcesSubjects = false;
+	void OnSubjectSelectionChangedHandler(FLiveLinkSubjectUIEntryPtr SubjectEntry, ESelectInfo::Type SelectInfo);
+	// Handles triggering the modified device's OnSettingChanged method.
+	void OnFinishedChangingDeviceProperties(const FPropertyChangedEvent& InPropertyChangedEvent);
+	// Handles triggering LiveLinkClient's OnPropertyChanged method for the modified source.
+	void OnFinishedChangingSourceProperties(const FPropertyChangedEvent& InPropertyChangedEvent);
+	// Returns the index of the details widget that should be shown, according to the current selection.
+	int32 GetCombinedDetailWidgetIndex() const;
 
 public:
 	// Sources view
-	TSharedPtr<FLiveLinkSourcesView> SourcesView;
+	UE_DEPRECATED(5.8, "Use the subjects view instead.")
+	TSharedPtr<class FLiveLinkSourcesView> SourcesView;
 	// Subjects view
 	TSharedPtr<FLiveLinkSubjectsView> SubjectsView;
 	// Reference to connection settings struct details panel
 	TSharedPtr<IDetailsView> SourcesDetailsView;
 	// Reference to the data value struct details panel
 	TSharedPtr<SLiveLinkDataView> SubjectsDetailsView;
+	// Reference to device settings details panel
+	TSharedPtr<IDetailsView> DevicesDetailsView;
 	// LiveLink Client
 	FLiveLinkClient* Client;
 	// Handle to delegate when client sources list has changed
@@ -96,6 +106,12 @@ public:
 	TSharedPtr<FUICommandList> CommandList;
 	// Delegate called when the subject selection changes
 	FOnSubjectSelectionChanged SubjectSelectionChangedDelegate;
+
+private:
+	// Delegate called when a device is added.
+	FDelegateHandle DeviceAddedHandle;
+	// Delegate called when a device is removed.
+	FDelegateHandle DeviceRemovedHandle;
 };
 
 #undef UE_API
diff --git a/Engine/Plugins/Animation/LiveLinkHub/Source/LiveLinkHub/Private/LiveLinkHubCreatorAppMode.h b/Engine/Plugins/Animation/LiveLinkHub/Source/LiveLinkHub/Private/LiveLinkHubCreatorAppMode.h
--- a/Engine/Plugins/Animation/LiveLinkHub/Source/LiveLinkHub/Private/LiveLinkHubCreatorAppMode.h
+++ b/Engine/Plugins/Animation/LiveLinkHub/Source/LiveLinkHub/Private/LiveLinkHubCreatorAppMode.h
@@ -37,73 +37,14 @@
 
 #define LOCTEXT_NAMESPACE "LiveLinkHubCreatorMode"
 
-struct FLiveLinkHubSourcesTabSummoner : public FWorkflowTabFactory
+struct FLiveLinkHubDetailsTabSummoner : public FWorkflowTabFactory
 {
 public:
-	FLiveLinkHubSourcesTabSummoner(TSharedPtr<class FAssetEditorToolkit> InHostingApp, const TSharedPtr<FLiveLinkPanelController>& InPanelController)
-		: FWorkflowTabFactory(UE::LiveLinkHub::SourcesTabId, MoveTemp(InHostingApp))
-		, WeakPanelController(InPanelController)
-	{
-		TabLabel = LOCTEXT("SourcesTabLabel", "Sources");
-		TabIcon = FSlateIcon(UE::LiveLinkHub::LiveLinkStyleName, TEXT("LiveLinkHub.Sources.Icon"));
-
-		bIsSingleton = true;
-	}
-
-	virtual TSharedRef<SWidget> CreateTabBody(const FWorkflowTabSpawnInfo& Info) const override
-	{
-		FLiveLinkClient* Client = (FLiveLinkClient*)&IModularFeatures::Get().GetModularFeature<ILiveLinkClient>(ILiveLinkClient::ModularFeatureName);
-		FLiveLinkHubModule& LiveLinkHubModule = FModuleManager::Get().GetModuleChecked<FLiveLinkHubModule>("LiveLinkHub");
-
-		TSharedRef<SWidget> SourcesListView = SNullWidget::NullWidget;
-		if (TSharedPtr<FLiveLinkPanelController> PanelController = WeakPanelController.Pin())
-		{
-			SourcesListView = PanelController->SourcesView->GetWidget();
-		}
-
-		return SNew(SVerticalBox)
-			+ SVerticalBox::Slot()
-			.Padding(FMargin(4.0f, 0.f, 0.f, 0.f))
-			.AutoHeight()
-			[
-				SNew(SLiveLinkClientPanelToolbar, Client)
-					.SourceButtonAlignment(HAlign_Right)
-					.ParentWindow(LiveLinkHubModule.GetLiveLinkHub()->GetRootWindow())
-					.ShowPresetPicker(false)
-					.ShowSettings(false)
-					.IsEnabled_Lambda([this]()
-						{
-							const FLiveLinkHubModule& LiveLinkHubModule = FModuleManager::Get().GetModuleChecked<FLiveLinkHubModule>("LiveLinkHub");
-							return !LiveLinkHubModule.GetPlaybackController()->IsInPlayback() && !LiveLinkHubModule.GetRecordingController()->IsRecording();
-						})
-			]
-			+ SVerticalBox::Slot()
-			.FillHeight(1.f)
-			.Padding(FMargin(0.0f, 4.0f, 0.0f, 0.0f))
-			[
-				SourcesListView
-			];
-	}
-
-	virtual FText GetTabToolTipText(const FWorkflowTabSpawnInfo& Info) const override
-	{
-		return LOCTEXT("SourcesTabToolTip", "Displays the list of Live Link sources.");
-	}
-
-private:
-	/** Panel controller used to fetch the list of sources. */
-	TWeakPtr<FLiveLinkPanelController> WeakPanelController;
-};
-
-
-struct FLiveLinkHubSourceDetailsTabSummoner : public FWorkflowTabFactory
-{
-public:
-	FLiveLinkHubSourceDetailsTabSummoner(TSharedPtr<class FAssetEditorToolkit> InHostingApp, const TSharedPtr<FLiveLinkPanelController>& InPanelController)
+	FLiveLinkHubDetailsTabSummoner(TSharedPtr<class FAssetEditorToolkit> InHostingApp, const TSharedPtr<FLiveLinkPanelController>& InPanelController)
 		: FWorkflowTabFactory(UE::LiveLinkHub::SourceDetailsTabId, MoveTemp(InHostingApp))
 		, WeakPanelController(InPanelController)
 	{
-		TabLabel = LOCTEXT("SourceDetailsTabLabel", "Source Details");
+		TabLabel = LOCTEXT("DetailsTabLabel", "Details");
 		TabIcon = FSlateIcon(FAppStyle::GetAppStyleSetName(), "LevelEditor.Tabs.Details");
 
 		bIsSingleton = true;
@@ -111,13 +52,13 @@ public:
 
 	virtual TSharedRef<SWidget> CreateTabBody(const FWorkflowTabSpawnInfo& Info) const override
 	{
-		TSharedRef<SWidget> SourcesDetails = SNullWidget::NullWidget;
+		TSharedRef<SWidget> DetailsWidget = SNullWidget::NullWidget;
 		if (TSharedPtr<FLiveLinkPanelController> PanelController = WeakPanelController.Pin())
 		{
-			SourcesDetails = PanelController->SourcesDetailsView.ToSharedRef();
+			DetailsWidget = PanelController->GetCombinedDetailsWidget();
 		}
 
-		return SourcesDetails;
+		return DetailsWidget;
 	}
 
 	virtual FText GetTabToolTipText(const FWorkflowTabSpawnInfo& Info) const override
@@ -129,14 +70,14 @@ private:
 	TWeakPtr<FLiveLinkPanelController> WeakPanelController;
 };
 
-struct FLiveLinkHubSubjectsTabSummoner : public FWorkflowTabFactory
+struct FLiveLinkHubDataDevicesTabSummoner : public FWorkflowTabFactory
 {
 public:
-	FLiveLinkHubSubjectsTabSummoner(TSharedPtr<class FAssetEditorToolkit> InHostingApp, const TSharedPtr<FLiveLinkPanelController>& InPanelController)
+	FLiveLinkHubDataDevicesTabSummoner(TSharedPtr<class FAssetEditorToolkit> InHostingApp, const TSharedPtr<FLiveLinkPanelController>& InPanelController)
 		: FWorkflowTabFactory(UE::LiveLinkHub::SubjectsTabId, MoveTemp(InHostingApp))
 		, WeakPanelController(InPanelController)
 	{
-		TabLabel = LOCTEXT("SubjectsTabLabel", "Subjects");
+		TabLabel = LOCTEXT("DataDevicesLabel", "Data Devices");
 		TabIcon = FSlateIcon("LiveLinkStyle", "LiveLinkHub.Subjects.Icon");
 
 		bIsSingleton = true;
@@ -156,40 +97,14 @@ public:
 
 	virtual FText GetTabToolTipText(const FWorkflowTabSpawnInfo& Info) const override
 	{
-		return LOCTEXT("SubjectsTabToolTip", "View the list of Live Link subjects.");
+		return LOCTEXT("DataDevicesTabToolTip", "View the list of Live Link subjects, sources, and devices.");
 	}
 
 private:
-	/** Panel controller used to get the list of subjects. */
+	/** Panel controller used to get the list of data devices. */
 	TWeakPtr<FLiveLinkPanelController> WeakPanelController;
 };
 
-/** Tab summoner for the subject details tab. */
-struct FLiveLinkHubSubjectDetailsTabSummoner : public FWorkflowTabFactory
-{
-public:
-	FLiveLinkHubSubjectDetailsTabSummoner(TSharedPtr<class FAssetEditorToolkit> InHostingApp)
-		: FWorkflowTabFactory(UE::LiveLinkHub::SubjectDetailsTabId, MoveTemp(InHostingApp))
-	{
-		TabLabel = LOCTEXT("SubjectDetailsTabLabel", "Subject Details");
-		TabIcon = FSlateIcon("LiveLinkStyle", "LiveLinkHub.Subjects.Icon");
-
-		bIsSingleton = true;
-	}
-
-	virtual TSharedRef<SWidget> CreateTabBody(const FWorkflowTabSpawnInfo& Info) const override
-	{
-		FLiveLinkHubModule& LiveLinkHubModule = FModuleManager::Get().GetModuleChecked<FLiveLinkHubModule>("LiveLinkHub");
-		return LiveLinkHubModule.GetSubjectController()->MakeSubjectView();
-	}
-
-	virtual FText GetTabToolTipText(const FWorkflowTabSpawnInfo& Info) const override
-	{
-		return LOCTEXT("SubjectDetailsTabToolTip", "View details about the selected subject.");
-	}
-};
-
-
 struct FLiveLinkHubRecordingListTabSummoner : public FWorkflowTabFactory
 {
 public:
@@ -296,8 +211,7 @@ public:
 	FLiveLinkHubCreatorAppMode(TSharedPtr<FLiveLinkHubApplicationBase> App)
 		: FLiveLinkHubApplicationMode("CreatorMode", LOCTEXT("CreatorModeLabel", "Live Data"), App)
 	{
-
-		TabLayout = FTabManager::NewLayout("LiveLinkHubCreatorMode_v1.4");
+		TabLayout = FTabManager::NewLayout("LiveLinkHubCreatorMode_v1.5");
 		const TSharedRef<FTabManager::FArea> MainWindowArea = FTabManager::NewPrimaryArea()->SetOrientation(Orient_Horizontal);
 
 		PanelController = MakeShared<FLiveLinkPanelController>(TAttribute<bool>::CreateRaw(this, &FLiveLinkHubCreatorAppMode::IsSourcePanelReadOnly));
@@ -322,7 +236,7 @@ public:
 						FTabManager::NewStack()
 						->SetSizeCoefficient(0.5f)
 						->SetHideTabWell(false)
-						->AddTab(UE::LiveLinkHub::SourcesTabId, ETabState::OpenedTab)
+						->AddTab(UE::LiveLinkHub::SubjectsTabId, ETabState::OpenedTab)
 					)
 					->Split
 					(
@@ -344,12 +258,6 @@ public:
 						->AddTab(UE::LiveLinkHub::RecordingListTabId, ETabState::OpenedTab)
 						->SetForegroundTab(UE::LiveLinkHub::SubjectsTabId)
 					)
-					->Split
-					(
-						FTabManager::NewStack()
-						->SetSizeCoefficient(0.5f)
-						->AddTab(UE::LiveLinkHub::SubjectDetailsTabId, ETabState::OpenedTab)
-					)
 				)
 				->Split
 				(
@@ -380,10 +288,8 @@ public:
 			)
 		);
 		
-		TabFactories.RegisterFactory(MakeShared<FLiveLinkHubSourcesTabSummoner>(App, PanelController));
-		TabFactories.RegisterFactory(MakeShared<FLiveLinkHubSourceDetailsTabSummoner>(App, PanelController));
-		TabFactories.RegisterFactory(MakeShared<FLiveLinkHubSubjectsTabSummoner>(App, PanelController));
-		TabFactories.RegisterFactory(MakeShared<FLiveLinkHubSubjectDetailsTabSummoner>(App));
+		TabFactories.RegisterFactory(MakeShared<FLiveLinkHubDetailsTabSummoner>(App, PanelController));
+		TabFactories.RegisterFactory(MakeShared<FLiveLinkHubDataDevicesTabSummoner>(App, PanelController));
 		TabFactories.RegisterFactory(MakeShared<FLiveLinkHubPlaybackTabSummoner>(App));
 		TabFactories.RegisterFactory(MakeShared<FLiveLinkHubRecordingListTabSummoner>(App));
 		TabFactories.RegisterFactory(MakeShared<FLiveLinkHubClientsTabSummoner>(App));
diff --git a/Engine/Source/Programs/LiveLinkHub/Config/DefaultEngine.ini b/Engine/Source/Programs/LiveLinkHub/Config/DefaultEngine.ini
--- a/Engine/Source/Programs/LiveLinkHub/Config/DefaultEngine.ini
+++ b/Engine/Source/Programs/LiveLinkHub/Config/DefaultEngine.ini
@@ -3,7 +3,7 @@ bCanSaveAllSections=true
 
 [LiveLink]
 bUseLiveLinkHubSettingsDetailCustomization=true
-bPanelControllerSeparateSourcesSubjects=true
+bPanelControllerSeparateSourcesSubjects=false
 bUseModularClientReference=true
 bEnableParentSubjects=true
 bCreateLiveLinkHubInstance=true
@@ -42,6 +42,7 @@ bShareMaterialShaderCode=false
 [ConsoleVariables]
 DetailsPanel.Style.EnableCardLayout=true
 Slate.EnableGamepadEditorNavigation=false
+LiveLinkHub.EnableStatusBar=false
 
 [/Script/BuildSettings.BuildSettings]
 DefaultEditorTarget=LiveLinkHub
