LiveLinkHub - Add filtering support to the Data Devices panel

#Synced-CL 49625772

diff --git a/Engine/Plugins/Animation/LiveLink/Source/LiveLinkEditor/LiveLinkEditor.Build.cs b/Engine/Plugins/Animation/LiveLink/Source/LiveLinkEditor/LiveLinkEditor.Build.cs
--- a/Engine/Plugins/Animation/LiveLink/Source/LiveLinkEditor/LiveLinkEditor.Build.cs
+++ b/Engine/Plugins/Animation/LiveLink/Source/LiveLinkEditor/LiveLinkEditor.Build.cs
@@ -50,6 +50,7 @@ namespace UnrealBuildTool.Rules
 					"EditorFramework",
 					"UnrealEd",
 					"WorkspaceMenuStructure",
+					"ToolMenus",
 					"ToolWidgets",
 				}
 			);
diff --git a/Engine/Plugins/Animation/LiveLink/Source/LiveLinkEditor/Private/LiveLinkClientPanelViews.cpp b/Engine/Plugins/Animation/LiveLink/Source/LiveLinkEditor/Private/LiveLinkClientPanelViews.cpp
--- a/Engine/Plugins/Animation/LiveLink/Source/LiveLinkEditor/Private/LiveLinkClientPanelViews.cpp
+++ b/Engine/Plugins/Animation/LiveLink/Source/LiveLinkEditor/Private/LiveLinkClientPanelViews.cpp
@@ -8,6 +8,9 @@
 #include "Editor.h"
 #include "EditorFontGlyphs.h"
 #include "Features/IModularFeatures.h"
+#include "Filters/FilterBase.h"
+#include "Filters/GenericFilter.h"
+#include "Filters/SBasicFilterBar.h"
 #include "Framework/Commands/UICommandList.h"
 #include "IDetailsView.h"
 #include "ILiveLinkClient.h"
@@ -31,6 +34,7 @@
 #include "SPositiveActionButton.h"
 #include "Styling/SlateStyle.h"
 #include "Styling/SlateStyleRegistry.h"
+#include "ToolMenus.h"
 #include "Widgets/DeclarativeSyntaxSupport.h"
 #include "Widgets/Input/SButton.h"
 #include "Widgets/Input/SCheckBox.h"
@@ -38,7 +42,6 @@
 #include "Widgets/Input/STextEntryPopup.h"
 #include "Widgets/Layout/SUniformGridPanel.h"
 
-
 #define LOCTEXT_NAMESPACE "LiveLinkClientPanel.PanelViews"
 
 // Static Source UI FNames
@@ -51,6 +54,69 @@ namespace SourceListUI
 };
 
 
+/**
+ * Widget for the filter bar, needs to be a subclass that overrides MakeAddFilterMenu to give the filter bar its own unique menu name,
+ * otherwise the editor will get confused with any other basic filter bar used elsewhere.
+ */
+template<typename TFilterType>
+class SLiveLinkFilterBar : public SBasicFilterBar<TFilterType>
+{
+	using Super = SBasicFilterBar<TFilterType>;
+
+public:
+
+	using FOnFilterChanged = typename SBasicFilterBar<TFilterType>::FOnFilterChanged;
+	using FCreateTextFilter = typename SBasicFilterBar<TFilterType>::FCreateTextFilter;
+
+	SLATE_BEGIN_ARGS(SLiveLinkFilterBar)
+		{
+		}
+		SLATE_EVENT(SLiveLinkFilterBar<TFilterType>::FOnFilterChanged, OnFilterChanged)
+		SLATE_ARGUMENT(TArray<TSharedRef<FFilterBase<TFilterType>>>, CustomFilters)
+	SLATE_END_ARGS()
+
+	void Construct(const FArguments& InArgs)
+	{
+		typename SBasicFilterBar<TFilterType>::FArguments Args;
+		Args._OnFilterChanged = InArgs._OnFilterChanged;
+		Args._CustomFilters = InArgs._CustomFilters;
+		Args._UseSectionsForCategories = true;
+
+		SBasicFilterBar<TFilterType>::Construct(Args.FilterPillStyle(EFilterPillStyle::Basic));
+	}
+
+private:
+	virtual TSharedRef<SWidget> MakeAddFilterMenu() override
+	{
+		const FName FilterMenuName = "LiveLinkFilterBar.FilterMenu";
+		if (!UToolMenus::Get()->IsMenuRegistered(FilterMenuName))
+		{
+			UToolMenu* Menu = UToolMenus::Get()->RegisterMenu(FilterMenuName);
+			Menu->bShouldCloseWindowAfterMenuSelection = true;
+			Menu->bCloseSelfOnly = true;
+
+			Menu->AddDynamicSection(NAME_None, FNewToolMenuDelegate::CreateLambda([](UToolMenu* InMenu)
+				{
+					if (UFilterBarContext* Context = InMenu->FindContext<UFilterBarContext>())
+					{
+						Context->PopulateFilterMenu.ExecuteIfBound(InMenu);
+						Context->OnExtendAddFilterMenu.ExecuteIfBound(InMenu);
+					}
+				}));
+		}
+
+		UFilterBarContext* FilterBarContext = NewObject<UFilterBarContext>();
+		FilterBarContext->PopulateFilterMenu = FOnPopulateAddFilterMenu::CreateLambda([this](UToolMenu* Menu)
+			{
+				Super::PopulateCommonFilterSections(Menu);
+				Super::PopulateCustomFilters(Menu);
+			});
+		FToolMenuContext ToolMenuContext(FilterBarContext);
+
+		return UToolMenus::Get()->GenerateWidget(FilterMenuName, ToolMenuContext);
+	}
+};
+
 /** Dialog to create a new virtual subject */
 class SVirtualSubjectCreateDialog : public SCompoundWidget
 {
@@ -611,22 +677,28 @@ FLiveLinkSubjectUIEntry::FLiveLinkSubjectUIEntry(const FLiveLinkSubjectKey& InSu
 	: SubjectKey(InSubjectKey)
 	, Client(InClient)
 {
-	Type = bInIsSource ? ELiveLinkEntryType::Source : ELiveLinkEntryType::Subject;
-
-	if (InClient)
+	if (Client)
 	{
-		bIsVirtualSubject = InClient->IsVirtualSubject(InSubjectKey);
+		bIsVirtualSubject = Client->IsVirtualSubject(InSubjectKey);
 
 		// LiveLinkClient doesn't have a concept of playback sources since they're a LiveLinkHub concept.
 		// We could replace this by allowing sources to modify their appearance in the source list.
 		bIsPlaybackSource = IsSource() && Client->GetSourceSettings(SubjectKey.Source)->GetName().Contains(TEXT("Playback"));
 	}
+
+	Type = bInIsSource ? ELiveLinkEntryType::Source : ELiveLinkEntryType::Subject;
 }
 
 FLiveLinkSubjectUIEntry::FLiveLinkSubjectUIEntry(FLiveLinkClient* InClient, ELiveLinkEntryType InType)
 	: Client(InClient)
 	, Type(InType)
 {
+	if (Client)
+	{
+		// LiveLinkClient doesn't have a concept of playback sources since they're a LiveLinkHub concept.
+		// We could replace this by allowing sources to modify their appearance in the source list.
+		bIsPlaybackSource = IsSource() && Client->GetSourceSettings(SubjectKey.Source)->GetName().Contains(TEXT("Playback"));
+	}
 }
 
 bool FLiveLinkSubjectUIEntry::IsSubject() const
@@ -774,6 +846,11 @@ bool FLiveLinkSubjectUIEntry::IsPaused() const
 	return Client->GetSubjectState(SubjectKey.SubjectName) == ELiveLinkSubjectState::Paused;
 }
 
+void FLiveLinkSubjectUIEntry::GetFilterText(TArray<FString>& OutStrings) const
+{
+	OutStrings.Add(GetItemText().ToString());
+}
+
 class SLiveLinkClientPanelSubjectRow : public SMultiColumnTableRow<FLiveLinkSubjectUIEntryPtr>
 {
 public:
@@ -1214,10 +1291,6 @@ void FLiveLinkSourcesView::CreateSourcesListView(const TSharedPtr<FUICommandList
 		[
 			SAssignNew(FilterSearchBox, SLiveLinkFilterSearchBox<FLiveLinkSourceUIEntryPtr>)
 			.ItemSource(&SourceData)
-			.OnGatherItems_Lambda([this](TArray<FLiveLinkSourceUIEntryPtr>& Items)
-				{
-					Items.Append(SourceData);
-				})
 			.OnUpdateFilteredList_Lambda([this](const TArray<FLiveLinkSourceUIEntryPtr>& FilteredItems)
 				{
 					FilteredList = FilteredItems;
@@ -1407,7 +1480,18 @@ TSharedRef<ITableRow> FLiveLinkSubjectsView::MakeTreeRowWidget(FLiveLinkSubjectU
 
 void FLiveLinkSubjectsView::GetChildrenForInfo(FLiveLinkSubjectUIEntryPtr InInfo, TArray< FLiveLinkSubjectUIEntryPtr >& OutChildren)
 {
-	OutChildren = InInfo->Children;
+	if (InInfo)
+	{
+		OutChildren.Reserve(InInfo->Children.Num());
+
+		for (const FLiveLinkSubjectUIEntryPtr& Entry : InInfo->Children)
+		{
+			if (Entry && !Entry->bFilteredOut)
+			{
+				OutChildren.Add(Entry);
+			}
+		}
+	}
 }
 
 TSharedPtr<SWidget> FLiveLinkSubjectsView::OnOpenVirtualSubjectContextMenu(TSharedPtr<FUICommandList> InCommandList)
@@ -1669,6 +1753,16 @@ void FLiveLinkSubjectsView::HandlePauseSubject()
 
 TSharedRef<SWidget> FLiveLinkSubjectsView::GetWidget()
 {
+	InitializeFilters();
+
+	TSharedPtr<SLayeredImage> FilterImage = SNew(SLayeredImage)
+		.Image(FAppStyle::Get().GetBrush("Icons.Filter"))
+		.ColorAndOpacity(FSlateColor::UseForeground());
+
+	FilterBar = SNew(SLiveLinkFilterBar<FFilterType>)
+		.CustomFilters(CustomFilters)
+		.OnFilterChanged(this, &FLiveLinkSubjectsView::OnFilterChanged);
+
 	return SNew(SVerticalBox)
 		+ SVerticalBox::Slot()
 		.Padding(FMargin(4.0, 8.0, 4.0, 6.0))
@@ -1690,18 +1784,8 @@ TSharedRef<SWidget> FLiveLinkSubjectsView::GetWidget()
 			.Padding(FMargin(4.0, 2.0))
 			.FillWidth(1.0f)
 			[
-				SAssignNew(FilterSearchBox, SLiveLinkFilterSearchBox<FLiveLinkSubjectUIEntryPtr>)
+				SAssignNew(FilterSearchBox, SLiveLinkFilterSearchBox<FLiveLinkSubjectUIEntryPtr>, FilterBar)
 					.ItemSource(&SubjectData)
-					.OnGatherItems_Lambda([this](TArray<FLiveLinkSubjectUIEntryPtr>& Items)
-						{
-							for (FLiveLinkSubjectUIEntryPtr Source : SubjectData)
-							{
-								for (FLiveLinkSubjectUIEntryPtr Subject : Source->Children)
-								{
-									Items.Add(Subject);
-								}
-							}
-						})
 					.OnUpdateFilteredList_Lambda([this](const TArray<FLiveLinkSubjectUIEntryPtr>& FilteredItems)
 						{
 							FilteredList = FilteredItems;
@@ -1709,11 +1793,21 @@ TSharedRef<SWidget> FLiveLinkSubjectsView::GetWidget()
 							for (const FLiveLinkSubjectUIEntryPtr& Item : FilteredList)
 							{
 								SubjectsTreeView->SetItemExpansion(Item, true);
+								for (const FLiveLinkSubjectUIEntryPtr& Child : Item->Children)
+								{
+									SubjectsTreeView->SetItemExpansion(Child, true);
+								}
 							}
 
 							SubjectsTreeView->RequestListRefresh();
 						})
 			]
+			+ SHorizontalBox::Slot()
+			.Padding(FMargin(4.0, 2.0))
+			.AutoWidth()
+			[
+				SLiveLinkFilterBar<FFilterType>::MakeAddFilterButton(FilterBar.ToSharedRef())
+			]
 		]
 		+ SVerticalBox::Slot()
 		.VAlign(VAlign_Fill)
@@ -1833,5 +1927,55 @@ TSharedRef<SWidget> FLiveLinkSourcesView::OnGenerateSourceMenu()
 	return MenuBuilder.MakeWidget();
 }
 
+void FLiveLinkSubjectsView::InitializeFilters()
+{
+	CustomFilters.Reset();
+
+	// Basic filters, for filtering out sources, subjects and devices.
+	{
+		TSharedPtr<FFilterCategory> BasicFiltersCategory = MakeShared<FFilterCategory>(LOCTEXT("BasicFiltersCategory", "Basic"), FText::GetEmpty());
+
+		TSharedRef<FGenericFilter<FFilterType>> SourcesFilter = MakeShared<FGenericFilter<FFilterType>>(
+			BasicFiltersCategory,
+			TEXT("SourcesFilter"),
+			LOCTEXT("SourcesFilterName", "Sources"),
+			FGenericFilter<FFilterType>::FOnItemFiltered::CreateLambda([](FFilterType InItem)
+				{
+					return InItem && InItem->IsSource();
+				}));
+
+		SourcesFilter->SetToolTipText(LOCTEXT("SourcesFilter", "Only show sources"));
+		CustomFilters.Add(SourcesFilter);
+
+		TSharedRef<FGenericFilter<FFilterType>> SubjectsFilter = MakeShared<FGenericFilter<FFilterType>>(
+			BasicFiltersCategory,
+			TEXT("SubjectsFilter"),
+			LOCTEXT("SubjectsFilterName", "Subjects"),
+			FGenericFilter<FFilterType>::FOnItemFiltered::CreateLambda([](FFilterType InItem)
+				{
+					return InItem && InItem->IsSubject();
+				}));
+
+		SubjectsFilter->SetToolTipText(LOCTEXT("SubjectsFilter", "Only show subjects"));
+		CustomFilters.Add(SubjectsFilter);
+
+		TSharedRef<FGenericFilter<FFilterType>> DevicesFilter = MakeShared<FGenericFilter<FFilterType>>(
+			BasicFiltersCategory,
+			TEXT("DevicesFilter"),
+			LOCTEXT("DevicesFilterName", "Devices"),
+			FGenericFilter<FFilterType>::FOnItemFiltered::CreateLambda([](FFilterType InItem)
+				{
+					return InItem && InItem->IsDevice();
+				}));
+
+		SubjectsFilter->SetToolTipText(LOCTEXT("DevicesFilter", "Only show devices"));
+		CustomFilters.Add(DevicesFilter);
+	}
+}
+
+void FLiveLinkSubjectsView::OnFilterChanged()
+{
+	FilterSearchBox->Update();
+}
 
 #undef LOCTEXT_NAMESPACE /**LiveLinkClientPanel*/
diff --git a/Engine/Plugins/Animation/LiveLink/Source/LiveLinkEditor/Public/LiveLinkClientPanelViews.h b/Engine/Plugins/Animation/LiveLink/Source/LiveLinkEditor/Public/LiveLinkClientPanelViews.h
--- a/Engine/Plugins/Animation/LiveLink/Source/LiveLinkEditor/Public/LiveLinkClientPanelViews.h
+++ b/Engine/Plugins/Animation/LiveLink/Source/LiveLinkEditor/Public/LiveLinkClientPanelViews.h
@@ -33,6 +33,12 @@ class SLiveLinkDataView;
 class STableViewBase;
 class SWidget;
 
+template <typename TFilterType>
+class SLiveLinkFilterBar;
+
+template <typename TFilterType>
+class FFilterBase;
+
 typedef TSharedPtr<FLiveLinkSourceUIEntry> FLiveLinkSourceUIEntryPtr;
 typedef TSharedPtr<FLiveLinkSubjectUIEntry> FLiveLinkSubjectUIEntryPtr;
 
@@ -69,9 +75,13 @@ struct FLiveLinkSubjectUIEntry
 	TWeakObjectPtr<class ULiveLinkDevice> WeakDevice;
 	// What type of item this row represents.
 	ELiveLinkEntryType Type = ELiveLinkEntryType::Subject;
+
 	// Children (if this entry represents a source instead of a specific subject
 	TArray<TSharedPtr<FLiveLinkSubjectUIEntry>> Children;
 
+	// Whether the item should be hidden since it's been filtered out.
+	bool bFilteredOut = false;
+
 	// Whether the subject entry is a subject
 	bool IsSubject() const;
 	// Whether the subject entry is a source
@@ -105,7 +115,7 @@ struct FLiveLinkSubjectUIEntry
 	// Returns whether a subject is currently paused.
 	bool IsPaused() const;
 	// Get item string representation for filtering.
-	void GetFilterText(TArray<FString>& OutStrings) const { OutStrings.Add(SubjectKey.SubjectName.Name.ToString()); }
+	void GetFilterText(TArray<FString>& OutStrings) const;
 
 private:
 	// Whether the subject is virtual
@@ -146,11 +156,15 @@ public:
 		OutStrings.Add(GetSourceType().ToString());
 	}
 
+	// Whether the item should be hidden since it's been filtered out.
+	bool bFilteredOut = false;
+
 private:
 	// GUID of the source
 	FGuid EntryGuid;
 	// Pointer to the livelink client
 	FLiveLinkClient* Client = nullptr;
+
 };
 
 namespace UE::LiveLink
@@ -225,6 +239,9 @@ class FLiveLinkSubjectsView : public TSharedFromThis<FLiveLinkSubjectsView>
 public:
 	DECLARE_DELEGATE_TwoParams(FOnSubjectSelectionChanged, FLiveLinkSubjectUIEntryPtr, ESelectInfo::Type);
 
+	using FFilterType = FLiveLinkSubjectUIEntryPtr;
+	using FMediaTreeFilter = TTextFilter<FFilterType>;
+
 	struct FSubjectsViewArgs
 	{
 		TAttribute<bool> bInReadOnly;
@@ -272,6 +289,14 @@ private:
 	TSharedPtr<SWidget> OnCombinedListOpenContextMenu(TSharedPtr<FUICommandList> InCommandList);
 	// Populates entries for the Add combo button.
 	TSharedRef<SWidget> OnGenerateAddMenuEntries();
+	// Initialize default filters for the list view.
+	void InitializeFilters();
+	// Triggers a refresh of the filtered elements when the filter changes.
+	void OnFilterChanged();
+	// Sources, Subjects and Devices filters.
+	TArray<TSharedRef<FFilterBase<FFilterType>>> CustomFilters;
+	// Filter bar widget. Necessary to create the Add filter button, but it's not actually displayed in the LiveLink UI. 
+	TSharedPtr<SLiveLinkFilterBar<FFilterType>> FilterBar;
 
 public:
 	// Subject tree widget
diff --git a/Engine/Plugins/Animation/LiveLink/Source/LiveLinkEditor/Public/SLiveLinkFilterSearchBox.h b/Engine/Plugins/Animation/LiveLink/Source/LiveLinkEditor/Public/SLiveLinkFilterSearchBox.h
--- a/Engine/Plugins/Animation/LiveLink/Source/LiveLinkEditor/Public/SLiveLinkFilterSearchBox.h
+++ b/Engine/Plugins/Animation/LiveLink/Source/LiveLinkEditor/Public/SLiveLinkFilterSearchBox.h
@@ -3,9 +3,24 @@
 #pragma once
 
 #include "Delegates/Delegate.h"
+#include "Filters/SBasicFilterBar.h"
 #include "Misc/TextFilter.h"
 #include "Widgets/Input/SSearchBox.h"
 
+/** Concept that ensures a type has a field bFilteredOut. */
+template <class Type>
+concept CHasFilteredOutField = requires(Type T)
+{
+	T->bFilteredOut;
+};
+
+/** Concept that ensures a type has a Children field. */
+template <class Type>
+concept CHasChildren = requires(Type T)
+{
+	T->Children;
+};
+
 /**
  * Wrapper around a SSearchBox that provides utility to filter items in a list or tree view.
  */
@@ -23,14 +38,16 @@ public:
 		SLATE_EVENT(FOnUpdateFilteredList, OnUpdateFilteredList);
 	SLATE_END_ARGS()
 
-	void Construct(const FArguments& InArgs)
+	void Construct(const FArguments& InArgs, TSharedPtr<SBasicFilterBar<ItemType>> InFilterBar = nullptr)
 	{
+		static_assert(CHasFilteredOutField<ItemType>);
+
 		SearchTextFilter = MakeShared<FEntryTextFilter>(
 			FEntryTextFilter::FItemToStringArray::CreateSP(this, &SLiveLinkFilterSearchBox::GetItemString));
 
 		ItemSource = InArgs._ItemSource;
-		OnGatherItems = InArgs._OnGatherItems;
 		OnUpdateFilteredList = InArgs._OnUpdateFilteredList;
+		FilterBar = InFilterBar;
 
 		check(ItemSource);
 
@@ -67,20 +84,72 @@ private:
 		bUpdateFilteredItems = true;
 	}
 
+	/** Returns whether an item passes any filter bar filter. */
+	bool PassesAnyFilterBarFilter(const ItemType& InItem)
+	{
+		TSharedPtr<TFilterCollection<ItemType>> FilterCollection = FilterBar->GetAllActiveFilters();
+		const int32 NumFilters = FilterCollection->Num();
+		if (NumFilters == 0)
+		{
+			return true;
+		}
+
+		for (int32 Index = 0; Index < NumFilters; ++Index)
+		{
+			if (FilterCollection->GetFilterAtIndex(Index)->PassesFilter(InItem))
+			{
+				return true;
+			}
+		}
+
+		return false;
+	};
+	
+	/** Updates the bFilteredOut field on the item and its children, then returns if the item or any of its child is visible. */
+	bool UpdateVisibleChildren(ItemType& Item)
+	{
+		const bool bPassesTextSearchbar = SearchTextFilter->GetRawFilterText().IsEmpty() || SearchTextFilter->PassesFilter(Item);
+		bool bFilteredOut = !PassesAnyFilterBarFilter(Item) || !bPassesTextSearchbar;
+
+		for (ItemType& Child : Item->Children)
+		{
+			if (UpdateVisibleChildren(Child))
+			{
+				bFilteredOut = false;
+			}
+		}
+
+		Item->bFilteredOut = bFilteredOut;
+		return !bFilteredOut;
+	}
+
 	/** Updates the filtered item source to match the SearchBox filter. */
 	void UpdateFilteredItemSource()
 	{
 		FilteredItemSource.Reset();
 
-		if (!SearchTextFilter->GetRawFilterText().IsEmpty())
+		if (FilterBar)
 		{
-			// GatherChildren.Execute(ItemType);
-			TArray<ItemType> SourceItems;
-			OnGatherItems.Execute(SourceItems);
-
-			Algo::TransformIf(SourceItems, FilteredItemSource,
-				[this](const ItemType& Item) { return SearchTextFilter->PassesFilter(Item); },
-				[](const ItemType& Item) { return Item; });
+			if constexpr (CHasChildren<ItemType>)
+			{
+				for (ItemType& Item : *ItemSource)
+				{
+					const bool bAnyVisibleChildren = UpdateVisibleChildren(Item);
+					if (bAnyVisibleChildren)
+					{
+						FilteredItemSource.Add(Item);
+					}
+				}
+			}
+			else
+			{
+				Algo::TransformIf(*ItemSource, FilteredItemSource,
+					[this](const ItemType& Item)
+					{
+						return PassesAnyFilterBarFilter(Item) && SearchTextFilter->PassesFilter(Item);
+					},
+					[](const ItemType& Item) { return Item; });
+			}
 		}
 		else
 		{
@@ -109,12 +178,10 @@ private:
 	using FEntryTextFilter = TTextFilter<ItemType>;
 	/** Text filter used for the subject list. */
 	TSharedPtr<FEntryTextFilter> SearchTextFilter;
+
+	/** Optional filter bar used to filter items according to a predefined list of filters. */
+	TSharedPtr<SBasicFilterBar<ItemType>> FilterBar;
 	
-	/** 
-	 * Delegate called to gather items that need to be filtered.
-	 * Normally we could rely on the ItemSource ptr but for a TreeView, we may want to only display the TreeView children.
-	 */
-	FOnGatherItems OnGatherItems;
 	/** Delegate called to the owner of this widget to inform them that the filtered list has been updated. */
 	FOnUpdateFilteredList OnUpdateFilteredList;
 };
diff --git a/Engine/Plugins/Animation/LiveLinkHub/Source/LiveLinkHub/Private/Clients/SLiveLinkHubClientsView.h b/Engine/Plugins/Animation/LiveLinkHub/Source/LiveLinkHub/Private/Clients/SLiveLinkHubClientsView.h
--- a/Engine/Plugins/Animation/LiveLinkHub/Source/LiveLinkHub/Private/Clients/SLiveLinkHubClientsView.h
+++ b/Engine/Plugins/Animation/LiveLinkHub/Source/LiveLinkHub/Private/Clients/SLiveLinkHubClientsView.h
@@ -148,6 +148,8 @@ public:
 	FLiveLinkHubClientId ClientId;
 	/** ClientsModel used to retrieve information about clients/subjects. */
 	TWeakPtr<ILiveLinkHubClientsModel> ClientsModel;
+	/** Whether this client has been filtered out. */
+	bool bFilteredOut = false;
 };
 
 /** Holds a client row's data. */
@@ -481,10 +483,6 @@ public:
 				[
 					SAssignNew(FilterSearchBox, SLiveLinkFilterSearchBox<FClientTreeItemPtr>)
 					.ItemSource(&Clients)
-					.OnGatherItems_Lambda([this](TArray<FClientTreeItemPtr>& Items)
-						{
-							Items.Append(Clients);
-						})
 					.OnUpdateFilteredList_Lambda([this](const TArray<FClientTreeItemPtr>& FilteredItems)
 						{
 							FilteredList = FilteredItems;
