LiveLinkHub - Fix crash when enabling TimedDataMonitor

LiveLinkSubject's FrameData was being accessed on the Game Thread by TimedDataMonitor while it was being written to in a different thread. To fix the issue, the data required for TimeDataMonitor is now cached in LiveLinkSubject's Update method (which runs either on the Game Thread or dedicated thread in LiveLinkHub)

#Synced-CL 49625772

diff --git a/Engine/Plugins/Animation/LiveLink/Source/LiveLink/Private/LiveLinkSubject.cpp b/Engine/Plugins/Animation/LiveLink/Source/LiveLink/Private/LiveLinkSubject.cpp
--- a/Engine/Plugins/Animation/LiveLink/Source/LiveLink/Private/LiveLinkSubject.cpp
+++ b/Engine/Plugins/Animation/LiveLink/Source/LiveLink/Private/LiveLinkSubject.cpp
@@ -224,6 +224,10 @@ void FLiveLinkSubject::Update()
 		UE_LOG(LogLiveLink, Verbose, TEXT("Subject '%s' created snapshot with frame ID %d - Timecode:[%s.%0.3f] - SourceTime: %0.4f, Offset: %0.6f, CorrectedTime: %0.4f"), *SubjectKey.SubjectName.ToString(), FrameSnapshot.FrameData.GetBaseData()->FrameId, *FrameTC.ToString(), FrameSnapshot.FrameData.GetBaseData()->MetaData.SceneTime.Time.GetSubFrame(), FrameSnapshot.FrameData.GetBaseData()->WorldTime.GetSourceTime(), FrameSnapshot.FrameData.GetBaseData()->WorldTime.GetOffset(), FrameSnapshot.FrameData.GetBaseData()->WorldTime.GetOffsettedTime());
 	}
 
+	UpdateCachedOldestDataTime();
+	UpdateCachedNewestDataTime();
+	UdpateCachedDataTimes();
+
 	UpdateState();
 }
 
@@ -1346,6 +1350,16 @@ ETimedDataInputState FLiveLinkSubject::GetState() const
 
 FTimedDataChannelSampleTime FLiveLinkSubject::GetOldestDataTime() const
 {
+	FScopeLock Lock(&CachedTimedDataMonitorInfoCS);
+	return CachedTimeDataMonitorInfo.OldestDataTime;
+}
+
+void FLiveLinkSubject::UpdateCachedOldestDataTime()
+{
+	FScopeLock Lock(&CachedTimedDataMonitorInfoCS);
+
+	CachedTimeDataMonitorInfo.OldestDataTime = FTimedDataChannelSampleTime();
+
 	if (FrameData.Num() > 0)
 	{
 		//For timed data, adjust frame timings to match with clock offset to have valid timing diagrams
@@ -1357,13 +1371,23 @@ FTimedDataChannelSampleTime FLiveLinkSubject::GetOldestDataTime() const
 
 		const double OffsettedOldestTime = FrameData[0].GetBaseData()->WorldTime.GetSourceTime() + CachedSettings.BufferSettings.EngineTimeClockOffset + CachedSettings.BufferSettings.SmoothEngineTimeOffset;
 		const FQualifiedFrameTime OldestSceneTime = FQualifiedFrameTime(FrameData[0].GetBaseData()->MetaData.SceneTime.Time + TimecodeClockOffsetTime, FrameData[0].GetBaseData()->MetaData.SceneTime.Rate);
-		return FTimedDataChannelSampleTime(OffsettedOldestTime, OldestSceneTime);
+
+		CachedTimeDataMonitorInfo.OldestDataTime = FTimedDataChannelSampleTime(OffsettedOldestTime, OldestSceneTime);
 	}
-	return FTimedDataChannelSampleTime();
 }
 
 FTimedDataChannelSampleTime FLiveLinkSubject::GetNewestDataTime() const
 {
+	FScopeLock Lock(&CachedTimedDataMonitorInfoCS);
+	return CachedTimeDataMonitorInfo.NewestDataTime;
+}
+
+void FLiveLinkSubject::UpdateCachedNewestDataTime()
+{
+	FScopeLock Lock(&CachedTimedDataMonitorInfoCS);
+
+	CachedTimeDataMonitorInfo.NewestDataTime = FTimedDataChannelSampleTime();
+
 	if (FrameData.Num() > 0)
 	{
 		//For timed data, adjust frame timings to match with clock offset to have valid timing diagrams
@@ -1375,14 +1399,22 @@ FTimedDataChannelSampleTime FLiveLinkSubject::GetNewestDataTime() const
 
 		const double OffsettedNewestTime = FrameData.Last().GetBaseData()->WorldTime.GetSourceTime() + CachedSettings.BufferSettings.EngineTimeClockOffset + CachedSettings.BufferSettings.SmoothEngineTimeOffset;
 		const FQualifiedFrameTime NewestSceneTime = FQualifiedFrameTime(FrameData.Last().GetBaseData()->MetaData.SceneTime.Time + TimecodeClockOffsetTime, FrameData.Last().GetBaseData()->MetaData.SceneTime.Rate);
-		return FTimedDataChannelSampleTime(OffsettedNewestTime, NewestSceneTime);
+		CachedTimeDataMonitorInfo.NewestDataTime = FTimedDataChannelSampleTime(OffsettedNewestTime, NewestSceneTime);
 	}
-	return FTimedDataChannelSampleTime();
 }
 
 TArray<FTimedDataChannelSampleTime> FLiveLinkSubject::GetDataTimes() const
 {
-	TArray<FTimedDataChannelSampleTime> Result;
+	FScopeLock Lock(&CachedTimedDataMonitorInfoCS);
+	return CachedTimeDataMonitorInfo.DataTimes;
+}
+
+void FLiveLinkSubject::UdpateCachedDataTimes()
+{
+	FScopeLock Lock(&CachedTimedDataMonitorInfoCS);
+
+	CachedTimeDataMonitorInfo.DataTimes.Reset();
+
 	if (FrameData.Num() > 0)
 	{
 		//For timed data, adjust frame timings to match with clock offset to have valid timing diagrams
@@ -1392,16 +1424,14 @@ TArray<FTimedDataChannelSampleTime> FLiveLinkSubject::GetDataTimes() const
 			TimecodeClockOffsetTime = FrameData[0].GetBaseData()->MetaData.SceneTime.Rate.AsFrameTime(CachedSettings.BufferSettings.TimecodeClockOffset);
 		}
 
-		Result.Reset(FrameData.Num());
+		CachedTimeDataMonitorInfo.DataTimes.Reserve(FrameData.Num());
 		for (const FLiveLinkFrameDataStruct& Data : FrameData)
 		{
 			const double OffsettedEngineTime = Data.GetBaseData()->WorldTime.GetSourceTime() + CachedSettings.BufferSettings.EngineTimeClockOffset + CachedSettings.BufferSettings.SmoothEngineTimeOffset;
 			const FQualifiedFrameTime AdjustedSceneTime = FQualifiedFrameTime(Data.GetBaseData()->MetaData.SceneTime.Time + TimecodeClockOffsetTime, Data.GetBaseData()->MetaData.SceneTime.Rate);
-			Result.Emplace(OffsettedEngineTime, AdjustedSceneTime);
+			CachedTimeDataMonitorInfo.DataTimes.Emplace(OffsettedEngineTime, AdjustedSceneTime);
 		}
 	}
-	
-	return Result;
 }
 
 int32 FLiveLinkSubject::GetNumberOfSamples() const
diff --git a/Engine/Plugins/Animation/LiveLink/Source/LiveLink/Public/LiveLinkSubject.h b/Engine/Plugins/Animation/LiveLink/Source/LiveLink/Public/LiveLinkSubject.h
--- a/Engine/Plugins/Animation/LiveLink/Source/LiveLink/Public/LiveLinkSubject.h
+++ b/Engine/Plugins/Animation/LiveLink/Source/LiveLink/Public/LiveLinkSubject.h
@@ -42,6 +42,14 @@ struct FLiveLinkTimeSynchronizationData
 	FFrameTime SyncStartTime;
 };
 
+/** Struct used to cache information that needs to be accessed on the Game Thread by Timed Data Monitor. */
+struct FTimedDataMonitorInfo
+{
+	FTimedDataChannelSampleTime OldestDataTime;
+	FTimedDataChannelSampleTime NewestDataTime;
+	TArray<FTimedDataChannelSampleTime> DataTimes;
+};
+
 DECLARE_DELEGATE_OneParam(FOnStateChanged, ELiveLinkSubjectState /*NewState*/);
 
 /**
@@ -194,6 +202,11 @@ private:
 	/** Update the cached state for this subject. */
 	UE_API void UpdateState();
 
+	//~ Updates cached timed data monitor information from the LiveLink thread (either Game Thread or dedicated thread in Live Link Hub)
+	void UpdateCachedOldestDataTime();
+	void UpdateCachedNewestDataTime();
+	void UdpateCachedDataTimes();
+
 protected:
 	/** The role the subject was build with */
 	TSubclassOf<ULiveLinkRole> Role;
@@ -292,6 +305,15 @@ private:
 
 	/** Last frame ID that was rebroadcasted. Used to avoid LiveLinkClient from rebroadcasting the same frame twice when transmitting evaluated data. */
 	FLiveLinkFrameIdentifier LastRebroadcastFrameId = 0;
+
+	/** 
+	 * This struct is updated on the LiveLink thread (either GameThread or dedicated thread in Live Link Hub).
+	 * Then it's accessed by TimedDataMonitor on the GameThread.
+	 */
+	FTimedDataMonitorInfo CachedTimeDataMonitorInfo;
+
+	/** Critical section used to protect access to the cached timed data monitor information. */
+	mutable FCriticalSection CachedTimedDataMonitorInfoCS;
 	
 	/** 
 	 * Evaluation can be done on any thread so we need to protect statistic logging 
diff --git a/Engine/Source/Programs/LiveLinkHub/Config/DefaultEngine.ini b/Engine/Source/Programs/LiveLinkHub/Config/DefaultEngine.ini
--- a/Engine/Source/Programs/LiveLinkHub/Config/DefaultEngine.ini
+++ b/Engine/Source/Programs/LiveLinkHub/Config/DefaultEngine.ini
@@ -14,7 +14,7 @@ bEnableLLHMessageBusSourceFactory=true
 bUseMessagingConnectionManager=true
 
 [TimedDataMonitor]
-bEnableTimedDataMonitor=false
+bEnableTimedDataMonitor=true
 
 [/Script/MacTargetPlatform.XcodeProjectSettings]
 TemplateMacPlist=(FilePath="/Engine/Source/Programs/LiveLinkHub/Resources/Mac/Info.plist")
