LiveLink: Fix preprocessors being lost upon reconnection.

#rnx

#Synced-CL 50633722

diff --git a/Engine/Plugins/Animation/LiveLink/Source/LiveLink/Private/LiveLinkClient.cpp b/Engine/Plugins/Animation/LiveLink/Source/LiveLink/Private/LiveLinkClient.cpp
--- a/Engine/Plugins/Animation/LiveLink/Source/LiveLink/Private/LiveLinkClient.cpp
+++ b/Engine/Plugins/Animation/LiveLink/Source/LiveLink/Private/LiveLinkClient.cpp
@@ -755,7 +755,10 @@ void FLiveLinkClient::RemoveSource(FGuid InEntryGuid)
 
 void FLiveLinkClient::RemoveAllSources()
 {
+	CachedPresetSubjectSettings.Reset();
+
 	check(Collection);
+
 	Collection->ForEachSource([](FLiveLinkCollectionSourceItem& SourceItem)
 	{
 		// Don't remove virtual subject sources, or restoring Virtual Subjects will fail because their source doesn't exist anymore.
@@ -942,16 +945,29 @@ void FLiveLinkClient::PushSubjectStaticData_Internal(FPendingSubjectStatic&& Sub
 
 		{
 			// By default, try to reuse the previous subject settings when possible. The role and interpolation will be replaced but the other settings (such as bRebroadcastSubject) will be kept.
+			bool bRestoredFromPresetCache = false;
 			if (!SubjectSettings)
 			{
-				// Setting class should always be valid
-				UClass* SettingClass = DefaultSetting.SettingClass.Get();
-				if (SettingClass == nullptr)
+				// Check if we have cached preset settings for this subject.
+				// This allows preprocessors and other settings to survive source disconnection/reconnection.
+				if (TStrongObjectPtr<ULiveLinkSubjectSettings>* CachedSettings = CachedPresetSubjectSettings.Find(SubjectStaticData.SubjectKey.SubjectName))
 				{
-					SettingClass = ULiveLinkSubjectSettings::StaticClass();
+					SubjectSettings = TStrongObjectPtr{ DuplicateObject<ULiveLinkSubjectSettings>(CachedSettings->Get(), GetTransientPackage()) };
+					bRestoredFromPresetCache = true;
 				}
+				else
+				{
+					// Setting class should always be valid
+
+					UClass* SettingClass = DefaultSetting.SettingClass.Get();
 
-				SubjectSettings = TStrongObjectPtr{ NewObject<ULiveLinkSubjectSettings>(GetTransientPackage(), SettingClass) };
+					if (SettingClass == nullptr)
+					{
+						SettingClass = ULiveLinkSubjectSettings::StaticClass();
+					}
+
+					SubjectSettings = TStrongObjectPtr{ NewObject<ULiveLinkSubjectSettings>(GetTransientPackage(), SettingClass) };
+				}
 			}
 
 			SubjectSettings->Initialize(SubjectStaticData.SubjectKey);
@@ -963,51 +979,56 @@ void FLiveLinkClient::PushSubjectStaticData_Internal(FPendingSubjectStatic&& Sub
 				SubjectSettings->OriginalSourceName = **OriginalSourceName;
 			}
 
-			UClass* FrameInterpolationProcessorClass = DefaultSetting.FrameInterpolationProcessor.Get();
-			if (FrameInterpolationProcessorClass != nullptr)
-			{
-				UClass* InterpolationRole = FrameInterpolationProcessorClass->GetDefaultObject<ULiveLinkFrameInterpolationProcessor>()->GetRole();
-				if (SubjectStaticData.Role->IsChildOf(InterpolationRole))
-				{
-					SubjectSettings->InterpolationProcessor = NewObject<ULiveLinkFrameInterpolationProcessor>(SubjectSettings.Get(), FrameInterpolationProcessorClass);
-					// Clear async flag since this might've been created outside the game thread.
-					SubjectSettings->InterpolationProcessor->AtomicallyClearInternalFlags(EInternalObjectFlags::Async);
-				}
-				else
-				{
-					FLiveLinkLog::Warning(TEXT("The interpolator '%s' is not valid for the Role '%s'"), *FrameInterpolationProcessorClass->GetName(), *SubjectStaticData.Role->GetName());
-				}
-			}
-			else
+			// Skip default interpolation/preprocessor setup if we restored settings from a preset cache,
+			// since the cached settings already have the correct processors configured.
+			if (!bRestoredFromPresetCache)
 			{
-				// If no settings were found for a specific role, check if the default interpolator is compatible with the role
-				UClass* FallbackInterpolationProcessorClass = LiveLinkSettings->FrameInterpolationProcessor.Get();
-				if (FallbackInterpolationProcessorClass != nullptr)
+				UClass* FrameInterpolationProcessorClass = DefaultSetting.FrameInterpolationProcessor.Get();
+				if (FrameInterpolationProcessorClass != nullptr)
 				{
-					UClass* InterpolationRole = FallbackInterpolationProcessorClass->GetDefaultObject<ULiveLinkFrameInterpolationProcessor>()->GetRole();
+					UClass* InterpolationRole = FrameInterpolationProcessorClass->GetDefaultObject<ULiveLinkFrameInterpolationProcessor>()->GetRole();
 					if (SubjectStaticData.Role->IsChildOf(InterpolationRole))
 					{
-						SubjectSettings->InterpolationProcessor = NewObject<ULiveLinkFrameInterpolationProcessor>(SubjectSettings.Get(), FallbackInterpolationProcessorClass);
+						SubjectSettings->InterpolationProcessor = NewObject<ULiveLinkFrameInterpolationProcessor>(SubjectSettings.Get(), FrameInterpolationProcessorClass);
 						// Clear async flag since this might've been created outside the game thread.
 						SubjectSettings->InterpolationProcessor->AtomicallyClearInternalFlags(EInternalObjectFlags::Async);
 					}
+					else
+					{
+						FLiveLinkLog::Warning(TEXT("The interpolator '%s' is not valid for the Role '%s'"), *FrameInterpolationProcessorClass->GetName(), *SubjectStaticData.Role->GetName());
+					}
 				}
-			}
-
-			for (TSubclassOf<ULiveLinkFramePreProcessor> PreProcessor : DefaultSetting.FramePreProcessors)
-			{
-				if (PreProcessor != nullptr)
+				else
 				{
-					UClass* PreProcessorRole = PreProcessor->GetDefaultObject<ULiveLinkFramePreProcessor>()->GetRole();
-					if (SubjectStaticData.Role->IsChildOf(PreProcessorRole))
+					// If no settings were found for a specific role, check if the default interpolator is compatible with the role
+					UClass* FallbackInterpolationProcessorClass = LiveLinkSettings->FrameInterpolationProcessor.Get();
+					if (FallbackInterpolationProcessorClass != nullptr)
 					{
-						TObjectPtr<ULiveLinkFramePreProcessor>& FramePreprocessor = SubjectSettings->PreProcessors.Add_GetRef(NewObject<ULiveLinkFramePreProcessor>(SubjectSettings.Get(), PreProcessor.Get()));
-						// Clear async flag since this might've been created outside the game thread.
-						FramePreprocessor->AtomicallyClearInternalFlags(EInternalObjectFlags::Async);
+						UClass* InterpolationRole = FallbackInterpolationProcessorClass->GetDefaultObject<ULiveLinkFrameInterpolationProcessor>()->GetRole();
+						if (SubjectStaticData.Role->IsChildOf(InterpolationRole))
+						{
+							SubjectSettings->InterpolationProcessor = NewObject<ULiveLinkFrameInterpolationProcessor>(SubjectSettings.Get(), FallbackInterpolationProcessorClass);
+							// Clear async flag since this might've been created outside the game thread.
+							SubjectSettings->InterpolationProcessor->AtomicallyClearInternalFlags(EInternalObjectFlags::Async);
+						}
 					}
-					else
+				}
+
+				for (TSubclassOf<ULiveLinkFramePreProcessor> PreProcessor : DefaultSetting.FramePreProcessors)
+				{
+					if (PreProcessor != nullptr)
 					{
-						FLiveLinkLog::Warning(TEXT("The pre processor '%s' is not valid for the Role '%s'"), *PreProcessor->GetName(), *SubjectStaticData.Role->GetName());
+						UClass* PreProcessorRole = PreProcessor->GetDefaultObject<ULiveLinkFramePreProcessor>()->GetRole();
+						if (SubjectStaticData.Role->IsChildOf(PreProcessorRole))
+						{
+							TObjectPtr<ULiveLinkFramePreProcessor>& FramePreprocessor = SubjectSettings->PreProcessors.Add_GetRef(NewObject<ULiveLinkFramePreProcessor>(SubjectSettings.Get(), PreProcessor.Get()));
+							// Clear async flag since this might've been created outside the game thread.
+							FramePreprocessor->AtomicallyClearInternalFlags(EInternalObjectFlags::Async);
+						}
+						else
+						{
+							FLiveLinkLog::Warning(TEXT("The pre processor '%s' is not valid for the Role '%s'"), *PreProcessor->GetName(), *SubjectStaticData.Role->GetName());
+						}
 					}
 				}
 			}
@@ -1306,6 +1327,14 @@ bool FLiveLinkClient::CreateSubject(const FLiveLinkSubjectPreset& InSubjectPrese
 		return false;
 	}
 
+	// Cache the preset settings so they can be restored if the subject is recreated later
+	// (e.g., after source disconnection/reconnection).
+	if (InSubjectPreset.Settings)
+	{
+		CachedPresetSubjectSettings.Add(InSubjectPreset.Key.SubjectName,
+			TStrongObjectPtr{DuplicateObject<ULiveLinkSubjectSettings>(InSubjectPreset.Settings, GetTransientPackage())});
+	}
+
 	FLiveLinkCollectionSourceItem* SourceItem = Collection->FindSource(InSubjectPreset.Key.Source);
 	if (SourceItem == nullptr || SourceItem->bPendingKill)
 	{
diff --git a/Engine/Plugins/Animation/LiveLink/Source/LiveLink/Public/LiveLinkClient.h b/Engine/Plugins/Animation/LiveLink/Source/LiveLink/Public/LiveLinkClient.h
--- a/Engine/Plugins/Animation/LiveLink/Source/LiveLink/Public/LiveLinkClient.h
+++ b/Engine/Plugins/Animation/LiveLink/Source/LiveLink/Public/LiveLinkClient.h
@@ -293,6 +293,9 @@ private:
 	/** Pending frame info to add to a subject. */
 	TArray<FPendingSubjectFrame> SubjectFrameToPush;
 
+	/** Cached subject settings from presets, used to restore preprocessors when subjects are recreated (e.g. after source reconnection). */
+	TMap<FLiveLinkSubjectName, TStrongObjectPtr<ULiveLinkSubjectSettings>> CachedPresetSubjectSettings;
+
 	/** Key funcs for looking up a set of cached keys by its layout element */
 	TMap<FLiveLinkSubjectName, FLiveLinkSubjectKey> EnabledSubjects;
 
