LiveLink: Providers now auto-assign FrameId, and LiveLink Message Bus uses this to resequence them into order.

UDP messaging async nature can deliver frames to live link out of order, which can affect live link custom timestep or other features that may assume that packets arrive in order. This change resequences them based on the FrameId, and only buffers them when it gets packets out of order, which can happen relatively often especially when MaxConcurrentDeserializationTasks is more than 1.

Also, the default evaluation mode for the LL TC provider is now Latest. In this mode, when LL genlock and LL TC provider are used at the same time, it ensures that the timecode used is the one corresponding to the genlock trigger.

In addition, subframe behavior for Custom Timestep genlock should now deterministic.

When sequencer replays a live link track and creates a new live link source with the same subject name, this is detected by the LL custom timestep and self-disables. When the original subject comes back online then it re-synchronizes. The timecode provider however tries to use the subject from the previous source and continue to work normally.

#rnx

#Synced-CL 50633722

diff --git a/Engine/Plugins/Animation/LiveLink/Source/LiveLink/Private/LiveLinkClient.cpp b/Engine/Plugins/Animation/LiveLink/Source/LiveLink/Private/LiveLinkClient.cpp
--- a/Engine/Plugins/Animation/LiveLink/Source/LiveLink/Private/LiveLinkClient.cpp
+++ b/Engine/Plugins/Animation/LiveLink/Source/LiveLink/Private/LiveLinkClient.cpp
@@ -485,6 +485,11 @@ bool FLiveLinkClient::HandleSubjectRebroadcast(ILiveLinkSubject* InSubject, cons
 void FLiveLinkClient::OnSubjectRemovedCallback(FLiveLinkSubjectKey InSubjectKey)
 {
 	RemoveRebroadcastedSubject(InSubjectKey);
+
+	{
+		FScopeLock Lock(&SubjectFrameCountersLock);
+		SubjectFrameCounters.Remove(InSubjectKey);
+	}
 }
 
 void FLiveLinkClient::RemoveRebroadcastedSubject(FLiveLinkSubjectKey InSubjectKey)
@@ -1110,6 +1115,18 @@ void FLiveLinkClient::PushSubjectFrameData_AnyThread(const FLiveLinkSubjectKey&
 
 	TRACE_CPUPROFILER_EVENT_SCOPE_TEXT(*FString::Printf(TEXT("FLiveLinkClient::PushSubjectFrameData_AnyThread %s"), *InSubjectKey.SubjectName.Name.ToString()));
 
+	// Assign FrameId using uint32 wraparound arithmetic.
+	if (FLiveLinkBaseFrameData* BaseData = SubjectFrame.FrameData.GetBaseData())
+	{
+		if (BaseData->FrameId == INDEX_NONE)
+		{
+			FScopeLock Lock(&SubjectFrameCountersLock);
+			FLiveLinkFrameIdentifier& Counter = SubjectFrameCounters.FindOrAdd(InSubjectKey, 0);
+			BaseData->FrameId = Counter;
+			Counter = int32(uint32(Counter) + 1u);
+		}
+	}
+
 	{
 		FScopeLock Lock(&PendingFramesCriticalSection);
 
diff --git a/Engine/Plugins/Animation/LiveLink/Source/LiveLink/Private/LiveLinkCustomTimeStep.cpp b/Engine/Plugins/Animation/LiveLink/Source/LiveLink/Private/LiveLinkCustomTimeStep.cpp
--- a/Engine/Plugins/Animation/LiveLink/Source/LiveLink/Private/LiveLinkCustomTimeStep.cpp
+++ b/Engine/Plugins/Animation/LiveLink/Source/LiveLink/Private/LiveLinkCustomTimeStep.cpp
@@ -5,6 +5,8 @@
 #include "Features/IModularFeatures.h"
 #include "GenlockedCustomTimeStep.h"
 #include "ILiveLinkClient.h"
+#include "LiveLinkClient.h"
+#include "LiveLinkTypes.h"
 #include "Misc/App.h"
 
 #include UE_INLINE_GENERATED_CPP_BY_NAME(LiveLinkCustomTimeStep)
@@ -35,7 +37,11 @@ void ULiveLinkCustomTimeStep::Shutdown(UEngine* InEngine)
 
 bool ULiveLinkCustomTimeStep::UpdateTimeStep(UEngine* InEngine)
 {
-	if (State != ECustomTimeStepSynchronizationState::Synchronized)
+	const bool bIsSequencer = bIsSequencerSourceAnyThread.load(std::memory_order_relaxed);
+	const bool bSynchronized = (State == ECustomTimeStepSynchronizationState::Synchronized);
+
+	// Skip genlock if not synchronized OR if subject is from sequencer
+	if (!bSynchronized || bIsSequencer)
 	{
 		EventWaitForSync->Trigger(); // Prevents start/restart quasi-deadlock.
 
@@ -141,6 +147,12 @@ bool ULiveLinkCustomTimeStep::WaitForSync()
 	return EventLiveLink->Wait(FTimespan(TimeoutInSeconds * ETimespan::TicksPerSecond));
 }
 
+FQualifiedFrameTime ULiveLinkCustomTimeStep::GetGenlockTriggerSceneTime() const
+{
+	FScopeLock Lock(&GenlockSceneTimeLock);
+	return GenlockTriggerSceneTime;
+}
+
 
 void ULiveLinkCustomTimeStep::BeginDestroy()
 {
@@ -208,12 +220,31 @@ void ULiveLinkCustomTimeStep::RegisterLiveLinkSubject()
 {
 	RegisteredSubjectKey = SubjectKey;
 
-	FDelegateHandle Tmp;
+	if (!LiveLinkClient)
+	{
+		return;
+	}
+
+	// Check if this is a sequencer source and cache the result
+
+	FText SourceType = LiveLinkClient->GetSourceType(SubjectKey.Source);
+	FString SourceTypeStr = SourceType.ToString();
+	bool bIsSequencer = SourceTypeStr.Contains(UE::LiveLink::SourceTypes::SequencerSourcePrefix, ESearchCase::IgnoreCase);
+	bIsSequencerSourceAnyThread.store(bIsSequencer, std::memory_order_relaxed);
+
+	if (bIsSequencer)
+	{
+		UE_LOG(LogLiveLink, Log, TEXT("Sequencer source registered for '%s', so genlock will be bypassed to avoid deadlocks"),
+			*SubjectKey.SubjectName.ToString());
+	}
+
+	FDelegateHandle DummyStaticDelegateHandle;
+
 	LiveLinkClient->RegisterForFrameDataReceived(
 		RegisteredSubjectKey
 		, FOnLiveLinkSubjectStaticDataReceived::FDelegate()
 		, FOnLiveLinkSubjectFrameDataReceived::FDelegate::CreateUObject(this, &ULiveLinkCustomTimeStep::OnLiveLinkFrameDataReceived_AnyThread)
-		, Tmp
+		, DummyStaticDelegateHandle
 		, RegisterForFrameDataReceivedHandle);
 }
 
@@ -222,10 +253,22 @@ void ULiveLinkCustomTimeStep::UnregisterLiveLinkSubject()
 {
 	if (RegisterForFrameDataReceivedHandle.IsValid())
 	{
+		// Signal the frame data callback to skip the lockstep wait. Without this, the game
+		// thread can stall on the broadcast critical section for up to LockStepTimeout because
+		// OnLiveLinkFrameDataReceived_AnyThread holds that lock during its lockstep wait loop.
+		bSkipLockStepWait.store(true, std::memory_order_release);
+		SyncCount.store(0, std::memory_order_release);
+		EventWaitForSync->Trigger();
+
 		LiveLinkClient->UnregisterForFrameDataReceived(RegisteredSubjectKey, FDelegateHandle(), RegisterForFrameDataReceivedHandle);
 		RegisterForFrameDataReceivedHandle.Reset();
+
+		bSkipLockStepWait.store(false, std::memory_order_release);
 	}
 
+	// Reset sequencer flag when subject is unregistered
+	bIsSequencerSourceAnyThread.store(false, std::memory_order_relaxed);
+
 	RegisteredSubjectKey = FLiveLinkSubjectKey();
 }
 
@@ -253,9 +296,30 @@ void ULiveLinkCustomTimeStep::OnLiveLinkSubjectAdded(FLiveLinkSubjectKey InSubje
 {
 	if (InSubjectKey.SubjectName == SubjectKey.SubjectName)
 	{
-		// Update the subject key, it's possible we didn't save with the same source ID.
+		// Check if already registered to a subject with same name.
+		if (RegisteredSubjectKey.SubjectName == InSubjectKey.SubjectName && RegisteredSubjectKey.Source.IsValid())
+		{
+			UnregisterLiveLinkSubject();
+		}
+
+		// Preserve original configured key before updating
+		FLiveLinkSubjectKey OriginalConfiguredKey = SubjectKey;
+
+		// Update the subject key
 		SubjectKey = InSubjectKey;
 		RegisterLiveLinkSubject();
+
+		// If successfully registered, clear error state
+		if (State == ECustomTimeStepSynchronizationState::Error)
+		{
+			State = ECustomTimeStepSynchronizationState::Synchronizing;
+		}
+
+		// Restore configured key if registered to sequencer
+		if (bIsSequencerSourceAnyThread.load(std::memory_order_relaxed))
+		{
+			SubjectKey = OriginalConfiguredKey;
+		}
 	}
 }
 
@@ -264,8 +328,42 @@ void ULiveLinkCustomTimeStep::OnLiveLinkSubjectRemoved(FLiveLinkSubjectKey InSub
 {
 	if (InSubjectKey.SubjectName == RegisteredSubjectKey.SubjectName)
 	{
-		State = ECustomTimeStepSynchronizationState::Error;
+		// Check if this is a sequencer source being removed
+		const bool bWasSequencerSource = bIsSequencerSourceAnyThread.load(std::memory_order_relaxed);
+
+		// Unregister from the removed source
 		UnregisterLiveLinkSubject();
+
+		// If the removed source was sequencer's, try to re-register to the original external subject
+		if (bWasSequencerSource && LiveLinkClient)
+		{
+			// Check if the original configured subject still exists
+			if (LiveLinkClient->GetSubjectSettings(SubjectKey))
+			{
+				UE_LOG(LogLiveLink, Warning, TEXT("OnLiveLinkSubjectRemoved: Original external subject '%s' found, re-registering"),
+					*SubjectKey.SubjectName.ToString());
+
+				// Re-register to the original external subject
+				RegisterLiveLinkSubject();
+
+				// Clear error state
+				if (State == ECustomTimeStepSynchronizationState::Error)
+				{
+					State = ECustomTimeStepSynchronizationState::Synchronizing;
+				}
+			}
+			else
+			{
+				UE_LOG(LogLiveLink, Log, TEXT("OnLiveLinkSubjectRemoved: Original subject not found, entering Error state"));
+				State = ECustomTimeStepSynchronizationState::Error;
+			}
+		}
+		else
+		{
+			// External subject was removed, enter error state
+			UE_LOG(LogLiveLink, Log, TEXT("OnLiveLinkSubjectRemoved: External subject removed, entering Error state"));
+			State = ECustomTimeStepSynchronizationState::Error;
+		}
 	}
 }
 
@@ -279,6 +377,7 @@ void ULiveLinkCustomTimeStep::OnLiveLinkFrameDataReceived_AnyThread(const FLiveL
 	// reflect source delta times in the engine delta times.
 
 	const FLiveLinkFrameIdentifier FrameId = InFrameData.GetBaseData()->FrameId;
+	const FQualifiedFrameTime CurrentSceneTime = InFrameData.GetBaseData()->MetaData.SceneTime;
 
 	// Receiving data from the desired live link source automatically means that we are in synchronized state.
 	if (State != ECustomTimeStepSynchronizationState::Synchronized)
@@ -287,6 +386,25 @@ void ULiveLinkCustomTimeStep::OnLiveLinkFrameDataReceived_AnyThread(const FLiveL
 
 		// Ensure our initial delta frame id is 1.
 		LastFrameId = FrameId - 1;
+
+		// Reset tracking for phase alignment
+		LastCapturedTimecodeFrameNumber.store(INDEX_NONE, std::memory_order_release);
+		ReferenceFrameId.store(INDEX_NONE, std::memory_order_release);
+
+		// Calculate frames per timecode from LiveLink data rate and subject timecode rate
+
+		const double SourceFPS = LiveLinkDataRate.AsDecimal();
+		const double TimecodeFPS = CurrentSceneTime.Rate.AsDecimal();
+
+		if (SourceFPS > 0.0 && TimecodeFPS > 0.0)
+		{
+			const int32 FramesPerTC = FMath::Max(1, FMath::RoundToInt(SourceFPS / TimecodeFPS));
+			FramesPerTimecodeAnyThread.store(FramesPerTC, std::memory_order_release);
+		}
+		else
+		{
+			UE_LOG(LogLiveLink, Error, TEXT("LiveLinkCustomTimeStep: Invalid frame rates - SourceFPS=%.2f, TimecodeFPS=%.2f"), SourceFPS, TimecodeFPS);
+		}
 	}
 
 	// Loop around cannot be correctly handled unless we know the lower and upper bounds of the frame id.
@@ -302,10 +420,65 @@ void ULiveLinkCustomTimeStep::OnLiveLinkFrameDataReceived_AnyThread(const FLiveL
 		SyncCount += FrameId - LastFrameId;
 	}
 
-	// Note: We don't expect OnLiveLinkFrameDataReceived_AnyThread to be called for the same live link source concurrently, 
+	// Note: We don't expect OnLiveLinkFrameDataReceived_AnyThread to be called for the same live link source concurrently,
 	// otherwise we would need to avoid race conditions here when updating and comparing with FrameId.
+	const FLiveLinkFrameIdentifier PrevFrameId = LastFrameId.load(std::memory_order_acquire);
 	LastFrameId = FrameId;
 
+	// Extract timecode frame for deterministic subframe capture
+
+	const int32 CurrentTimecodeFrame = CurrentSceneTime.Time.FloorToFrame().Value;
+	const int32 PreviousTimecodeFrame = LastCapturedTimecodeFrameNumber.load(std::memory_order_acquire);
+
+	int32 RefFrameId = ReferenceFrameId.load(std::memory_order_acquire);
+
+	const int32 FramesPerTC = FramesPerTimecodeAnyThread.load(std::memory_order_acquire);
+	const int32 Divider = static_cast<int32>(FrameRateDividerAnyThread.load(std::memory_order_relaxed));
+
+	// Establish reference frame on first new timecode with consecutive FrameId
+	// This ensures we capture from the actual first frame (subframe 0.0) and not a potentially skipped subframe
+	if (RefFrameId == INDEX_NONE || CurrentTimecodeFrame != PreviousTimecodeFrame)
+	{
+		// Check if timecode changed and frame is consecutive
+
+		const bool bIsFirstFrame = (PrevFrameId == FrameId - 1) || (PrevFrameId == 0);
+
+		// Note: FrameId should be filled in by our framework before this point, so bIsFirstFrame should be valid
+		if (CurrentTimecodeFrame != PreviousTimecodeFrame && bIsFirstFrame)
+		{
+			ReferenceFrameId.store(FrameId, std::memory_order_release);
+			LastCapturedTimecodeFrameNumber.store(CurrentTimecodeFrame, std::memory_order_release);
+			RefFrameId = FrameId; // Update local copy
+
+			// Capture this reference frame's time for genlock
+			{
+				FScopeLock Lock(&GenlockSceneTimeLock);
+				GenlockTriggerSceneTime = CurrentSceneTime;
+			}
+
+			UE_LOG(LogLiveLink, Verbose, TEXT("LiveLinkCustomTimeStep::OnLiveLinkFrameDataReceived_AnyThread FrameId=%d, TimecodeFrame=%d, SceneTimeFrame=%d"),
+				FrameId, CurrentTimecodeFrame, CurrentSceneTime.Time.FloorToFrame().Value);
+		}
+	}
+
+	// Update genlock trigger time on frames aligned with the divider pattern
+
+	if (RefFrameId != INDEX_NONE && FramesPerTC > 0 && Divider > 0)
+	{
+		const int32 FrameOffset = static_cast<int32>(FrameId - RefFrameId);
+
+		// Calculate capture stride using GCD for proper phase alignment
+		const int32 CaptureStride = FMath::GreatestCommonDivisor(Divider, FramesPerTC);
+
+		if (FrameOffset % CaptureStride == 0 && FrameOffset > 0)
+		{
+			// This frame is aligned with the divider pattern, so we update genlock trigger time
+
+			FScopeLock Lock(&GenlockSceneTimeLock);
+			GenlockTriggerSceneTime = CurrentSceneTime;
+		}
+	}
+
 	// Trigger the event to release WaitForSync.
 	// If the engine cannot keep up, it means that SyncCount will be greater than FrameRateDividerAnyThread and delta time will be multiplied accordingly
 	if (SyncCount >= FrameRateDividerAnyThread)
@@ -314,10 +487,21 @@ void ULiveLinkCustomTimeStep::OnLiveLinkFrameDataReceived_AnyThread(const FLiveL
 
 		// In lockstep mode, the engine can stall the live link producer until it has had a chance to process its current data.
 		// We can't block if this function is being called in the game thread because it would deadlock.
-		if (bLockStepModeAnyThread && !IsInGameThread())
+		if (bLockStepModeAnyThread && !IsInGameThread() && !bSkipLockStepWait.load(std::memory_order_acquire))
 		{
 			TRACE_CPUPROFILER_EVENT_SCOPE(ULiveLinkCustomTimeStep::EventWaitForSync->Wait);
-			EventWaitForSync->Wait(FTimespan((TimeoutInSecondsAnyThread / 2) * ETimespan::TicksPerSecond));
+
+			// Wait until the engine has consumed SyncCount.
+
+			const FTimespan LockStepTimeout((TimeoutInSecondsAnyThread / 2) * ETimespan::TicksPerSecond);
+
+			while (SyncCount.load(std::memory_order_acquire) >= FrameRateDividerAnyThread.load(std::memory_order_relaxed))
+			{
+				if (!EventWaitForSync->Wait(LockStepTimeout) || bSkipLockStepWait.load(std::memory_order_acquire))
+				{
+					break;
+				}
+			}
 		}
 	}
 }
diff --git a/Engine/Plugins/Animation/LiveLink/Source/LiveLink/Private/LiveLinkMessageBusSource.cpp b/Engine/Plugins/Animation/LiveLink/Source/LiveLink/Private/LiveLinkMessageBusSource.cpp
--- a/Engine/Plugins/Animation/LiveLink/Source/LiveLink/Private/LiveLinkMessageBusSource.cpp
+++ b/Engine/Plugins/Animation/LiveLink/Source/LiveLink/Private/LiveLinkMessageBusSource.cpp
@@ -6,12 +6,18 @@
 #include "ILiveLinkClient.h"
 #include "ILiveLinkModule.h"
 #include "INetworkMessagingExtension.h"
+#include "LiveLinkClient.h"
 #include "LiveLinkCompression.h"
 #include "LiveLinkHeartbeatEmitter.h"
 #include "LiveLinkLog.h"
+#include "Misc/ScopeLock.h"
+
+#include <atomic>
+
 #if WITH_LIVELINK_DISCOVERY_MANAGER_THREAD
 #include "LiveLinkMessageBusDiscoveryManager.h"
 #endif
+
 #include "LiveLinkMessageBusSourceSettings.h"
 #include "LiveLinkMessages.h"
 #include "LiveLinkRole.h"
@@ -35,6 +41,73 @@ static TAutoConsoleVariable<int32> CVarLiveLinkSupportCompressPayloads(
 	TEXT("Whether to add the annotation indicating that we support compressed animation data. Can be set to 0 to simulate that compressed payloads are not supported."),
 	ECVF_RenderThreadSafe);
 
+int32 GLiveLinkMessageBusEnableResequencing = 1;
+static FAutoConsoleVariableRef CVarLiveLinkMessageBusEnableResequencing(
+	TEXT("LiveLink.MessageBus.EnableResequencing"),
+	GLiveLinkMessageBusEnableResequencing,
+	TEXT("Enable frame resequencing to handle async delivery from UDP Message Bus."),
+	ECVF_Default);
+
+namespace UE::LiveLink::Private
+{
+	/**
+	 * Returns true if FrameA comes before FrameB, handling int32 wraparound.
+	 */
+	FORCEINLINE bool IsFrameIdEarlier(int32 FrameA, int32 FrameB)
+	{
+		return int32(uint32(FrameA) - uint32(FrameB)) < 0;
+	}
+
+	/**
+	 * Returns true if TestFrameId is the next sequential frame after BaseFrameId, handling wraparound.
+	 */
+	FORCEINLINE bool IsNextFrameId(int32 TestFrameId, int32 BaseFrameId)
+	{
+		return TestFrameId == int32(uint32(BaseFrameId) + 1u);
+	}
+
+	/**
+	 * Returns true if the difference between two FrameIds is suspiciously large, indicating wraparound, restart, or extended pause.
+	 */
+	FORCEINLINE bool ShouldResyncFrameId(int32 IncomingFrameId, int32 LastDeliveredFrameId, int32 MaxBufferSize)
+	{
+		const int32 Diff = int32(uint32(IncomingFrameId) - uint32(LastDeliveredFrameId));
+		const uint32 AbsDiff = (Diff >= 0) ? uint32(Diff) : (0u - uint32(Diff));
+
+		// 2x buffer size as headroom, and a minimum in case the buffer is small
+
+		const int32 Threshold = FMath::Max(MaxBufferSize * 2, 10);
+
+		return AbsDiff > uint32(Threshold);
+	}
+
+	/**
+	 * Returns true if frame A should be delivered before frame B based on FrameId.
+	 */
+	bool IsFrameEarlierThan(const FLiveLinkFrameDataStruct& A, const FLiveLinkFrameDataStruct& B)
+	{
+		const FLiveLinkBaseFrameData* DataA = A.GetBaseData();
+		const FLiveLinkBaseFrameData* DataB = B.GetBaseData();
+
+		if (!DataA || !DataB)
+		{
+			return false;
+		}
+
+		// Sort by FrameId when available
+
+		const bool bHasFrameIdA = DataA->FrameId != INDEX_NONE;
+		const bool bHasFrameIdB = DataB->FrameId != INDEX_NONE;
+
+		if (bHasFrameIdA && bHasFrameIdB)
+		{
+			return UE::LiveLink::Private::IsFrameIdEarlier(DataA->FrameId, DataB->FrameId);
+		}
+
+		// If FrameIds are not available, keep reception order
+		return false;
+	}
+}
 
 FText FLiveLinkMessageBusSource::ValidSourceStatus()
 {
@@ -65,6 +138,12 @@ FLiveLinkMessageBusSource::FLiveLinkMessageBusSource(const FText& InSourceType,
 
 void FLiveLinkMessageBusSource::InitializeSettings(ULiveLinkSourceSettings* Settings)
 {
+	if (const ULiveLinkMessageBusSourceSettings* MessageBusSettings = Cast<ULiveLinkMessageBusSourceSettings>(Settings))
+	{
+		bCachedEnableFrameResequencing.store(MessageBusSettings->bEnableFrameResequencing, std::memory_order_relaxed);
+		CachedResequencerMaxBufferSize.store(MessageBusSettings->ResequencerMaxBufferSize, std::memory_order_relaxed);
+		CachedResequencerMaxWaitTimeSeconds.store(MessageBusSettings->ResequencerMaxWaitTimeSeconds, std::memory_order_relaxed);
+	}
 }
 
 void FLiveLinkMessageBusSource::ReceiveClient(ILiveLinkClient* InClient, FGuid InSourceGuid)
@@ -376,9 +455,105 @@ void FLiveLinkMessageBusSource::InitializeAndPushFrameData_AnyThread(FName Subje
 	{
 		static const FName NAME_InvalidFrameData = "LiveLinkMessageBusSource_InvalidFrameData";
 		FLiveLinkLog::ErrorOnce(NAME_InvalidFrameData, SubjectKey, TEXT("Invalid frame data was provided for '%s' with connection '%s'"), *SubjectName.ToString(), *GetSourceMachineName().ToString());
+		return;
 	}
 
-	PushClientSubjectFrameData_AnyThread(SubjectKey, MoveTemp(DataStruct));
+	// Push frame (sorted or not depending on the user setting)
+
+	if (bCachedEnableFrameResequencing.load(std::memory_order_relaxed) && GLiveLinkMessageBusEnableResequencing)
+	{
+		ResequenceAndPushFrame(SubjectKey, MoveTemp(DataStruct));
+	}
+	else
+	{
+		PushClientSubjectFrameData_AnyThread(SubjectKey, MoveTemp(DataStruct));
+	}
+}
+
+void FLiveLinkMessageBusSource::ResequenceAndPushFrame(const FLiveLinkSubjectKey& SubjectKey, FLiveLinkFrameDataStruct&& FrameData)
+{
+	FScopeLock Lock(&ResequencerCriticalSection);
+
+	FFrameResequencer& Resequencer = SubjectResequencers.FindOrAdd(SubjectKey.SubjectName);
+	const double CurrentTime = FPlatformTime::Seconds();
+
+	// Update resequencer configuration from cached settings
+	Resequencer.MaxBufferSize = CachedResequencerMaxBufferSize.load(std::memory_order_relaxed);
+	Resequencer.MaxWaitTimeSeconds = CachedResequencerMaxWaitTimeSeconds.load(std::memory_order_relaxed);
+
+	// Check if this is the next expected frame, and in that case deliver it immediately
+
+	const FLiveLinkBaseFrameData* IncomingData = FrameData.GetBaseData();
+
+	if (IncomingData && IncomingData->FrameId != INDEX_NONE)
+	{
+		const bool bIsFirstFrame = !Resequencer.LastDeliveredFrameId.IsSet();
+		const bool bIsNextExpectedFrame = Resequencer.LastDeliveredFrameId.IsSet() &&
+			UE::LiveLink::Private::IsNextFrameId(IncomingData->FrameId, Resequencer.LastDeliveredFrameId.GetValue());
+
+		// Check if we need to resync the FrameId expectation
+		const bool bShouldResync = Resequencer.LastDeliveredFrameId.IsSet() &&
+			UE::LiveLink::Private::ShouldResyncFrameId(IncomingData->FrameId, Resequencer.LastDeliveredFrameId.GetValue(), Resequencer.MaxBufferSize);
+
+		if (bIsFirstFrame || bIsNextExpectedFrame || bShouldResync)
+		{
+			const FQualifiedFrameTime& SceneTime = IncomingData->MetaData.SceneTime;
+
+			if (bShouldResync)
+			{
+				UE_LOG(LogLiveLink, Warning, TEXT("Resequencer [%s] RESYNC: FrameId=%d (was %d), SceneTime=%d+%.4f - Flushing %d buffered frames"),
+					*SubjectKey.SubjectName.ToString(),
+					IncomingData->FrameId,
+					Resequencer.LastDeliveredFrameId.GetValue(),
+					SceneTime.Time.GetFrame().Value,
+					SceneTime.Time.GetSubFrame(),
+					Resequencer.BufferedFrames.Num());
+
+				// Flush old buffered frames on resync
+				Resequencer.BufferedFrames.Empty();
+				Resequencer.OldestFrameBufferTime = 0.0;
+			}
+
+			// Update tracking before delivering
+			Resequencer.UpdateLastDelivered(FrameData);
+
+			// Deliver immediately
+			PushClientSubjectFrameData_AnyThread(SubjectKey, MoveTemp(FrameData));
+
+			// Try to deliver any buffered frames that are now ready
+
+			FLiveLinkFrameDataStruct ReadyFrame;
+
+			while (Resequencer.TryPopNextFrame(ReadyFrame, CurrentTime))
+			{
+				const FLiveLinkBaseFrameData* BufferedData = ReadyFrame.GetBaseData();
+				PushClientSubjectFrameData_AnyThread(SubjectKey, MoveTemp(ReadyFrame));
+			}
+			return;
+		}
+	}
+
+	// Frame is out of order, so we buffer it while we wait for the right one.
+
+	if (IncomingData)
+	{
+		const FQualifiedFrameTime& SceneTime = IncomingData->MetaData.SceneTime;
+	}
+
+	Resequencer.AddFrame(MoveTemp(FrameData), CurrentTime);
+
+	// Track peak buffer depth
+	Resequencer.PeakBufferDepth = FMath::Max(Resequencer.PeakBufferDepth, Resequencer.BufferedFrames.Num());
+
+	// Try to deliver frames from buffer in order
+
+	FLiveLinkFrameDataStruct ReadyFrame;
+
+	while (Resequencer.TryPopNextFrame(ReadyFrame, CurrentTime))
+	{
+		const FLiveLinkBaseFrameData* BufferedData = ReadyFrame.GetBaseData();
+		PushClientSubjectFrameData_AnyThread(SubjectKey, MoveTemp(ReadyFrame));
+	}
 }
 
 double FLiveLinkMessageBusSource::GetDeadSourceTimeout() const
@@ -411,6 +586,25 @@ void FLiveLinkMessageBusSource::HandleClearSubject(const FLiveLinkClearSubject&
 	if (!Message.SubjectName.IsNone())
 	{
 		const FLiveLinkSubjectKey SubjectKey(SourceGuid, Message.SubjectName);
+
+		// Flush any buffered frames for this subject before removing
+		{
+			FScopeLock Lock(&ResequencerCriticalSection);
+
+			if (FFrameResequencer* Resequencer = SubjectResequencers.Find(Message.SubjectName))
+			{
+				TArray<FLiveLinkFrameDataStruct> RemainingFrames;
+				Resequencer->Flush(RemainingFrames);
+
+				for (FLiveLinkFrameDataStruct& Frame : RemainingFrames)
+				{
+					PushClientSubjectFrameData_AnyThread(SubjectKey, MoveTemp(Frame));
+				}
+
+				SubjectResequencers.Remove(Message.SubjectName);
+			}
+		}
+
 		Client->RemoveSubject_AnyThread(SubjectKey);
 	}
 }
@@ -483,6 +677,12 @@ bool FLiveLinkMessageBusSource::RequestSourceShutdown()
 
 	RoleInstances.Empty();
 
+	// Clean up resequencer
+	{
+		FScopeLock Lock(&ResequencerCriticalSection);
+		SubjectResequencers.Empty();
+	}
+
 	bInitialized = false;
 	bIsShuttingDown = true;
 
@@ -506,6 +706,179 @@ PRAGMA_DISABLE_DEPRECATION_WARNINGS
 PRAGMA_ENABLE_DEPRECATION_WARNINGS
 }
 
+void FLiveLinkMessageBusSource::FFrameResequencer::AddFrame(FLiveLinkFrameDataStruct&& Frame, double CurrentTime)
+{
+	// Check if frame is too late (already delivered a later frame)
+	if (IsFrameTooLate(Frame))
+	{
+		DroppedLateFrameCount++;
+		UE_LOG(LogLiveLink, Verbose, TEXT("Dropped late frame (already delivered later frame)"));
+		return;
+	}
+
+	// Add to buffer
+	FFrameResequencer::FBufferedFrame& NewEntry = BufferedFrames.AddDefaulted_GetRef();
+	NewEntry.Frame = MoveTemp(Frame);
+	NewEntry.BufferedAtSeconds = CurrentTime;
+
+	// Sort by ordering tier
+	BufferedFrames.StableSort([](const FFrameResequencer::FBufferedFrame& A, const FFrameResequencer::FBufferedFrame& B)
+	{
+		return UE::LiveLink::Private::IsFrameEarlierThan(A.Frame, B.Frame);
+	});
+
+	// Update oldest frame buffer time if this is the first frame
+	if (BufferedFrames.Num() == 1)
+	{
+		OldestFrameBufferTime = CurrentTime;
+	}
+}
+
+bool FLiveLinkMessageBusSource::FFrameResequencer::TryPopNextFrame(FLiveLinkFrameDataStruct& OutFrame, double CurrentTime)
+{
+	if (BufferedFrames.Num() == 0)
+	{
+		return false;
+	}
+
+	auto PopOldestByBufferedTime = [&]()
+	{
+		int32 OldestIndex = 0;
+		double OldestTime = BufferedFrames[0].BufferedAtSeconds;
+		for (int32 Index = 1; Index < BufferedFrames.Num(); ++Index)
+		{
+			if (BufferedFrames[Index].BufferedAtSeconds < OldestTime)
+			{
+				OldestTime = BufferedFrames[Index].BufferedAtSeconds;
+				OldestIndex = Index;
+			}
+		}
+		OutFrame = MoveTemp(BufferedFrames[OldestIndex].Frame);
+		BufferedFrames.RemoveAt(OldestIndex);
+	};
+
+	auto UpdateOldestBufferTime = [&]()
+	{
+		if (BufferedFrames.Num() > 0)
+		{
+			double NewOldest = BufferedFrames[0].BufferedAtSeconds;
+			for (int32 Index = 1; Index < BufferedFrames.Num(); ++Index)
+			{
+				NewOldest = FMath::Min(NewOldest, BufferedFrames[Index].BufferedAtSeconds);
+			}
+			OldestFrameBufferTime = NewOldest;
+		}
+		else
+		{
+			OldestFrameBufferTime = 0.0;
+		}
+	};
+
+	// Check for buffer overflow
+	if (BufferedFrames.Num() >= MaxBufferSize)
+	{
+		UE_LOG(LogLiveLink, Verbose, TEXT("Resequencer buffer overflow (%d frames), flushing oldest"), BufferedFrames.Num());
+
+		PopOldestByBufferedTime();
+
+		UpdateLastDelivered(OutFrame);
+
+		UpdateOldestBufferTime();
+
+		ReorderedFrameCount++;
+
+		return true;
+	}
+
+	// Check for timeout
+	if (CurrentTime - OldestFrameBufferTime > MaxWaitTimeSeconds)
+	{
+		UE_LOG(LogLiveLink, Verbose, TEXT("Resequencer timeout (%.3fs), flushing %d frames"), CurrentTime - OldestFrameBufferTime, BufferedFrames.Num());
+
+		PopOldestByBufferedTime();
+
+		UpdateLastDelivered(OutFrame);
+
+		UpdateOldestBufferTime();
+
+		return true;
+	}
+
+	// Deliver buffered frame only if it's the next expected frame
+	// Rely on timeout/overflow handling above for frames that are truly lost
+	if (BufferedFrames.Num() > 0)
+	{
+		const FLiveLinkBaseFrameData* OldestData = BufferedFrames[0].Frame.GetBaseData();
+		if (OldestData && OldestData->FrameId != INDEX_NONE)
+		{
+			const bool bIsNextExpected = LastDeliveredFrameId.IsSet() &&
+				UE::LiveLink::Private::IsNextFrameId(OldestData->FrameId, LastDeliveredFrameId.GetValue());
+
+			if (bIsNextExpected)
+			{
+				OutFrame = MoveTemp(BufferedFrames[0].Frame);
+				BufferedFrames.RemoveAt(0);
+				UpdateLastDelivered(OutFrame);
+				UpdateOldestBufferTime();
+				ReorderedFrameCount++;
+				return true;
+			}
+		}
+	}
+
+	return false;
+}
+
+void FLiveLinkMessageBusSource::FFrameResequencer::Flush(TArray<FLiveLinkFrameDataStruct>& OutFrames)
+{
+	OutFrames.Empty(BufferedFrames.Num());
+	for (FFrameResequencer::FBufferedFrame& BufferedFrame : BufferedFrames)
+	{
+		OutFrames.Add(MoveTemp(BufferedFrame.Frame));
+	}
+	BufferedFrames.Empty();
+	OldestFrameBufferTime = 0.0;
+}
+
+bool FLiveLinkMessageBusSource::FFrameResequencer::IsFrameTooLate(const FLiveLinkFrameDataStruct& Frame) const
+{
+	const FLiveLinkBaseFrameData* Data = Frame.GetBaseData();
+	if (!Data)
+	{
+		return false;
+	}
+
+	// Check if frame arrived too late based on FrameId (with wraparound support)
+	if (LastDeliveredFrameId.IsSet() && Data->FrameId != INDEX_NONE)
+	{
+		// If gap is suspiciously large (>2x buffer), allow resync instead of dropping
+		// This handles wraparound recovery (e.g., LastId=INT_MAX-5, IncomingId=1)
+		if (UE::LiveLink::Private::ShouldResyncFrameId(Data->FrameId, LastDeliveredFrameId.GetValue(), MaxBufferSize))
+		{
+			return false; // Don't drop - allow this frame to trigger resync
+		}
+
+		return UE::LiveLink::Private::IsFrameIdEarlier(Data->FrameId, LastDeliveredFrameId.GetValue());
+	}
+
+	return false;
+}
+
+void FLiveLinkMessageBusSource::FFrameResequencer::UpdateLastDelivered(const FLiveLinkFrameDataStruct& Frame)
+{
+	const FLiveLinkBaseFrameData* Data = Frame.GetBaseData();
+	if (!Data)
+	{
+		return;
+	}
+
+	// Update last delivered FrameId
+	if (Data->FrameId != INDEX_NONE)
+	{
+		LastDeliveredFrameId = Data->FrameId;
+	}
+}
+
 void FLiveLinkMessageBusSource::AddAnnotations(TMap<FName, FString>& InOutAnnotations) const
 {
 	if (CVarLiveLinkSupportCompressPayloads.GetValueOnAnyThread())
diff --git a/Engine/Plugins/Animation/LiveLink/Source/LiveLink/Private/LiveLinkSubject.cpp b/Engine/Plugins/Animation/LiveLink/Source/LiveLink/Private/LiveLinkSubject.cpp
--- a/Engine/Plugins/Animation/LiveLink/Source/LiveLink/Private/LiveLinkSubject.cpp
+++ b/Engine/Plugins/Animation/LiveLink/Source/LiveLink/Private/LiveLinkSubject.cpp
@@ -463,23 +463,38 @@ void FLiveLinkSubject::AddFrameData(FLiveLinkFrameDataStruct&& InFrameData)
 			PreProcessor->PreProcessFrame(GetStaticData(), InFrameData);
 		}
 
-		// Assign identifier to incoming frame and insert it where it belongs
-		const FLiveLinkFrameIdentifier ThisFrameIdentifier = NextIdentifier++;
-		ReceivedOrderedFrames.Enqueue(ThisFrameIdentifier);
-		InFrameData.GetBaseData()->FrameId = ThisFrameIdentifier;
-		
-		// Extra verbose logging about new frame being added
-		const FTimecode FrameTC = FTimecode::FromFrameNumber(InFrameData.GetBaseData()->MetaData.SceneTime.Time.GetFrame(), InFrameData.GetBaseData()->MetaData.SceneTime.Rate);
-		UE_LOG(LogLiveLink, Verbose, TEXT("Subject '%s' adding frame ID %d at index %d - Timecode:[%s.%0.3f] - SourceTime: %0.4f, Offset: %0.6f, CorrectedTime: %0.4f"), *SubjectKey.SubjectName.ToString(), ThisFrameIdentifier, FrameIndex, *FrameTC.ToString(), InFrameData.GetBaseData()->MetaData.SceneTime.Time.GetSubFrame(), InFrameData.GetBaseData()->WorldTime.GetSourceTime(), InFrameData.GetBaseData()->WorldTime.GetOffset(), InFrameData.GetBaseData()->WorldTime.GetOffsettedTime());
-		
-		FrameData.Insert(MoveTemp(InFrameData), FrameIndex);
+		// FrameId should be populated by LiveLinkClient framework before reaching here
+		if (FLiveLinkBaseFrameData* BaseData = InFrameData.GetBaseData())
+		{
+			if (BaseData->FrameId == INDEX_NONE)
+			{
+				UE_LOG(LogLiveLink, 
+					Error, 
+					TEXT("FLiveLinkSubject::AddFrameData: Subject '%s': FrameId was unexpectedly not set at this point in the pipeline, skipping it."), 
+					*SubjectKey.SubjectName.ToString()
+				);
 
-		if (CachedSettings.BufferSettings.bGenerateSubFrame && CachedSettings.SourceMode == ELiveLinkSourceMode::Timecode)
+				IncreaseFrameDroppedStat();
+
+				FrameIndex = -1; // Prevents insertion into FrameData.
+			}
+			else
+			{
+				ReceivedOrderedFrames.Enqueue(BaseData->FrameId);
+			}
+		}
+
+		if (FrameIndex >= 0)
 		{
-			AdjustSubFrame_SceneTime(FrameIndex);
+			FrameData.Insert(MoveTemp(InFrameData), FrameIndex);
+
+			if (CachedSettings.BufferSettings.bGenerateSubFrame && CachedSettings.SourceMode == ELiveLinkSourceMode::Timecode)
+			{
+				AdjustSubFrame_SceneTime(FrameIndex);
+			}
 		}
 	}
-	else 
+	else
 	{
 		// For some reason, the frame can't be added and is just discarded. Logs will have more info about the why
 		IncreaseFrameDroppedStat();
diff --git a/Engine/Plugins/Animation/LiveLink/Source/LiveLink/Private/LiveLinkTimecodeProvider.cpp b/Engine/Plugins/Animation/LiveLink/Source/LiveLink/Private/LiveLinkTimecodeProvider.cpp
--- a/Engine/Plugins/Animation/LiveLink/Source/LiveLink/Private/LiveLinkTimecodeProvider.cpp
+++ b/Engine/Plugins/Animation/LiveLink/Source/LiveLink/Private/LiveLinkTimecodeProvider.cpp
@@ -2,14 +2,32 @@
 
 #include "LiveLinkTimecodeProvider.h"
 
+#include "Engine/Engine.h"
 #include "Features/IModularFeatures.h"
 #include "ILiveLinkClient.h"
+#include "LiveLinkClient.h"
+#include "LiveLinkCustomTimeStep.h"
+#include "LiveLinkTypes.h"
 
 #include UE_INLINE_GENERATED_CPP_BY_NAME(LiveLinkTimecodeProvider)
 
 
+namespace UE::LiveLink::Private
+{
+static const ULiveLinkCustomTimeStep* GetLiveLinkGenlock()
+{
+	if (!GEngine)
+	{
+		return nullptr;
+	}
+
+	return Cast<ULiveLinkCustomTimeStep>(GEngine->GetCustomTimeStep());
+}
+}
+
+
 ULiveLinkTimecodeProvider::ULiveLinkTimecodeProvider()
-	: Evaluation(ELiveLinkTimecodeProviderEvaluationType::Lerp)
+	: Evaluation(ELiveLinkTimecodeProviderEvaluationType::Latest)
 	, bOverrideFrameRate(false)
 	, OverrideFrameRate(24, 1)
 	, BufferSize(4)
@@ -25,6 +43,19 @@ FQualifiedFrameTime ULiveLinkTimecodeProvider::GetQualifiedFrameTime() const
 	{
 		if (Evaluation == ELiveLinkTimecodeProviderEvaluationType::Latest || SubjectFrameTimes.Num() == 1)
 		{
+			// When using Latest mode, prefer genlocked timecode unless sourcing from sequencer
+			const bool bIsSequencer = bIsSequencerSourceAnyThread.load(std::memory_order_relaxed);
+
+			if (Evaluation == ELiveLinkTimecodeProviderEvaluationType::Latest && !bIsSequencer)
+			{
+				if (const ULiveLinkCustomTimeStep* LiveLinkCTS = UE::LiveLink::Private::GetLiveLinkGenlock())
+				{
+					const FQualifiedFrameTime GenlockTime = LiveLinkCTS->GetGenlockTriggerSceneTime();
+					return ConvertTo(GenlockTime);
+				}
+			}
+
+			// Fall back to buffered latest frame
 			return ConvertTo(SubjectFrameTimes.Last().SceneTime);
 		}
 		else
@@ -196,6 +227,11 @@ void ULiveLinkTimecodeProvider::UninitClient()
 
 void ULiveLinkTimecodeProvider::RegisterSubject()
 {
+	if (!LiveLinkClient)
+	{
+		return;
+	}
+
 	{
 		FScopeLock ScopeLock(&SubjectFrameLock);
 		SubjectFrameTimes.Reset();
@@ -203,6 +239,12 @@ void ULiveLinkTimecodeProvider::RegisterSubject()
 
 	RegisteredSubjectKey = SubjectKey;
 
+	// Check if this is a sequencer source and cache the result
+	FText SourceType = LiveLinkClient->GetSourceType(SubjectKey.Source);
+	FString SourceTypeStr = SourceType.ToString();
+	bool bIsSequencer = SourceTypeStr.Contains(UE::LiveLink::SourceTypes::SequencerSourcePrefix, ESearchCase::IgnoreCase);
+	bIsSequencerSourceAnyThread.store(bIsSequencer, std::memory_order_relaxed);
+
 	FDelegateHandle Tmp;
 	LiveLinkClient->RegisterForFrameDataReceived(
 		RegisteredSubjectKey
@@ -215,11 +257,25 @@ void ULiveLinkTimecodeProvider::RegisterSubject()
 
 void ULiveLinkTimecodeProvider::UnregisterSubject()
 {
+	if (!LiveLinkClient)
+	{
+		return;
+	}
+
 	if (RegisterForFrameDataReceivedHandle.IsValid())
 	{
 		LiveLinkClient->UnregisterForFrameDataReceived(RegisteredSubjectKey, FDelegateHandle(), RegisterForFrameDataReceivedHandle);
 		RegisterForFrameDataReceivedHandle.Reset();
 	}
+
+	// Also unregister from external subject if we registered to it
+	if (ExternalSubjectFrameDataReceivedHandle.IsValid() && OriginalExternalSubjectKey.Source.IsValid())
+	{
+		LiveLinkClient->UnregisterForFrameDataReceived(OriginalExternalSubjectKey, FDelegateHandle(), ExternalSubjectFrameDataReceivedHandle);
+		ExternalSubjectFrameDataReceivedHandle.Reset();
+	}
+
+	bIsSequencerSourceAnyThread.store(false, std::memory_order_relaxed);
 	RegisteredSubjectKey = FLiveLinkSubjectKey();
 
 	{
@@ -250,10 +306,65 @@ void ULiveLinkTimecodeProvider::OnLiveLinkClientUnregistered(const FName& Type,
 
 void ULiveLinkTimecodeProvider::OnLiveLinkSubjectAdded(FLiveLinkSubjectKey InSubjectKey)
 {
-	if (InSubjectKey.SubjectName == SubjectKey.SubjectName)
+	if (!LiveLinkClient || InSubjectKey.SubjectName != SubjectKey.SubjectName)
+	{
+		return;
+	}
+
+	// Check if already registered to a subject with same name - unregister first
+	if (RegisteredSubjectKey.SubjectName == InSubjectKey.SubjectName && RegisteredSubjectKey.Source.IsValid())
+	{
+		UnregisterSubject();
+	}
+
+	// Check if this is a sequencer source before registering
+	FText SourceType = LiveLinkClient->GetSourceType(InSubjectKey.Source);
+	FString SourceTypeStr = SourceType.ToString();
+	bool bIsSequencer = SourceTypeStr.Contains(UE::LiveLink::SourceTypes::SequencerSourcePrefix, ESearchCase::IgnoreCase);
+
+	if (bIsSequencer)
+	{
+		// Store the original external key if this is the first sequencer takeover
+		if (!OriginalExternalSubjectKey.Source.IsValid())
+		{
+			OriginalExternalSubjectKey = SubjectKey;
+		}
+
+		// Register to the sequencer source without updating SubjectKey (keeps UI showing external source)
+		{
+			FScopeLock ScopeLock(&SubjectFrameLock);
+			SubjectFrameTimes.Reset();
+		}
+
+		RegisteredSubjectKey = InSubjectKey;
+		bIsSequencerSourceAnyThread.store(true, std::memory_order_relaxed);
+
+		FDelegateHandle Tmp;
+		LiveLinkClient->RegisterForFrameDataReceived(
+			RegisteredSubjectKey,
+			FOnLiveLinkSubjectStaticDataReceived::FDelegate(),
+			FOnLiveLinkSubjectFrameDataReceived::FDelegate::CreateUObject(this, &ULiveLinkTimecodeProvider::OnLiveLinkFrameDataReceived_AnyThread),
+			Tmp,
+			RegisterForFrameDataReceivedHandle);
+
+		// Also register to the disabled external subject to read its live timecode
+		if (OriginalExternalSubjectKey.Source.IsValid() && LiveLinkClient->GetSubjectSettings(OriginalExternalSubjectKey))
+		{
+			FDelegateHandle TmpExternal;
+			LiveLinkClient->RegisterForFrameDataReceived(
+				OriginalExternalSubjectKey,
+				FOnLiveLinkSubjectStaticDataReceived::FDelegate(),
+				FOnLiveLinkSubjectFrameDataReceived::FDelegate::CreateUObject(this, &ULiveLinkTimecodeProvider::OnLiveLinkFrameDataReceived_AnyThread),
+				TmpExternal,
+				ExternalSubjectFrameDataReceivedHandle);
+
+			UE_LOG(LogLiveLink, Log, TEXT("LiveLinkTimecodeProvider: Registered to both sequencer and external subject for timecode"));
+		}
+	}
+	else
 	{
-		// This provider might have been initialized with an invalid key because the subject wasn't live at the time of creation. 
-		// So we update the SubjectKey to match the new subject since it's now been added.
+		// External source - update SubjectKey and register normally
+		OriginalExternalSubjectKey = InSubjectKey;
 		SubjectKey = InSubjectKey;
 		RegisterSubject();
 	}
@@ -264,8 +375,25 @@ void ULiveLinkTimecodeProvider::OnLiveLinkSubjectRemoved(FLiveLinkSubjectKey InS
 {
 	if (InSubjectKey.SubjectName == RegisteredSubjectKey.SubjectName)
 	{
-		State = ETimecodeProviderSynchronizationState::Error;
+		const bool bWasSequencerSource = bIsSequencerSourceAnyThread.load(std::memory_order_relaxed);
 		UnregisterSubject();
+
+		// If sequencer source was removed, try to re-register to the original external subject
+		if (bWasSequencerSource && LiveLinkClient && OriginalExternalSubjectKey.Source.IsValid()
+			&& LiveLinkClient->GetSubjectSettings(OriginalExternalSubjectKey))
+		{
+			SubjectKey = OriginalExternalSubjectKey;
+			RegisterSubject();
+
+			if (State == ETimecodeProviderSynchronizationState::Error)
+			{
+				State = ETimecodeProviderSynchronizationState::Synchronizing;
+			}
+		}
+		else
+		{
+			State = ETimecodeProviderSynchronizationState::Error;
+		}
 	}
 }
 
diff --git a/Engine/Plugins/Animation/LiveLink/Source/LiveLink/Public/LiveLinkClient.h b/Engine/Plugins/Animation/LiveLink/Source/LiveLink/Public/LiveLinkClient.h
--- a/Engine/Plugins/Animation/LiveLink/Source/LiveLink/Public/LiveLinkClient.h
+++ b/Engine/Plugins/Animation/LiveLink/Source/LiveLink/Public/LiveLinkClient.h
@@ -286,6 +286,12 @@ protected:
 	/** Lock to protect access on SubjectFrameToPush and SubjectStaticToPush. */
 	mutable FCriticalSection PendingFramesCriticalSection;
 
+	/** FrameId counters used when sources don't set it themselves. */
+	TMap<FLiveLinkSubjectKey, FLiveLinkFrameIdentifier> SubjectFrameCounters;
+
+	/** Lock to protect access to SubjectFrameCounters. */
+	mutable FCriticalSection SubjectFrameCountersLock;
+
 private:
 	/** Pending static info to add to a subject. */
 	TArray<FPendingSubjectStatic> SubjectStaticToPush;
diff --git a/Engine/Plugins/Animation/LiveLink/Source/LiveLink/Public/LiveLinkCustomTimeStep.h b/Engine/Plugins/Animation/LiveLink/Source/LiveLink/Public/LiveLinkCustomTimeStep.h
--- a/Engine/Plugins/Animation/LiveLink/Source/LiveLink/Public/LiveLinkCustomTimeStep.h
+++ b/Engine/Plugins/Animation/LiveLink/Source/LiveLink/Public/LiveLinkCustomTimeStep.h
@@ -7,6 +7,7 @@
 #include "HAL/Event.h"
 #include "LiveLinkTypes.h"
 #include "Misc/FrameRate.h"
+#include "Misc/QualifiedFrameTime.h"
 
 #include <atomic>
 
@@ -53,6 +54,9 @@ public:
 	LIVELINK_API virtual bool WaitForSync() override;
 	LIVELINK_API virtual uint32 GetExpectedSyncCountDelta() const override;
 
+	/** Returns the SceneTime captured at the most recent genlock trigger point. */
+	LIVELINK_API FQualifiedFrameTime GetGenlockTriggerSceneTime() const;
+
 	//~ Begin UObject Interface
 	LIVELINK_API virtual void BeginDestroy() override;
 #if WITH_EDITOR
@@ -153,4 +157,27 @@ private:
 
 	/** Shadow of TimeoutInSeconds property */
 	std::atomic<double> TimeoutInSecondsAnyThread = 1.0;
+
+	/** SceneTime captured at the most recent genlock trigger point. */
+	FQualifiedFrameTime GenlockTriggerSceneTime;
+
+	/** Protects GenlockTriggerSceneTime across threads. */
+	mutable FCriticalSection GenlockSceneTimeLock;
+
+	/** Frame number of the last captured genlock trigger time, used to ensure deterministic subframe behavior */
+	std::atomic<int32> LastCapturedTimecodeFrameNumber = INDEX_NONE;
+
+	/** FrameId of the reference frame, used for phase alignment */
+	std::atomic<int32> ReferenceFrameId = INDEX_NONE;
+
+	/** Cached frames per timecode period (SourceRate / TimecodeRate), calculated at initialization */
+	std::atomic<int32> FramesPerTimecodeAnyThread = 0;
+
+	/** Cached flag indicating whether the monitored subject is from a sequencer source */
+	std::atomic<bool> bIsSequencerSourceAnyThread = false;
+
+	/** When true, OnLiveLinkFrameDataReceived_AnyThread skips the lockstep wait.
+	 *  This is set before unregistering to avoid blocking the game thread on the
+	 *  critical section held during BroadcastFrameDataUpdate's lockstep wait. */
+	std::atomic<bool> bSkipLockStepWait = false;
 };
diff --git a/Engine/Plugins/Animation/LiveLink/Source/LiveLink/Public/LiveLinkMessageBusSource.h b/Engine/Plugins/Animation/LiveLink/Source/LiveLink/Public/LiveLinkMessageBusSource.h
--- a/Engine/Plugins/Animation/LiveLink/Source/LiveLink/Public/LiveLinkMessageBusSource.h
+++ b/Engine/Plugins/Animation/LiveLink/Source/LiveLink/Public/LiveLinkMessageBusSource.h
@@ -5,6 +5,7 @@
 #include "ILiveLinkSource.h"
 
 #include "HAL/ThreadSafeBool.h"
+#include "LiveLinkTypes.h"
 #include "MessageEndpoint.h"
 
 #define UE_API LIVELINK_API
@@ -119,6 +120,11 @@ protected:
 	// Our identifier in LiveLink
 	FGuid SourceGuid;
 
+	// Cached settings values (atomics for thread-safe access from _AnyThread methods)
+	std::atomic<bool> bCachedEnableFrameResequencing = true;
+	std::atomic<int32> CachedResequencerMaxBufferSize = 32;
+	std::atomic<float> CachedResequencerMaxWaitTimeSeconds = 0.1f;
+
 private:
 	void CreateAndInitializeMessageEndpoint();
 
@@ -140,7 +146,69 @@ private:
 	// Initialize (can be used after a shutdown request)
 	void Initialize();
 
+	// Resequencer for maintaining frame order per subject
+	void ResequenceAndPushFrame(const FLiveLinkSubjectKey& SubjectKey, FLiveLinkFrameDataStruct&& FrameData);
+
 private:
+
+	/** Per-subject frame resequencer to deliver live link frames in FrameId order */
+	struct FFrameResequencer
+	{
+	public:
+		struct FBufferedFrame
+		{
+			FLiveLinkFrameDataStruct Frame;
+			double BufferedAtSeconds = 0.0;
+		};
+
+		/** Buffered frames waiting for delivery (only used when frames arrive out of order) */
+		TArray<FBufferedFrame> BufferedFrames;
+
+		/** Last delivered FrameId - used to detect and deliver the next expected frame immediately */
+		TOptional<FLiveLinkFrameIdentifier> LastDeliveredFrameId;
+
+		/** Timestamp when oldest frame was buffered (for timeout detection) */
+		double OldestFrameBufferTime = 0.0;
+
+		/** Maximum number of frames to buffer */
+		int32 MaxBufferSize = 32;
+
+		/** Maximum time in seconds to wait for missing frames */
+		double MaxWaitTimeSeconds = 0.1;
+
+		/** Stats: Number of frames that arrived out of order and were resequenced */
+		int32 ReorderedFrameCount = 0;
+
+		/** Stats: Number of frames dropped because they arrived too late to be delivered in order */
+		int32 DroppedLateFrameCount = 0;
+
+		/** Stats: Maximum number of frames held in the buffer simultaneously */
+		int32 PeakBufferDepth = 0;
+
+	public:
+
+		/** Add a frame to the buffer */
+		void AddFrame(FLiveLinkFrameDataStruct&& Frame, double CurrentTime);
+
+		/** Try to pop the next ordered frame from buffer */
+		bool TryPopNextFrame(FLiveLinkFrameDataStruct& OutFrame, double CurrentTime);
+
+		/** Flush all buffered frames */
+		void Flush(TArray<FLiveLinkFrameDataStruct>& OutFrames);
+
+		/** Check if a frame should be dropped as too late */
+		bool IsFrameTooLate(const FLiveLinkFrameDataStruct& Frame) const;
+
+		/** Update last delivered tracking after delivering a frame */
+		void UpdateLastDelivered(const FLiveLinkFrameDataStruct& Frame);
+	};
+
+	// Per-subject map of frame resequencers for maintaining ordered frame delivery
+	TMap<FName, FFrameResequencer> SubjectResequencers;
+
+	// Critical section for thread-safe access to SubjectResequencers
+	FCriticalSection ResequencerCriticalSection;
+
 	// List of the roles available when the bus was opened
 	TArray<TWeakObjectPtr<ULiveLinkRole>> RoleInstances;
 
diff --git a/Engine/Plugins/Animation/LiveLink/Source/LiveLink/Public/LiveLinkMessageBusSourceSettings.h b/Engine/Plugins/Animation/LiveLink/Source/LiveLink/Public/LiveLinkMessageBusSourceSettings.h
--- a/Engine/Plugins/Animation/LiveLink/Source/LiveLink/Public/LiveLinkMessageBusSourceSettings.h
+++ b/Engine/Plugins/Animation/LiveLink/Source/LiveLink/Public/LiveLinkMessageBusSourceSettings.h
@@ -15,7 +15,7 @@
  * Settings for LiveLinkMessageBusSource.
  * Used to apply default Evaluation mode from project settings when constructed
  */
-UCLASS(MinimalAPI)
+UCLASS(config=Game, MinimalAPI)
 class ULiveLinkMessageBusSourceSettings : public ULiveLinkSourceSettings
 {
 public:
@@ -23,6 +23,18 @@ public:
 
 public:
 	UE_API ULiveLinkMessageBusSourceSettings();
+
+	/** Enable frame resequencing to counteract UDP Messaging delivering them out of order. */
+	UPROPERTY(config, EditAnywhere, Category = "Resequencing")
+	bool bEnableFrameResequencing = true;
+
+	/** Maximum number of frames to buffer per subject. If exceeded, oldest frames are pushed even if out of order. */
+	UPROPERTY(config, EditAnywhere, Category = "Resequencing", meta = (EditCondition = "bEnableFrameResequencing", ClampMin = "4", ClampMax = "128"))
+	int32 ResequencerMaxBufferSize = 32;
+
+	/** Maximum time in seconds to hold frames waiting for missing frames. */
+	UPROPERTY(config, EditAnywhere, Category = "Resequencing", meta = (EditCondition = "bEnableFrameResequencing", ClampMin = "0.01", ClampMax = "10.0"))
+	float ResequencerMaxWaitTimeSeconds = 0.1f;
 };
 
 #undef UE_API
diff --git a/Engine/Plugins/Animation/LiveLink/Source/LiveLink/Public/LiveLinkTimecodeProvider.h b/Engine/Plugins/Animation/LiveLink/Source/LiveLink/Public/LiveLinkTimecodeProvider.h
--- a/Engine/Plugins/Animation/LiveLink/Source/LiveLink/Public/LiveLinkTimecodeProvider.h
+++ b/Engine/Plugins/Animation/LiveLink/Source/LiveLink/Public/LiveLinkTimecodeProvider.h
@@ -20,9 +20,11 @@ enum class ELiveLinkTimecodeProviderEvaluationType
 {
 	/** Interpolate between, or extrapolate using the 2 frames that are the closest to the current world time. */
 	Lerp,
+
 	/** Use the frame that is closest to the current world time. */
 	Nearest,
-	/** Use the newest frame that was received. */
+
+	/** Use the newest frame that was received. When LiveLink Custom Time Step is active, uses the timecode captured at the genlock trigger point. */
 	Latest,
 };
 
@@ -125,6 +127,15 @@ private:
 	mutable FCriticalSection SubjectFrameLock; // Only lock SubjectFrameTimes
 	FDelegateHandle RegisterForFrameDataReceivedHandle;
 
+	/** Cached flag indicating whether the monitored subject is from a sequencer source */
+	std::atomic<bool> bIsSequencerSourceAnyThread = false;
+
+	/** Stores the original external subject key before sequencer takes over, so we can restore it when sequencer stops */
+	FLiveLinkSubjectKey OriginalExternalSubjectKey;
+
+	/** When sequencer takes over, we also register to the disabled external subject to read its timecode */
+	FDelegateHandle ExternalSubjectFrameDataReceivedHandle;
+
 	friend struct FLiveLinkHubTimecodeSettings;
 };
 
diff --git a/Engine/Plugins/Animation/LiveLink/Source/LiveLinkMovieScene/Private/MovieScene/MovieSceneLiveLinkSource.cpp b/Engine/Plugins/Animation/LiveLink/Source/LiveLinkMovieScene/Private/MovieScene/MovieSceneLiveLinkSource.cpp
--- a/Engine/Plugins/Animation/LiveLink/Source/LiveLinkMovieScene/Private/MovieScene/MovieSceneLiveLinkSource.cpp
+++ b/Engine/Plugins/Animation/LiveLink/Source/LiveLinkMovieScene/Private/MovieScene/MovieSceneLiveLinkSource.cpp
@@ -4,6 +4,7 @@
 #include "Features/IModularFeatures.h"
 #include "ILiveLinkClient.h"
 #include "ILiveLinkSource.h"
+#include "LiveLinkTypes.h"
 
 FMovieSceneLiveLinkSource::FMovieSceneLiveLinkSource()
 	: Client(nullptr)
@@ -99,7 +100,9 @@ FText FMovieSceneLiveLinkSource::GetSourceStatus() const
 
 FText FMovieSceneLiveLinkSource::GetSourceType() const
 {
-	return FText::Format(NSLOCTEXT("FMovieSceneLiveLinkSource", "MovieSceneLiveLinkSourceType", "Sequencer Live Link ({0})"),FText::FromName(SubjectPreset.Key.SubjectName));
+	return FText::Format(NSLOCTEXT("FMovieSceneLiveLinkSource", "MovieSceneLiveLinkSourceType", "{0} ({1})"),
+		FText::FromString(UE::LiveLink::SourceTypes::SequencerSourcePrefix),
+		FText::FromName(SubjectPreset.Key.SubjectName));
 }
 
 void FMovieSceneLiveLinkSource::PublishLiveLinkStaticData(FLiveLinkStaticDataStruct& StaticData)
diff --git a/Engine/Source/Runtime/LiveLinkInterface/Public/LiveLinkTypes.h b/Engine/Source/Runtime/LiveLinkInterface/Public/LiveLinkTypes.h
--- a/Engine/Source/Runtime/LiveLinkInterface/Public/LiveLinkTypes.h
+++ b/Engine/Source/Runtime/LiveLinkInterface/Public/LiveLinkTypes.h
@@ -232,8 +232,9 @@ struct FLiveLinkBaseFrameData
 	/** Time at which this frame was received - populated once inside pipeline */
 	FLiveLinkTime ArrivalTime;
 
-	/** This frame identifier - populated once inside pipeline*/
-	FLiveLinkFrameIdentifier FrameId = INDEX_NONE;
+	/** Frame identifier for tracking sent frames and detecting dropped frames when set */
+	UPROPERTY()
+	int32 FrameId = INDEX_NONE;
 
 	/** Return the LiveLinkTime struct constructed from this frame data. */
 	FLiveLinkTime GetLiveLinkTime() const { return FLiveLinkTime(WorldTime.GetOffsettedTime(), MetaData.SceneTime); }
@@ -791,3 +792,12 @@ struct
 	}
 };
 PRAGMA_ENABLE_DEPRECATION_WARNINGS
+
+/**
+ * Namespace containing well-known Live Link source type identifiers.
+ */
+namespace UE::LiveLink::SourceTypes
+{
+	/** Identifier prefix for sequencer-driven Live Link sources */
+	inline constexpr const TCHAR* SequencerSourcePrefix = TEXT("Sequencer Live Link");
+}
diff --git a/Engine/Source/Runtime/LiveLinkMessageBusFramework/Private/LiveLinkProvider.cpp b/Engine/Source/Runtime/LiveLinkMessageBusFramework/Private/LiveLinkProvider.cpp
--- a/Engine/Source/Runtime/LiveLinkMessageBusFramework/Private/LiveLinkProvider.cpp
+++ b/Engine/Source/Runtime/LiveLinkMessageBusFramework/Private/LiveLinkProvider.cpp
@@ -253,16 +253,22 @@ void FLiveLinkProvider::SetLastSubjectFrameData(FName SubjectName, FLiveLinkFram
 void FLiveLinkProvider::ClearTrackedSubject(const FName& SubjectName)
 {
 	Subjects.Remove(SubjectName);
+
 	const int32 FrameIndex = FrameDatas.IndexOfByKey(SubjectName);
+	
 	if (FrameIndex != INDEX_NONE)
 	{
 		FrameDatas.RemoveAtSwap(FrameIndex);
 	}
+
 	const int32 StaticIndex = StaticDatas.IndexOfByKey(SubjectName);
+
 	if (StaticIndex != INDEX_NONE)
 	{
 		StaticDatas.RemoveAtSwap(StaticIndex);
 	}
+
+	SubjectFrameIds.Remove(SubjectName);
 }
 
 FLiveLinkProvider::FLiveLinkProvider(const FString& InProviderName)
@@ -434,6 +440,24 @@ bool FLiveLinkProvider::UpdateSubjectFrameData(const FName SubjectName, FLiveLin
 		return false;
 	}
 
+	// Assign a FrameId to this subject
+
+	if (FLiveLinkBaseFrameData* BaseData = FrameData.GetBaseData())
+	{
+		int32& NextFrameId = SubjectFrameIds.FindOrAdd(SubjectName, 0);
+
+		if (BaseData->FrameId < 0)
+		{
+			BaseData->FrameId = NextFrameId;
+			NextFrameId = int32(uint32(NextFrameId) + 1u);
+		}
+		else
+		{
+			// Provider manually set FrameId, so honor it and update our internal one with it.
+			NextFrameId = FMath::Max(NextFrameId, int32(uint32(BaseData->FrameId) + 1u));
+		}
+	}
+
 	ValidateConnections();
 
 	TMap<FName, FString> Annotations;
diff --git a/Engine/Source/Runtime/LiveLinkMessageBusFramework/Public/LiveLinkProviderImpl.h b/Engine/Source/Runtime/LiveLinkMessageBusFramework/Public/LiveLinkProviderImpl.h
--- a/Engine/Source/Runtime/LiveLinkMessageBusFramework/Public/LiveLinkProviderImpl.h
+++ b/Engine/Source/Runtime/LiveLinkMessageBusFramework/Public/LiveLinkProviderImpl.h
@@ -50,11 +50,16 @@ private:
 	TArray<struct FTrackedFrameData> FrameDatas;
 	TMap<FName, FTrackedSubject> Subjects;
 
+	// Per-subject FrameIds
+	TMap<FName, int32> SubjectFrameIds;
+
 	// Delegate to notify interested parties when the client sources have changed
 	FLiveLinkProviderConnectionStatusChanged OnConnectionStatusChanged;
 
 private:
-	//Message bus message handlers
+
+	// Begin message bus message handlers
+
 	// Handles ping message and sends back a pong
 	void HandlePingMessage(const FLiveLinkPingMessage& Message, const TSharedRef<class IMessageContext, ESPMode::ThreadSafe>& Context);
 
