LiveLink - Fix crash when encountering invalid serialized frame data

#Synced-CL 49594222

diff --git a/Engine/Plugins/Animation/LiveLink/Source/LiveLink/Private/LiveLinkMessageBusSource.cpp b/Engine/Plugins/Animation/LiveLink/Source/LiveLink/Private/LiveLinkMessageBusSource.cpp
--- a/Engine/Plugins/Animation/LiveLink/Source/LiveLink/Private/LiveLinkMessageBusSource.cpp
+++ b/Engine/Plugins/Animation/LiveLink/Source/LiveLink/Private/LiveLinkMessageBusSource.cpp
@@ -427,6 +427,11 @@ void FLiveLinkMessageBusSource::InitializeAndPushFrameData_AnyThread(FName Subje
 		FStructOnScope Payload;
 		SerializedMessage->GetPayload(Payload);
 
+		if (!ensure(Payload.IsValid()))
+		{
+			return;
+		}
+
 		// Special case: We want to send anim data.
 		if (Payload.GetStruct() == FLiveLinkFloatAnimationFrameData::StaticStruct())
 		{
