LiveLink - Allow transmitting VirtualSubject even if part of the subject is invalid
Hide virtual subject class picker in LiveLinkHub

#Synced-CL 50755163

diff --git a/Engine/Plugins/Animation/LiveLink/Source/LiveLink/Private/LiveLinkAnimationVirtualSubject.cpp b/Engine/Plugins/Animation/LiveLink/Source/LiveLink/Private/LiveLinkAnimationVirtualSubject.cpp
--- a/Engine/Plugins/Animation/LiveLink/Source/LiveLink/Private/LiveLinkAnimationVirtualSubject.cpp
+++ b/Engine/Plugins/Animation/LiveLink/Source/LiveLink/Private/LiveLinkAnimationVirtualSubject.cpp
@@ -86,25 +86,23 @@ void ULiveLinkAnimationVirtualSubject::Update()
 
 	TArray<FLiveLinkSubjectKey> ActiveSubjects = LiveLinkClient->GetSubjects(false, false);
 
-	if (AreSubjectsValid(ActiveSubjects))
+	TArray<FLiveLinkSubjectFrameData> SubjectSnapshot;
+
+	if (bSubjectsNeedSorting)
 	{
-		TArray<FLiveLinkSubjectFrameData> SubjectSnapshot;
+		SortSubjects();
+	}
 
-		if (bSubjectsNeedSorting)
-		{
-			SortSubjects();
-		}
+	if (!IsPaused())
+	{
+		bLastSnapshotIsComplete = BuildSubjectSnapshot(SubjectSnapshot);
+		BuildSkeleton(SubjectSnapshot);
+		BuildFrame(SubjectSnapshot);
 
-		if (!IsPaused())
-		{
-			if (BuildSubjectSnapshot(SubjectSnapshot))
-			{
-				BuildSkeleton(SubjectSnapshot);
-				BuildFrame(SubjectSnapshot);
-			}
-		}
+		LastFrameSnapshots = MoveTemp(SubjectSnapshot);
 	}
-	else
+
+	if (!AreSubjectsValid(ActiveSubjects))
 	{
 		for (const FLiveLinkVirtualSubjectBoneAttachment& Attachment : Attachments)
 		{
@@ -153,8 +151,9 @@ bool ULiveLinkAnimationVirtualSubject::BuildSubjectSnapshot(TArray<FLiveLinkSubj
 {
 	OutSnapshot.Reset(Subjects.Num());
 
-	bool bSnapshotDone = true;
+	bool bAllSnapshotsSucceeded = true;
 
+	uint32 SubjectIndex = 0;
 	for (const FName& SubjectName : Subjects)
 	{
 		FLiveLinkSubjectFrameData& NextSnapshot = OutSnapshot.AddDefaulted_GetRef();
@@ -164,12 +163,16 @@ bool ULiveLinkAnimationVirtualSubject::BuildSubjectSnapshot(TArray<FLiveLinkSubj
 		TSubclassOf<ULiveLinkRole> DesiredRole = SubjectRole && SubjectRole->IsChildOf(ULiveLinkBasicRole::StaticClass()) ? SubjectRole : GetRole();
 		if (!LiveLinkClient->EvaluateFrame_AnyThread(SubjectName, DesiredRole, NextSnapshot))
 		{
-			bSnapshotDone = false;
-			break;
+			bAllSnapshotsSucceeded = false;
+			// Snapshots are non-copyable, so this should be fine since the resulting snapshot array will be used
+			// to build the next LastFrameSnapshots list.
+			NextSnapshot = MoveTemp(LastFrameSnapshots[SubjectIndex]);
 		}
+
+		SubjectIndex++;
 	}
 
-	return bSnapshotDone;
+	return bAllSnapshotsSucceeded;
 }
 
 void ULiveLinkAnimationVirtualSubject::BuildSkeleton(const TArray<FLiveLinkSubjectFrameData>& InSubjectSnapshots)
@@ -274,8 +277,11 @@ void ULiveLinkAnimationVirtualSubject::BuildFrame(const TArray<FLiveLinkSubjectF
 	}
 
 	if (NewSnapshotFrameData->MetaData.SceneTime.Time == 0)
-	{	
-		if (TOptional<FQualifiedFrameTime> FrameTime = FApp::GetCurrentFrameTime())
+	{
+		PRAGMA_DISABLE_DEPRECATION_WARNINGS
+		FAppTime AppTime = FAppTime::Get_UnsafeNoCheck();
+		PRAGMA_ENABLE_DEPRECATION_WARNINGS
+		if (TOptional<FQualifiedFrameTime> FrameTime = AppTime.CurrentFrameTime)
 		{
 			NewSnapshotFrameData->MetaData.SceneTime = *FrameTime;
 		}
diff --git a/Engine/Plugins/Animation/LiveLink/Source/LiveLink/Private/LiveLinkClient.cpp b/Engine/Plugins/Animation/LiveLink/Source/LiveLink/Private/LiveLinkClient.cpp
--- a/Engine/Plugins/Animation/LiveLink/Source/LiveLink/Private/LiveLinkClient.cpp
+++ b/Engine/Plugins/Animation/LiveLink/Source/LiveLink/Private/LiveLinkClient.cpp
@@ -563,17 +563,9 @@ ELiveLinkSubjectState FLiveLinkClient::GetSubjectState(FLiveLinkSubjectName InSu
 			{
 				return ELiveLinkSubjectState::InvalidOrDisabled;
 			}
-			else if (VirtualSubject->IsPaused())
-			{
-				return ELiveLinkSubjectState::Paused;
-			}
-			else if (VirtualSubject->HasValidFrameSnapshot())
-			{
-				return ELiveLinkSubjectState::Connected;
-			}
 			else
 			{
-				return ELiveLinkSubjectState::InvalidOrDisabled;
+				return VirtualSubject->GetState();
 			}
 		}
 		else
@@ -1629,9 +1621,9 @@ bool FLiveLinkClient::IsSubjectValid(const FLiveLinkSubjectKey& InSubjectKey) co
 	{
 		if (SubjectItem->GetSubject()->HasValidFrameSnapshot())
 		{
-			if (SubjectItem->GetVirtualSubject())
+			if (ULiveLinkVirtualSubject* VSubject = SubjectItem->GetVirtualSubject())
 			{
-				return true;
+				return VSubject->GetState() == ELiveLinkSubjectState::Connected;
 			}
 
 			if (FLiveLinkSubject* LiveSubject = SubjectItem->GetLiveSubject())
diff --git a/Engine/Plugins/Animation/LiveLink/Source/LiveLink/Public/LiveLinkAnimationVirtualSubject.h b/Engine/Plugins/Animation/LiveLink/Source/LiveLink/Public/LiveLinkAnimationVirtualSubject.h
--- a/Engine/Plugins/Animation/LiveLink/Source/LiveLink/Public/LiveLinkAnimationVirtualSubject.h
+++ b/Engine/Plugins/Animation/LiveLink/Source/LiveLink/Public/LiveLinkAnimationVirtualSubject.h
@@ -69,6 +69,27 @@ protected:
 	/**  Apply transform offsets generated by the attachments to the frame data. */
 	UE_API void ProcessAttachmentsForFrameData(struct FLiveLinkAnimationFrameData* SnapshotFrameData);
 
+	//~ LiveLinkVirtualSubject Interface
+	ELiveLinkSubjectState GetState() const override
+	{
+		if (IsPaused())
+		{
+			return ELiveLinkSubjectState::Paused;
+		}
+		else if (!bLastSnapshotIsComplete)
+		{
+			return ELiveLinkSubjectState::InvalidOrDisabled;
+		}
+		else if (HasValidFrameSnapshot())
+		{
+			return ELiveLinkSubjectState::Connected;
+		}
+		else
+		{
+			return ELiveLinkSubjectState::InvalidOrDisabled;
+		}
+	}
+
 public:
 	/** List of bone attachments between subjects. */
 	UPROPERTY(EditAnywhere, Category = "LiveLink")
@@ -102,6 +123,15 @@ private:
 	/** Lookup table of Subject + bone name to index in the concatenated list of bone names. */
 	TMap<TPair<FName, FName>, int32> BoneNameToIndex;
 
+	/** Cached snapshots, used when the subjects are invalid in order to maintain tranmission of data. */
+	TArray<FLiveLinkSubjectFrameData> LastFrameSnapshots;
+
+	/** 
+	 * Returns whether all snapshots were successfully created. 
+	 * Will be false if no snapshot was created, or if one of the subject was invalid, causing the subject to use the previous frame's data. 
+	 */
+	bool bLastSnapshotIsComplete = false;
+
 	/** Whether the subjects need to be sorted because an attachment has been modified. */
 	bool bSubjectsNeedSorting = false;
 };
diff --git a/Engine/Plugins/Animation/LiveLink/Source/LiveLink/Public/LiveLinkSettings.h b/Engine/Plugins/Animation/LiveLink/Source/LiveLink/Public/LiveLinkSettings.h
--- a/Engine/Plugins/Animation/LiveLink/Source/LiveLink/Public/LiveLinkSettings.h
+++ b/Engine/Plugins/Animation/LiveLink/Source/LiveLink/Public/LiveLinkSettings.h
@@ -13,6 +13,8 @@ class ULiveLinkFrameInterpolationProcessor;
 class ULiveLinkFramePreProcessor;
 class ULiveLinkRole;
 class ULiveLinkSubjectSettings;
+class ULiveLinkVirtualSubject;
+
 enum class ELiveLinkSourceMode : uint8;
 
 
@@ -93,6 +95,10 @@ public:
 	UPROPERTY(config, EditAnywhere, Category = "LiveLink")
 	TSoftObjectPtr<ULiveLinkPreset> DefaultLiveLinkPreset;
 
+	/** The default preset that should be applied */
+	UPROPERTY(config)
+	TSubclassOf<ULiveLinkVirtualSubject> DefaultVirtualSubjectClass;
+
 #if WITH_EDITORONLY_DATA
 	UE_DEPRECATED(5.1, "PresetSaveDir was moved into LiveLinkUserSettings. Please use ULiveLinkUserSettings::GetPresetSaveDir().")
 	UPROPERTY(config)
diff --git a/Engine/Plugins/Animation/LiveLink/Source/LiveLinkEditor/Private/LiveLinkClientPanelViews.cpp b/Engine/Plugins/Animation/LiveLink/Source/LiveLinkEditor/Private/LiveLinkClientPanelViews.cpp
--- a/Engine/Plugins/Animation/LiveLink/Source/LiveLinkEditor/Private/LiveLinkClientPanelViews.cpp
+++ b/Engine/Plugins/Animation/LiveLink/Source/LiveLinkEditor/Private/LiveLinkClientPanelViews.cpp
@@ -53,7 +53,6 @@ namespace SourceListUI
 	static const FName ActionsColumnName(TEXT("Action"));
 };
 
-
 /**
  * Widget for the filter bar, needs to be a subclass that overrides MakeAddFilterMenu to give the filter bar its own unique menu name,
  * otherwise the editor will get confused with any other basic filter bar used elsewhere.
@@ -117,6 +116,7 @@ private:
 	}
 };
 
+
 /** Dialog to create a new virtual subject */
 class SVirtualSubjectCreateDialog : public SCompoundWidget
 {
@@ -125,6 +125,7 @@ public:
 
 		/** Pointer to the LiveLinkClient instance. */
 		SLATE_ARGUMENT(FLiveLinkClient*, LiveLinkClient)
+		SLATE_ARGUMENT(TSubclassOf<ULiveLinkVirtualSubject>, DefaultClass)
 
 	SLATE_END_ARGS()
 
@@ -134,7 +135,8 @@ public:
 		static const FName DefaultVirtualSubjectName = TEXT("Virtual");
 
 		bOkClicked = false;
-		VirtualSubjectClass = nullptr;
+		VirtualSubjectClass = InArgs._DefaultClass;
+		bCreatedWithDefaultClass = InArgs._DefaultClass != nullptr;
 		LiveLinkClient = InArgs._LiveLinkClient;
 
 		check(LiveLinkClient);
@@ -143,14 +145,12 @@ public:
 			{
 				return LiveLinkClient->IsVirtualSubject(SubjectKey) ? 1 : 0;
 			}, 0);
-
 		VirtualSubjectName = DefaultVirtualSubjectName;
 
 		if (NumVirtualSubjects > 0)
 		{
 			VirtualSubjectName = *FString::Printf(TEXT("%s %d"), *DefaultVirtualSubjectName.ToString(), NumVirtualSubjects + 1);
 		}
-
 		//Default VirtualSubject Source should always exist
 		TArray<FGuid> Sources = LiveLinkClient->GetVirtualSources();
 		check(Sources.Num() > 0);
@@ -180,14 +180,15 @@ public:
 						]
 
 						+ SVerticalBox::Slot()
-						.FillHeight(1.0)
+						.FillHeight(VirtualSubjectClass == nullptr ? 1.0 : 0.0)
 						[
 							SNew(SBorder)
-								.BorderImage(FAppStyle::GetBrush("ToolPanel.GroupBorder"))
-								.Content()
-								[
-									SAssignNew(RoleClassPicker, SVerticalBox)
-								]
+							.Visibility(VirtualSubjectClass == nullptr ? EVisibility::Visible : EVisibility::Collapsed)
+							.BorderImage(FAppStyle::GetBrush("ToolPanel.GroupBorder"))
+							.Content()
+							[
+								SAssignNew(RoleClassPicker, SVerticalBox)
+							]
 						]
 
 						// Ok/Cancel buttons
@@ -233,9 +234,10 @@ public:
 
 	bool ConfigureVirtualSubject()
 	{
+		const double WindowHeight = bCreatedWithDefaultClass ? 100 : 300;
 		TSharedRef<SWindow> Window = SNew(SWindow)
 			.Title(LOCTEXT("CreateVirtualSubjectCreation", "Create Virtual Subject"))
-			.ClientSize(FVector2D(400, 300))
+			.ClientSize(FVector2D(400, WindowHeight))
 			.SupportsMinimize(false)
 			.SupportsMaximize(false)
 			[
@@ -397,6 +399,9 @@ private:
 
 	/** True if Ok was clicked */
 	bool bOkClicked;
+
+	/** Whether the widget was created with a default class, used to resize the window accordingly. */
+	bool bCreatedWithDefaultClass = false;
 };
 
 namespace UE::LiveLink
@@ -541,8 +546,11 @@ void AddVirtualSubject()
 {
 	FLiveLinkClient* LiveLinkClient = (FLiveLinkClient*)&IModularFeatures::Get().GetModularFeature<ILiveLinkClient>(ILiveLinkClient::ModularFeatureName);
 
+	const TSubclassOf<ULiveLinkVirtualSubject> Subclass = GetDefault<ULiveLinkSettings>()->DefaultVirtualSubjectClass;
+
 	TSharedRef<SVirtualSubjectCreateDialog> Dialog =
 		SNew(SVirtualSubjectCreateDialog)
+		.DefaultClass(Subclass)
 		.LiveLinkClient(LiveLinkClient);
 
 	Dialog->ConfigureVirtualSubject();
diff --git a/Engine/Plugins/Animation/LiveLinkHub/Source/LiveLinkHub/Private/LiveLinkHubAnimationVirtualSubject.h b/Engine/Plugins/Animation/LiveLinkHub/Source/LiveLinkHub/Private/LiveLinkHubAnimationVirtualSubject.h
--- a/Engine/Plugins/Animation/LiveLinkHub/Source/LiveLinkHub/Private/LiveLinkHubAnimationVirtualSubject.h
+++ b/Engine/Plugins/Animation/LiveLinkHub/Source/LiveLinkHub/Private/LiveLinkHubAnimationVirtualSubject.h
@@ -38,18 +38,8 @@ public:
 
 	virtual FText GetDisplayName() const override
 	{
-		FText DisplayName;
-
-		if (*OutboundName == SubjectKey.SubjectName)
-		{
-			DisplayName = FText::FromName(SubjectKey.SubjectName);
-		}
-		else
-		{
-			DisplayName = FText::Format(INVTEXT("{0} ({1})"), FText::FromString(OutboundName), FText::FromName(SubjectKey.SubjectName));
-		}
-
-		return DisplayName;
+		// In the case of Virtual Subjects, we don't care about the original name.
+		return FText::FromString(OutboundName);
 	}
 
 	virtual FName GetRebroadcastName() const override
diff --git a/Engine/Source/Programs/LiveLinkHub/Config/DefaultGame.ini b/Engine/Source/Programs/LiveLinkHub/Config/DefaultGame.ini
--- a/Engine/Source/Programs/LiveLinkHub/Config/DefaultGame.ini
+++ b/Engine/Source/Programs/LiveLinkHub/Config/DefaultGame.ini
@@ -13,6 +13,7 @@ bIsForExternalDistribution=true
 DefaultSettingsClass=/Script/LiveLink.LiveLinkHubSubjectSettings
 -DefaultRoleSettings=(Role="/Script/CoreUObject.Class'/Script/LiveLinkInterface.LiveLinkAnimationRole'",SettingClass="/Script/CoreUObject.Class'/Script/LiveLinkInterface.LiveLinkSubjectSettings'",FrameInterpolationProcessor="/Script/CoreUObject.Class'/Script/LiveLink.LiveLinkAnimationFrameInterpolationProcessor'",FramePreProcessors=)
 +DefaultRoleSettings=(Role="/Script/CoreUObject.Class'/Script/LiveLinkInterface.LiveLinkAnimationRole'",SettingClass="/Script/CoreUObject.Class'/Script/LiveLink.LiveLinkHubSubjectSettings'",FrameInterpolationProcessor="/Script/CoreUObject.Class'/Script/LiveLink.LiveLinkAnimationFrameInterpolationProcessor'",FramePreProcessors=)
+DefaultVirtualSubjectClass=/Script/LiveLinkHub.LiveLinkHubAnimationVirtualSubject
 bPreProcessRebroadcastFrames=true
 bTranslateRebroadcastFrames=true
 
diff --git a/Engine/Source/Runtime/LiveLinkInterface/Private/LiveLinkVirtualSubject.cpp b/Engine/Source/Runtime/LiveLinkInterface/Private/LiveLinkVirtualSubject.cpp
--- a/Engine/Source/Runtime/LiveLinkInterface/Private/LiveLinkVirtualSubject.cpp
+++ b/Engine/Source/Runtime/LiveLinkInterface/Private/LiveLinkVirtualSubject.cpp
@@ -90,6 +90,22 @@ bool ULiveLinkVirtualSubject::DependsOnSubject(FName SubjectName) const
 	return Subjects.Contains(SubjectName);
 }
 
+ELiveLinkSubjectState ULiveLinkVirtualSubject::GetState() const
+{
+	if (IsPaused())
+	{
+		return ELiveLinkSubjectState::Paused;
+	}
+	else if (HasValidFrameSnapshot())
+	{
+		return ELiveLinkSubjectState::Connected;
+	}
+	else
+	{
+		return ELiveLinkSubjectState::InvalidOrDisabled;
+	}
+}
+
 void ULiveLinkVirtualSubject::UpdateTranslatorsForThisFrame()
 {
 	// Create the new translator for this frame
diff --git a/Engine/Source/Runtime/LiveLinkInterface/Public/LiveLinkVirtualSubject.h b/Engine/Source/Runtime/LiveLinkInterface/Public/LiveLinkVirtualSubject.h
--- a/Engine/Source/Runtime/LiveLinkInterface/Public/LiveLinkVirtualSubject.h
+++ b/Engine/Source/Runtime/LiveLinkInterface/Public/LiveLinkVirtualSubject.h
@@ -5,6 +5,7 @@
 #include "Containers/Array.h"
 #include "CoreTypes.h"
 #include "HAL/CriticalSection.h"
+#include "ILiveLinkClient.h"
 #include "ILiveLinkSubject.h"
 #include "LiveLinkFrameTranslator.h"
 #include "LiveLinkTypes.h"
@@ -92,6 +93,9 @@ public:
 		return SubjectKey.SubjectName;
 	}
 
+	/** Get the state of the virtual subject. */
+	LIVELINKINTERFACE_API virtual ELiveLinkSubjectState GetState() const;
+
 public:
 	/** If set, the virtual subject will only update and transmit data when the sync subject receives new data. */
 	UPROPERTY(EditAnywhere, Category = "LiveLink")
