LiveLink- Fix LiveLinkBroadcast component not creating subject if using component name as reference

#Synced-CL 50755163

diff --git a/Engine/Plugins/Animation/LiveLink/Source/LiveLink/Private/LiveLinkBroadcastComponent.cpp b/Engine/Plugins/Animation/LiveLink/Source/LiveLink/Private/LiveLinkBroadcastComponent.cpp
--- a/Engine/Plugins/Animation/LiveLink/Source/LiveLink/Private/LiveLinkBroadcastComponent.cpp
+++ b/Engine/Plugins/Animation/LiveLink/Source/LiveLink/Private/LiveLinkBroadcastComponent.cpp
@@ -143,7 +143,7 @@ void ULiveLinkBroadcastComponent::TickComponent(float DeltaTime, ELevelTick Tick
 
 void ULiveLinkBroadcastComponent::InitiateLiveLink()
 {
-	if (Role && !SourceMesh.PathToComponent.IsEmpty())
+	if (Role && (!SourceMesh.PathToComponent.IsEmpty() || !SourceMesh.ComponentProperty.IsNone()))
 	{
 		SubjectKey = GEngine->GetEngineSubsystem<ULiveLinkBroadcastSubsystem>()->CreateSubject(SubjectName, Role);
 	}
