LiveLink - Fix modifications to resequencing options not being into account

Message Bus Source settings were only being cached when the source was created, so any modification made after that was not reflected in the message bus source.

#Synced-CL 50901149

diff --git a/Engine/Plugins/Animation/LiveLink/Source/LiveLink/Private/LiveLinkMessageBusSource.cpp b/Engine/Plugins/Animation/LiveLink/Source/LiveLink/Private/LiveLinkMessageBusSource.cpp
--- a/Engine/Plugins/Animation/LiveLink/Source/LiveLink/Private/LiveLinkMessageBusSource.cpp
+++ b/Engine/Plugins/Animation/LiveLink/Source/LiveLink/Private/LiveLinkMessageBusSource.cpp
@@ -134,11 +134,26 @@ FLiveLinkMessageBusSource::FLiveLinkMessageBusSource(const FText& InSourceType,
 	, MachineTimeOffset(InMachineTimeOffset)
 	, bDiscovering(false)
 	, bInitialized(false )
-{}
+{
+}
 
 void FLiveLinkMessageBusSource::InitializeSettings(ULiveLinkSourceSettings* Settings)
 {
-	if (const ULiveLinkMessageBusSourceSettings* MessageBusSettings = Cast<ULiveLinkMessageBusSourceSettings>(Settings))
+	if (ULiveLinkMessageBusSourceSettings* MessageBusSettings = Cast<ULiveLinkMessageBusSourceSettings>(Settings))
+	{
+		SourceSettings = TWeakObjectPtr<ULiveLinkMessageBusSourceSettings>(MessageBusSettings);
+
+#if WITH_EDITOR
+		MessageBusSettings->OnPropertyChangedDelegate.AddRaw(this, &FLiveLinkMessageBusSource::CacheMessageBusSourceSettings);
+#endif
+
+		CacheMessageBusSourceSettings();
+	}
+}
+
+void FLiveLinkMessageBusSource::CacheMessageBusSourceSettings()
+{
+	if (const ULiveLinkMessageBusSourceSettings* MessageBusSettings = SourceSettings.Get())
 	{
 		bCachedEnableFrameResequencing.store(MessageBusSettings->bEnableFrameResequencing, std::memory_order_relaxed);
 		CachedResequencerMaxBufferSize.store(MessageBusSettings->ResequencerMaxBufferSize, std::memory_order_relaxed);
@@ -146,6 +161,7 @@ void FLiveLinkMessageBusSource::InitializeSettings(ULiveLinkSourceSettings* Sett
 	}
 }
 
+
 void FLiveLinkMessageBusSource::ReceiveClient(ILiveLinkClient* InClient, FGuid InSourceGuid)
 {
 	Client = InClient;
@@ -662,6 +678,13 @@ void FLiveLinkMessageBusSource::SendConnectMessage()
 
 bool FLiveLinkMessageBusSource::RequestSourceShutdown()
 {
+#if WITH_EDITOR
+	if (ULiveLinkMessageBusSourceSettings* MessageBusSettings = SourceSettings.Get())
+	{
+		MessageBusSettings->OnPropertyChangedDelegate.RemoveAll(this);
+	}
+#endif
+
 	if (!bInitialized)
 	{
 		return true;
diff --git a/Engine/Plugins/Animation/LiveLink/Source/LiveLink/Public/LiveLinkMessageBusSource.h b/Engine/Plugins/Animation/LiveLink/Source/LiveLink/Public/LiveLinkMessageBusSource.h
--- a/Engine/Plugins/Animation/LiveLink/Source/LiveLink/Public/LiveLinkMessageBusSource.h
+++ b/Engine/Plugins/Animation/LiveLink/Source/LiveLink/Public/LiveLinkMessageBusSource.h
@@ -148,6 +148,9 @@ private:
 
 	// Resequencer for maintaining frame order per subject
 	void ResequenceAndPushFrame(const FLiveLinkSubjectKey& SubjectKey, FLiveLinkFrameDataStruct&& FrameData);
+	
+	// Cache Message Bus settings
+	void CacheMessageBusSourceSettings();
 
 private:
 
@@ -215,6 +218,9 @@ private:
 	FText SourceType;
 	FText SourceMachineName;
 
+	/** Weak pointer to the source settings, used to cache source settings values when they're modified. */
+	TWeakObjectPtr<class ULiveLinkMessageBusSourceSettings> SourceSettings;
+
 	// Time we last received anything 
 	double ConnectionLastActive;
 
diff --git a/Engine/Plugins/Animation/LiveLink/Source/LiveLink/Public/LiveLinkMessageBusSourceSettings.h b/Engine/Plugins/Animation/LiveLink/Source/LiveLink/Public/LiveLinkMessageBusSourceSettings.h
--- a/Engine/Plugins/Animation/LiveLink/Source/LiveLink/Public/LiveLinkMessageBusSourceSettings.h
+++ b/Engine/Plugins/Animation/LiveLink/Source/LiveLink/Public/LiveLinkMessageBusSourceSettings.h
@@ -35,6 +35,25 @@ public:
 	/** Maximum time in seconds to hold frames waiting for missing frames. */
 	UPROPERTY(config, EditAnywhere, Category = "Resequencing", meta = (EditCondition = "bEnableFrameResequencing", ClampMin = "0.01", ClampMax = "10.0"))
 	float ResequencerMaxWaitTimeSeconds = 0.1f;
+
+#if WITH_EDITOR
+	virtual void PostEditChangeProperty(FPropertyChangedEvent& PropertyChangedEvent) override
+	{
+		Super::PostEditChangeProperty(PropertyChangedEvent);
+
+		if (PropertyChangedEvent.Property &&
+				(PropertyChangedEvent.Property->GetFName() == GET_MEMBER_NAME_CHECKED(ULiveLinkMessageBusSourceSettings, bEnableFrameResequencing)
+				|| PropertyChangedEvent.Property->GetFName() == GET_MEMBER_NAME_CHECKED(ULiveLinkMessageBusSourceSettings, ResequencerMaxBufferSize)
+				|| PropertyChangedEvent.Property->GetFName() == GET_MEMBER_NAME_CHECKED(ULiveLinkMessageBusSourceSettings, ResequencerMaxWaitTimeSeconds)))
+		{
+			OnPropertyChangedDelegate.Broadcast();
+		}
+	}
+
+	/** Delegate triggered whenever a message bus source settings property is modified. */
+	FTSSimpleMulticastDelegate OnPropertyChangedDelegate;
+#endif
+
 };
 
 #undef UE_API
