LiveLink - If adding a source that is pending kill, remove the pending kill flag instead
Fix static data not being rebroadcast when it's received

Fix crash when duplicating settings object outside the game thread

#Synced-CL 50901149

#Synced-CL 50901149

diff --git a/Engine/Plugins/Animation/LiveLink/Source/LiveLink/Private/LiveLinkClient.cpp b/Engine/Plugins/Animation/LiveLink/Source/LiveLink/Private/LiveLinkClient.cpp
--- a/Engine/Plugins/Animation/LiveLink/Source/LiveLink/Private/LiveLinkClient.cpp
+++ b/Engine/Plugins/Animation/LiveLink/Source/LiveLink/Private/LiveLinkClient.cpp
@@ -606,7 +606,12 @@ FGuid FLiveLinkClient::AddSource(TSharedPtr<ILiveLinkSource> InSource)
 	check(Collection);
 
 	FGuid Guid;
-	if (Collection->FindSource(InSource) == nullptr)
+	if (FLiveLinkCollectionSourceItem* SourceItem = Collection->FindSource(InSource))
+	{
+		SourceItem->bPendingKill = false;
+		Guid = SourceItem->Guid;
+	}
+	else
 	{
 		ULiveLinkSourceSettings* Settings = nullptr;
 		Guid = FGuid::NewGuid();
@@ -626,6 +631,7 @@ FGuid FLiveLinkClient::AddSource(TSharedPtr<ILiveLinkSource> InSource)
 		InSource->ReceiveClient(this, Guid);
 		InSource->InitializeSettings(Settings);
 	}
+
 	return Guid;
 }
 
@@ -946,7 +952,12 @@ void FLiveLinkClient::PushSubjectStaticData_Internal(FPendingSubjectStatic&& Sub
 	// Prevent GC while we're creating UObjects since LL Client can potentially be ticked outside of the game thread.
 	FGCScopeGuard Guard;
 
-	if (LiveLinkSubject == nullptr)
+	if (LiveLinkSubject)
+	{
+		// If new static data was received, make sure it's rebroadcast
+		LiveLinkSubject->SetStaticDataAsRebroadcasted(false);
+	}
+	else
 	{
 		const ULiveLinkSettings* LiveLinkSettings = GetDefault<ULiveLinkSettings>();
 		const FLiveLinkRoleProjectSetting DefaultSetting = LiveLinkSettings->GetDefaultSettingForRole(SubjectStaticData.Role.Get());
@@ -960,7 +971,16 @@ void FLiveLinkClient::PushSubjectStaticData_Internal(FPendingSubjectStatic&& Sub
 				// This allows preprocessors and other settings to survive source disconnection/reconnection.
 				if (TStrongObjectPtr<ULiveLinkSubjectSettings>* CachedSettings = CachedPresetSubjectSettings.Find(SubjectStaticData.SubjectKey.SubjectName))
 				{
-					SubjectSettings = TStrongObjectPtr{ DuplicateObject<ULiveLinkSubjectSettings>(CachedSettings->Get(), GetTransientPackage()) };
+					FObjectDuplicationParameters Parameters = InitStaticDuplicateObjectParams(CachedSettings->Get(), GetTransientPackage());
+					Parameters.bSkipPostLoad = true;
+				
+					// Cannot call ordinary DuplicateObject() because it calls PostLoad for legacy reasons,
+					// which is not allowed anywhere except the GameThread.
+					SubjectSettings = TStrongObjectPtr{ Cast<ULiveLinkSubjectSettings>(StaticDuplicateObjectEx(Parameters)) };
+
+					// PATCH until UE-250259 bug is fixed
+					SubjectSettings->ClearFlags(EObjectFlags::RF_NeedPostLoad | EObjectFlags::RF_NeedPostLoadSubobjects);
+
 					bRestoredFromPresetCache = true;
 				}
 				else
