LiveLink - Add Character LiveLinkRole which combines animation role with transform role to allow streaming a skeleton + transform

#Synced-CL 51376035

diff --git a/Engine/Plugins/Animation/LiveLink/Source/LiveLink/Private/LiveLinkBroadcastComponent.cpp b/Engine/Plugins/Animation/LiveLink/Source/LiveLink/Private/LiveLinkBroadcastComponent.cpp
--- a/Engine/Plugins/Animation/LiveLink/Source/LiveLink/Private/LiveLinkBroadcastComponent.cpp
+++ b/Engine/Plugins/Animation/LiveLink/Source/LiveLink/Private/LiveLinkBroadcastComponent.cpp
@@ -15,6 +15,8 @@
 #include "Roles/LiveLinkAnimationRole.h"
 #include "Roles/LiveLinkCameraRole.h"
 #include "Roles/LiveLinkCameraTypes.h"
+#include "Roles/LiveLinkCharacterRole.h"
+#include "Roles/LiveLinkCharacterTypes.h"
 #include "Roles/LiveLinkTransformRole.h"
 #include "Roles/LiveLinkTransformTypes.h"
 #include "UObject/UObjectGlobals.h"
@@ -278,7 +280,11 @@ void ULiveLinkBroadcastComponent::BroadcastStaticData() const
 
 	FLiveLinkStaticDataStruct StaticDataStruct;
 
-	if (Role->IsChildOf<ULiveLinkAnimationRole>())
+	if (Role->IsChildOf<ULiveLinkCharacterRole>())
+	{
+		StaticDataStruct = GetCharacterStaticData();
+	}
+	else if (Role->IsChildOf<ULiveLinkAnimationRole>())
 	{
 		StaticDataStruct = GetSkeletonStaticData();
 	}
@@ -317,6 +323,40 @@ FLiveLinkFrameDataStruct ULiveLinkBroadcastComponent::GetSkeletonFrameData() con
 	return FrameDataStruct;
 }
 
+
+FLiveLinkStaticDataStruct ULiveLinkBroadcastComponent::GetCharacterStaticData() const
+{
+	FLiveLinkStaticDataStruct StaticDataStruct(FLiveLinkCharacterStaticData::StaticStruct());
+	FLiveLinkCharacterStaticData& StaticData = *StaticDataStruct.Cast<FLiveLinkCharacterStaticData>();
+
+	StaticData.BoneNames = GetSkeletalBoneNames();
+	StaticData.BoneParents = GetSkeletalBoneParents();
+	StaticData.PropertyNames = GetCurveNames();
+
+	StaticData.bIsLocationSupported = true;
+	StaticData.bIsRotationSupported = true;
+	StaticData.bIsScaleSupported = true;
+
+	return StaticDataStruct;
+}
+
+FLiveLinkFrameDataStruct ULiveLinkBroadcastComponent::GetCharacterFrameData() const
+{
+
+	FLiveLinkFrameDataStruct FrameDataStruct(FLiveLinkCharacterFrameData::StaticStruct());
+	FLiveLinkCharacterFrameData& FrameData = *FrameDataStruct.Cast<FLiveLinkCharacterFrameData>();
+
+	FrameData.Transforms = GetBoneTransforms();
+	FrameData.PropertyValues = GetCurveValues();
+
+	if (USceneComponent* SceneComponent = Cast<USceneComponent>(SourceMesh.GetComponent(GetOwner())))
+	{
+		FrameData.Transform = SceneComponent->GetComponentTransform();
+	}
+
+	return FrameDataStruct;
+}
+
 FLiveLinkStaticDataStruct ULiveLinkBroadcastComponent::GetCameraStaticData() const
 {
 	FLiveLinkStaticDataStruct StaticDataStruct(FLiveLinkCameraStaticData::StaticStruct());
@@ -408,7 +448,11 @@ void ULiveLinkBroadcastComponent::BroadcastFrameData() const
 
 	FLiveLinkFrameDataStruct FrameDataStruct;
 
-	if (Role->IsChildOf<ULiveLinkAnimationRole>())
+	if (Role->IsChildOf<ULiveLinkCharacterRole>())
+	{
+		FrameDataStruct = GetCharacterFrameData();
+	}
+	else if (Role->IsChildOf<ULiveLinkAnimationRole>())
 	{
 		FrameDataStruct = GetSkeletonFrameData();
 	}
diff --git a/Engine/Plugins/Animation/LiveLink/Source/LiveLink/Public/LiveLinkBroadcastComponent.h b/Engine/Plugins/Animation/LiveLink/Source/LiveLink/Public/LiveLinkBroadcastComponent.h
--- a/Engine/Plugins/Animation/LiveLink/Source/LiveLink/Public/LiveLinkBroadcastComponent.h
+++ b/Engine/Plugins/Animation/LiveLink/Source/LiveLink/Public/LiveLinkBroadcastComponent.h
@@ -10,17 +10,18 @@
 
 #include "LiveLinkBroadcastComponent.generated.h"
 
-
 class FLiveLinkBroadcastSource;
 
-UCLASS(BlueprintType, DisplayName="Live Link Broadcast", ClassGroup=("LiveLink"), meta=(BlueprintSpawnableComponent), HideCategories=(Navigation, Tags, AssetUserData, Activation))
+#define UE_API LIVELINK_API
+
+UCLASS(BlueprintType, DisplayName="Live Link Broadcast", ClassGroup=("LiveLink"), meta=(BlueprintSpawnableComponent), HideCategories=(Navigation, Tags, AssetUserData, Activation), MinimalAPI)
 class ULiveLinkBroadcastComponent : public UActorComponent 
 {
 	GENERATED_BODY()
 
 public:
 	// Sets default values for this component's properties
-	ULiveLinkBroadcastComponent();
+	UE_API ULiveLinkBroadcastComponent();
 
 	/** Whether this component should broadcast data. */
 	UPROPERTY(BlueprintReadWrite, EditAnywhere, Category = "Live Link Broadcast")
@@ -59,7 +60,13 @@ public:
 	TArray<FName> AllowedCurveNames;
 
 	//~ UActorComponent interface
-	LIVELINK_API virtual void TickComponent(float DeltaTime, ELevelTick TickType, FActorComponentTickFunction* ThisTickFunction) override;
+	UE_API virtual void TickComponent(float DeltaTime, ELevelTick TickType, FActorComponentTickFunction* ThisTickFunction) override;
+
+	/** Trigger retransmission of static data on the next tick. */
+	void MarkDirty()
+	{
+		bIsDirty = true;
+	}
 
 private:
 	/** Creates the rebroadcast source if it doesn't exist. */
@@ -67,11 +74,13 @@ private:
 
 	//~ Populates a static data struct from the chosen mesh
 	FLiveLinkStaticDataStruct GetSkeletonStaticData() const;
+	FLiveLinkStaticDataStruct GetCharacterStaticData() const;
 	FLiveLinkStaticDataStruct GetCameraStaticData() const;
 	FLiveLinkStaticDataStruct GetTransformStaticData() const;
 
 	//~ Populates a frame data struct from the chosen mesh
 	FLiveLinkFrameDataStruct GetSkeletonFrameData() const;
+	FLiveLinkFrameDataStruct GetCharacterFrameData() const;
 	FLiveLinkFrameDataStruct GetCameraFrameData() const;
 	FLiveLinkFrameDataStruct GetTransformFrameData() const;
 
@@ -98,13 +107,13 @@ private:
 	
 protected:
 	//~ Begin UActorComponent interface
-	virtual void OnComponentCreated() override;
-	virtual void OnComponentDestroyed(bool bDestroyingHierarchy) override;
+	UE_API virtual void OnComponentCreated() override;
+	UE_API virtual void OnComponentDestroyed(bool bDestroyingHierarchy) override;
 	//~ End UActorComponent interface
 	
 #if WITH_EDITOR
 	//~ Begin UObject interface 
-	virtual void PostEditChangeProperty(struct FPropertyChangedEvent& PropertyChangedEvent) override;
+	UE_API virtual void PostEditChangeProperty(struct FPropertyChangedEvent& PropertyChangedEvent) override;
 	//~ End UObject interface 
 
 	/** Object modified callback used to trigger the dirty flag if one of the properties on the tracked source object is modified. */
@@ -120,4 +129,4 @@ private:
 	FLiveLinkSubjectKey SubjectKey;
 };
 
-
+#undef UE_API
\ No newline at end of file
diff --git a/Engine/Plugins/Animation/LiveLink/Source/LiveLinkComponents/Private/Controllers/LiveLinkCharacterController.cpp b/Engine/Plugins/Animation/LiveLink/Source/LiveLinkComponents/Private/Controllers/LiveLinkCharacterController.cpp
new file mode 100644
--- /dev/null
+++ b/Engine/Plugins/Animation/LiveLink/Source/LiveLinkComponents/Private/Controllers/LiveLinkCharacterController.cpp
@@ -0,0 +1,40 @@
+// Copyright Epic Games, Inc. All Rights Reserved.
+
+#include "Controllers/LiveLinkCharacterController.h"
+
+#include "Components/SceneComponent.h"
+#include "Roles/LiveLinkCharacterTypes.h"
+#include "Roles/LiveLinkCharacterRole.h"
+
+#include UE_INLINE_GENERATED_CPP_BY_NAME(LiveLinkCharacterController)
+
+#define LOCTEXT_NAMESPACE "LiveLinkCharacterController"
+
+bool ULiveLinkCharacterController::IsRoleSupported(const TSubclassOf<ULiveLinkRole>& RoleToSupport)
+{
+	return RoleToSupport == ULiveLinkCharacterRole::StaticClass();
+}
+
+void ULiveLinkCharacterController::Tick(float DeltaTime, const FLiveLinkSubjectFrameData& SubjectData)
+{
+	const FLiveLinkCharacterStaticData* StaticData = SubjectData.StaticData.Cast<FLiveLinkCharacterStaticData>();
+	const FLiveLinkCharacterFrameData* FrameData = SubjectData.FrameData.Cast<FLiveLinkCharacterFrameData>();
+
+	if (StaticData && FrameData)
+	{
+		if (USceneComponent* SceneComponent = Cast<USceneComponent>(GetAttachedComponent()))
+		{
+			FTransform Transform = OffsetTransform * FrameData->Transform;
+
+			FLiveLinkTransformStaticData TransformStaticData;
+			TransformStaticData.bIsLocationSupported = StaticData->bIsLocationSupported;
+			TransformStaticData.bIsRotationSupported = StaticData->bIsRotationSupported;
+			TransformStaticData.bIsScaleSupported = StaticData->bIsScaleSupported;
+
+			TransformData.ApplyTransform(SceneComponent, Transform, TransformStaticData);
+		}
+	}
+}
+
+#undef LOCTEXT_NAMESPACE
+
diff --git a/Engine/Plugins/Animation/LiveLink/Source/LiveLinkComponents/Public/Controllers/LiveLinkCharacterController.h b/Engine/Plugins/Animation/LiveLink/Source/LiveLinkComponents/Public/Controllers/LiveLinkCharacterController.h
new file mode 100644
--- /dev/null
+++ b/Engine/Plugins/Animation/LiveLink/Source/LiveLinkComponents/Public/Controllers/LiveLinkCharacterController.h
@@ -0,0 +1,24 @@
+// Copyright Epic Games, Inc. All Rights Reserved.
+
+#pragma once
+
+#include "Controllers/LiveLinkTransformController.h"
+
+#include "LiveLinkCharacterController.generated.h"
+
+#define UE_API LIVELINKCOMPONENTS_API
+
+/** Controller that allows driving an actor's transform through a subject with the Character role. */
+UCLASS(MinimalAPI, BlueprintType)
+class ULiveLinkCharacterController : public ULiveLinkTransformController
+{
+	GENERATED_BODY()
+
+public:
+	//~Begin ULiveLinkControllerBase interface
+	UE_API virtual void Tick(float DeltaTime, const FLiveLinkSubjectFrameData& SubjectData) override;
+	UE_API virtual bool IsRoleSupported(const TSubclassOf<ULiveLinkRole>& RoleToSupport) override;
+	//~End ULiveLinkControllerBase interface
+};
+
+#undef UE_API
diff --git a/Engine/Source/Runtime/LiveLinkInterface/Private/Roles/LiveLinkCharacterRole.cpp b/Engine/Source/Runtime/LiveLinkInterface/Private/Roles/LiveLinkCharacterRole.cpp
new file mode 100644
--- /dev/null
+++ b/Engine/Source/Runtime/LiveLinkInterface/Private/Roles/LiveLinkCharacterRole.cpp
@@ -0,0 +1,65 @@
+// Copyright Epic Games, Inc. All Rights Reserved.
+
+#include "Roles/LiveLinkCharacterRole.h"
+
+#include "Roles/LiveLinkAnimationBlueprintStructs.h"
+#include "Roles/LiveLinkAnimationTypes.h"
+#include "Roles/LiveLinkCharacterTypes.h"
+#include "Roles/LiveLinkTransformRole.h"
+#include "Roles/LiveLinkTransformTypes.h"
+
+#include "LiveLinkRoleTrait.h"
+#include "LiveLinkTypes.h"
+
+#include UE_INLINE_GENERATED_CPP_BY_NAME(LiveLinkCharacterRole)
+
+#define LOCTEXT_NAMESPACE "LiveLinkRole"
+
+/**
+ * ULiveLinkCharacterRole
+ */
+UScriptStruct* ULiveLinkCharacterRole::GetStaticDataStruct() const
+{
+	return FLiveLinkCharacterStaticData::StaticStruct();
+}
+
+UScriptStruct* ULiveLinkCharacterRole::GetFrameDataStruct() const
+{
+	return FLiveLinkCharacterFrameData::StaticStruct();
+}
+
+UScriptStruct* ULiveLinkCharacterRole::GetBlueprintDataStruct() const
+{
+	return FSubjectFrameHandle::StaticStruct();
+}
+
+FText ULiveLinkCharacterRole::GetDisplayName() const
+{
+	return LOCTEXT("CharacterRole", "Character");
+}
+
+bool ULiveLinkCharacterRole::IsStaticDataValid(const FLiveLinkStaticDataStruct& InStaticData, bool& bOutShouldLogWarning) const
+{
+	bool bResult = Super::IsStaticDataValid(InStaticData, bOutShouldLogWarning);
+	if (bResult)
+	{
+		const FLiveLinkSkeletonStaticData* StaticData = InStaticData.Cast<FLiveLinkSkeletonStaticData>();
+		bResult = StaticData && StaticData->BoneParents.Num() == StaticData->BoneNames.Num();
+	}
+	return bResult;
+}
+
+
+bool ULiveLinkCharacterRole::IsFrameDataValid(const FLiveLinkStaticDataStruct& InStaticData, const FLiveLinkFrameDataStruct& InFrameData, bool& bOutShouldLogWarning) const
+{
+	bool bResult = Super::IsFrameDataValid(InStaticData, InFrameData, bOutShouldLogWarning);
+	if (bResult)
+	{
+		const FLiveLinkSkeletonStaticData* StaticData = InStaticData.Cast<FLiveLinkSkeletonStaticData>();
+		const FLiveLinkAnimationFrameData* FrameData = InFrameData.Cast<FLiveLinkAnimationFrameData>();
+		bResult = StaticData && FrameData && StaticData->BoneNames.Num() == FrameData->Transforms.Num();
+	}
+	return bResult;
+}
+
+#undef LOCTEXT_NAMESPACE
diff --git a/Engine/Source/Runtime/LiveLinkInterface/Public/Roles/LiveLinkCharacterRole.h b/Engine/Source/Runtime/LiveLinkInterface/Public/Roles/LiveLinkCharacterRole.h
new file mode 100644
--- /dev/null
+++ b/Engine/Source/Runtime/LiveLinkInterface/Public/Roles/LiveLinkCharacterRole.h
@@ -0,0 +1,32 @@
+// Copyright Epic Games, Inc. All Rights Reserved.
+
+#pragma once
+
+#include "Roles/LiveLinkAnimationRole.h"
+
+#include "Internationalization/Text.h"
+#include "LiveLinkFrameInterpolationProcessor.h"
+#include "UObject/ObjectMacros.h"
+#include "UObject/UObjectGlobals.h"
+
+#include "LiveLinkCharacterRole.generated.h"
+
+/**
+ * Role associated for streaming animation and component transform.
+ */
+UCLASS(BlueprintType, meta = (DisplayName = "Character Role"), MinimalAPI)
+class ULiveLinkCharacterRole : public ULiveLinkAnimationRole
+{
+	GENERATED_BODY()
+
+public:
+	//~ Begin ULiveLinkRole interface
+	LIVELINKINTERFACE_API virtual UScriptStruct* GetStaticDataStruct() const override;
+	LIVELINKINTERFACE_API virtual UScriptStruct* GetFrameDataStruct() const override;
+	LIVELINKINTERFACE_API virtual UScriptStruct* GetBlueprintDataStruct() const override;
+
+	LIVELINKINTERFACE_API virtual FText GetDisplayName() const override;
+	LIVELINKINTERFACE_API virtual bool IsStaticDataValid(const FLiveLinkStaticDataStruct& InStaticData, bool& bOutShouldLogWarning) const override;
+	LIVELINKINTERFACE_API virtual bool IsFrameDataValid(const FLiveLinkStaticDataStruct& InStaticData, const FLiveLinkFrameDataStruct& InFrameData, bool& bOutShouldLogWarning) const override;
+	//~ End ULiveLinkRole interface
+};
diff --git a/Engine/Source/Runtime/LiveLinkInterface/Public/Roles/LiveLinkCharacterTypes.h b/Engine/Source/Runtime/LiveLinkInterface/Public/Roles/LiveLinkCharacterTypes.h
new file mode 100644
--- /dev/null
+++ b/Engine/Source/Runtime/LiveLinkInterface/Public/Roles/LiveLinkCharacterTypes.h
@@ -0,0 +1,44 @@
+// Copyright Epic Games, Inc. All Rights Reserved.
+
+#pragma once
+
+#include "Roles/LiveLinkAnimationTypes.h"
+
+#include "Roles/LiveLinkTransformTypes.h"
+#include "LiveLinkCharacterTypes.generated.h"
+
+
+/**
+ * Static data for Animation purposes. Contains data about bones that shouldn't change every frame.
+ */
+USTRUCT(BlueprintType)
+struct FLiveLinkCharacterStaticData : public FLiveLinkSkeletonStaticData
+{
+	GENERATED_BODY()
+
+public:
+	UPROPERTY(EditAnywhere, BlueprintReadWrite, Category = "LiveLink")
+	bool bIsLocationSupported = true;
+
+	// Whether rotation in frame data should be used
+	UPROPERTY(EditAnywhere, BlueprintReadWrite, Category = "LiveLink")
+	bool bIsRotationSupported = true;
+
+	// Whether scale in frame data should be used
+	UPROPERTY(EditAnywhere, BlueprintReadWrite, Category = "LiveLink")
+	bool bIsScaleSupported = true;
+};
+
+/**
+ * Dynamic data for Animation purposes. 
+ */
+USTRUCT(BlueprintType)
+struct FLiveLinkCharacterFrameData : public FLiveLinkAnimationFrameData
+{
+	GENERATED_BODY()
+
+public:
+	// Transform of the frame
+	UPROPERTY(EditAnywhere, BlueprintReadWrite, Category = "Properties", Interp)
+	FTransform Transform;
+};
