LiveLink - Fix LiveLinkBroadcast component sending socket transform instead of bone transform when there is a name conflict

#Synced-CL 51376035

diff --git a/Engine/Plugins/Animation/LiveLink/Source/LiveLink/Private/LiveLinkBroadcastComponent.cpp b/Engine/Plugins/Animation/LiveLink/Source/LiveLink/Private/LiveLinkBroadcastComponent.cpp
--- a/Engine/Plugins/Animation/LiveLink/Source/LiveLink/Private/LiveLinkBroadcastComponent.cpp
+++ b/Engine/Plugins/Animation/LiveLink/Source/LiveLink/Private/LiveLinkBroadcastComponent.cpp
@@ -216,20 +216,47 @@ TArray<FTransform> ULiveLinkBroadcastComponent::GetBoneTransforms() const
 
 	if (USkeletalMeshComponent* SourceSkelMesh = Cast<USkeletalMeshComponent>(SourceMesh.GetComponent(GetOwner())))
 	{
+		auto GetBoneTransform = [SourceSkelMesh](int32 BoneIndex, TOptional<FName> BoneName = {})
+		{
+			if (!BoneName)
+			{
+				check(BoneIndex != INDEX_NONE);
+				BoneName = SourceSkelMesh->GetBoneName(BoneIndex);
+			}
+			else
+			{
+				check(BoneName);
+				BoneIndex = SourceSkelMesh->GetBoneIndex(*BoneName);
+			}
+
+			// Note: we purposefully don't call the SourceSkelMesh::GetBoneTransform override that takes a bone name, since it calls down into GetSocketTransform,
+			// which is problematic when a socket has the same name as a bone.
+			FTransform BoneTransform = SourceSkelMesh->GetBoneTransform(BoneIndex);
+
+			FName ParentBone = SourceSkelMesh->GetParentBone(*BoneName);
+			int32 ParentIndex = SourceSkelMesh->GetBoneIndex(ParentBone);
+			if (ParentIndex != INDEX_NONE)
+			{
+				return BoneTransform.GetRelativeTransform(SourceSkelMesh->GetBoneTransform(ParentIndex));
+			}
+			else
+			{
+				return BoneTransform.GetRelativeTransform(SourceSkelMesh->GetComponentTransform());
+			}
+		};
+
 		if (AllowedBoneNames.Num() > 0) //If there are any entries in the allowed bones list, just use those bones names.
 		{
 			for (FName BoneName : AllowedBoneNames)
 			{
-				FTransform BoneTransform = SourceSkelMesh->GetSocketTransform(BoneName, RTS_ParentBoneSpace);
-				BoneTransforms.Add(BoneTransform);
+				BoneTransforms.Add(GetBoneTransform(INDEX_NONE, BoneName));
 			}
 		}
 		else
 		{
-			for (int32 Index = 0; Index < SourceSkelMesh->GetNumBones(); Index++)
+			for (int32 BoneIndex = 0; BoneIndex < SourceSkelMesh->GetNumBones(); BoneIndex++)
 			{
-				FTransform BoneTransform = SourceSkelMesh->GetSocketTransform(SourceSkelMesh->GetBoneName(Index), RTS_ParentBoneSpace);
-				BoneTransforms.Add(BoneTransform);
+				BoneTransforms.Add(GetBoneTransform(BoneIndex));
 			}
 		}
 	}
