LiveLink - Fix LiveLink thread-safety crash in FLiveLinkSubject::FrameSnapshot

Added FRWLock to protect FrameSnapshot access from concurrent threads.
Animation worker threads reading via EvaluateFrame_AnyThread() were racing
with Update() modifications, causing null pointer crashes during TMap copy.

#Synced-CL 51448986

diff --git a/Engine/Plugins/Animation/LiveLink/Source/LiveLink/Private/LiveLinkSubject.cpp b/Engine/Plugins/Animation/LiveLink/Source/LiveLink/Private/LiveLinkSubject.cpp
--- a/Engine/Plugins/Animation/LiveLink/Source/LiveLink/Private/LiveLinkSubject.cpp
+++ b/Engine/Plugins/Animation/LiveLink/Source/LiveLink/Private/LiveLinkSubject.cpp
@@ -82,7 +82,8 @@ void FLiveLinkSubject::Update()
 	if (bPaused)
 	{
 		// No need to process if it's paused, just keep the same frame snapshot.
-		if (HasValidFrameSnapshot())
+		FWriteScopeLock WriteLock(FrameSnapshotRWLock);
+		if (FrameSnapshot.StaticData.IsValid() && FrameSnapshot.FrameData.IsValid())
 		{
 			FrameSnapshot.FrameData.GetBaseData()->WorldTime = FPlatformTime::Seconds();
 			// Only return if we have a valid snapshot. Otherwise we'll try to recreate it, since it may have been cleared.
@@ -154,74 +155,79 @@ void FLiveLinkSubject::Update()
 	}
 
 	// Build a snapshot for this role
-	bool bSnapshotIsValid = false;
-	if (FrameData.Num() > 0)
+	// Protect FrameSnapshot with write lock during snapshot construction
 	{
-		switch(GetMode())
-		{
-		case ELiveLinkSourceMode::Timecode:
+		FWriteScopeLock WriteLock(FrameSnapshotRWLock);
+
+		bool bSnapshotIsValid = false;
+		if (FrameData.Num() > 0)
 		{
-			if (FApp::GetCurrentFrameTime().IsSet())
+			switch(GetMode())
 			{
-				//If Timecode smooth latest is enabled, use the clock offset to adjust read time. 
-				//@note For now, snapshot created will have the source time. We are not overriding its time with the read time. If a use case arise for this, we'll adjust.
-				const FQualifiedFrameTime EngineTimecode = FApp::GetCurrentFrameTime().GetValue();
-				
-				FFrameTime ClockOffset;
-				if (CachedSettings.BufferSettings.bUseTimecodeSmoothLatest)
+			case ELiveLinkSourceMode::Timecode:
+			{
+				if (FApp::GetCurrentFrameTime().IsSet())
 				{
-					ClockOffset = EngineTimecode.Rate.AsFrameTime(CachedSettings.BufferSettings.TimecodeClockOffset);
-				}
+					//If Timecode smooth latest is enabled, use the clock offset to adjust read time.
+					//@note For now, snapshot created will have the source time. We are not overriding its time with the read time. If a use case arise for this, we'll adjust.
+					const FQualifiedFrameTime EngineTimecode = FApp::GetCurrentFrameTime().GetValue();
+
+					FFrameTime ClockOffset;
+					if (CachedSettings.BufferSettings.bUseTimecodeSmoothLatest)
+					{
+						ClockOffset = EngineTimecode.Rate.AsFrameTime(CachedSettings.BufferSettings.TimecodeClockOffset);
+					}
 
-				const FFrameTime FrameOffset = FQualifiedFrameTime(FFrameTime::FromDecimal(CachedSettings.BufferSettings.TimecodeFrameOffset), LastTimecodeFrameRate).ConvertTo(EngineTimecode.Rate);
-				const FFrameTime ReadTime = EngineTimecode.Time - (ClockOffset + FrameOffset);
-				const FQualifiedFrameTime LookupQFrameTime = FQualifiedFrameTime(ReadTime, EngineTimecode.Rate);
-				bSnapshotIsValid = GetFrameAtSceneTime(LookupQFrameTime, FrameSnapshot);
+					const FFrameTime FrameOffset = FQualifiedFrameTime(FFrameTime::FromDecimal(CachedSettings.BufferSettings.TimecodeFrameOffset), LastTimecodeFrameRate).ConvertTo(EngineTimecode.Rate);
+					const FFrameTime ReadTime = EngineTimecode.Time - (ClockOffset + FrameOffset);
+					const FQualifiedFrameTime LookupQFrameTime = FQualifiedFrameTime(ReadTime, EngineTimecode.Rate);
+					bSnapshotIsValid = GetFrameAtSceneTime(LookupQFrameTime, FrameSnapshot);
+				}
+				else
+				{
+					static const FName NAME_InvalidTime = "LiveLinkSubject_NoCurrentFrameTime";
+					FLiveLinkLog::WarningOnce(NAME_InvalidTime, SubjectKey, TEXT("Can't evaluate frame for subject '%s'. The engine doesn't have a timecode value set."), *SubjectKey.SubjectName.ToString());
+				}
 			}
-			else
+			break;
+
+			case ELiveLinkSourceMode::EngineTime:
 			{
-				static const FName NAME_InvalidTime = "LiveLinkSubject_NoCurrentFrameTime";
-				FLiveLinkLog::WarningOnce(NAME_InvalidTime, SubjectKey, TEXT("Can't evaluate frame for subject '%s'. The engine doesn't have a timecode value set."), *SubjectKey.SubjectName.ToString());
-			}
-		}
-		break;
+				PRAGMA_DISABLE_DEPRECATION_WARNINGS
+				FAppTime Time = FAppTime::Get_UnsafeNoCheck();
+				PRAGMA_ENABLE_DEPRECATION_WARNINGS
 
-		case ELiveLinkSourceMode::EngineTime:
-		{
-			PRAGMA_DISABLE_DEPRECATION_WARNINGS
-			FAppTime Time = FAppTime::Get_UnsafeNoCheck();
-			PRAGMA_ENABLE_DEPRECATION_WARNINGS
+				//Compute the evaluation time using all offsets (user offset + clock offset + smooth offset)
+				const double ReadTime = Time.CurrentTime - CachedSettings.BufferSettings.EngineTimeOffset - CachedSettings.BufferSettings.EngineTimeClockOffset - CachedSettings.BufferSettings.SmoothEngineTimeOffset;
 
-			//Compute the evaluation time using all offsets (user offset + clock offset + smooth offset)
-			const double ReadTime = Time.CurrentTime - CachedSettings.BufferSettings.EngineTimeOffset - CachedSettings.BufferSettings.EngineTimeClockOffset - CachedSettings.BufferSettings.SmoothEngineTimeOffset;
+				UE_LOG(LogLiveLink, Verbose, TEXT("Subject '%s' Read Time: %lf    User Offset: %.4lf    Clock Offset: %.4lf    Smooth Offset: %.4lf")
+					, *SubjectKey.SubjectName.ToString()
+					, ReadTime
+					, CachedSettings.BufferSettings.EngineTimeOffset
+					, CachedSettings.BufferSettings.EngineTimeClockOffset
+					, CachedSettings.BufferSettings.SmoothEngineTimeOffset);
 
-			UE_LOG(LogLiveLink, Verbose, TEXT("Subject '%s' Read Time: %lf    User Offset: %.4lf    Clock Offset: %.4lf    Smooth Offset: %.4lf")
-				, *SubjectKey.SubjectName.ToString()
-				, ReadTime
-				, CachedSettings.BufferSettings.EngineTimeOffset
-				, CachedSettings.BufferSettings.EngineTimeClockOffset
-				, CachedSettings.BufferSettings.SmoothEngineTimeOffset);
+				bSnapshotIsValid = GetFrameAtWorldTime(ReadTime, FrameSnapshot);
+				break;
+			}
 
-			bSnapshotIsValid = GetFrameAtWorldTime(ReadTime, FrameSnapshot);
-			break;
+			case ELiveLinkSourceMode::Latest:
+			default:
+				bSnapshotIsValid = GetLatestFrame(FrameSnapshot);
+			}
 		}
 
-		case ELiveLinkSourceMode::Latest:
-		default:
-			bSnapshotIsValid = GetLatestFrame(FrameSnapshot);
+		if (!bSnapshotIsValid)
+		{
+			// Invalidate the snapshot
+			FrameSnapshot.FrameData.Reset();
+		}
+		else
+		{
+			//Extra verbose logging about snapshot just created
+			const FTimecode FrameTC = FTimecode::FromFrameNumber(FrameSnapshot.FrameData.GetBaseData()->MetaData.SceneTime.Time.GetFrame(), FrameSnapshot.FrameData.GetBaseData()->MetaData.SceneTime.Rate);
+			UE_LOG(LogLiveLink, Verbose, TEXT("Subject '%s' created snapshot with frame ID %d - Timecode:[%s.%0.3f] - SourceTime: %0.4f, Offset: %0.6f, CorrectedTime: %0.4f"), *SubjectKey.SubjectName.ToString(), FrameSnapshot.FrameData.GetBaseData()->FrameId, *FrameTC.ToString(), FrameSnapshot.FrameData.GetBaseData()->MetaData.SceneTime.Time.GetSubFrame(), FrameSnapshot.FrameData.GetBaseData()->WorldTime.GetSourceTime(), FrameSnapshot.FrameData.GetBaseData()->WorldTime.GetOffset(), FrameSnapshot.FrameData.GetBaseData()->WorldTime.GetOffsettedTime());
 		}
-	}
-
-	if (!bSnapshotIsValid)
-	{
-		// Invalidate the snapshot
-		FrameSnapshot.FrameData.Reset();
-	}
-	else
-	{
-		//Extra verbose logging about snapshot just created
-		const FTimecode FrameTC = FTimecode::FromFrameNumber(FrameSnapshot.FrameData.GetBaseData()->MetaData.SceneTime.Time.GetFrame(), FrameSnapshot.FrameData.GetBaseData()->MetaData.SceneTime.Rate);
-		UE_LOG(LogLiveLink, Verbose, TEXT("Subject '%s' created snapshot with frame ID %d - Timecode:[%s.%0.3f] - SourceTime: %0.4f, Offset: %0.6f, CorrectedTime: %0.4f"), *SubjectKey.SubjectName.ToString(), FrameSnapshot.FrameData.GetBaseData()->FrameId, *FrameTC.ToString(), FrameSnapshot.FrameData.GetBaseData()->MetaData.SceneTime.Time.GetSubFrame(), FrameSnapshot.FrameData.GetBaseData()->WorldTime.GetSourceTime(), FrameSnapshot.FrameData.GetBaseData()->WorldTime.GetOffset(), FrameSnapshot.FrameData.GetBaseData()->WorldTime.GetOffsettedTime());
 	}
 
 	UpdateCachedOldestDataTime();
@@ -233,14 +239,23 @@ void FLiveLinkSubject::Update()
 
 void FLiveLinkSubject::ClearFrames()
 {
+	FWriteScopeLock WriteLock(FrameSnapshotRWLock);
 	FrameSnapshot.StaticData.Reset();
 	FrameSnapshot.FrameData.Reset();
 	FrameData.Reset();
 	ReceivedOrderedFrames.Empty();
 }
 
+bool FLiveLinkSubject::EvaluateFrame(TSubclassOf<ULiveLinkRole> InDesiredRole, FLiveLinkSubjectFrameData& OutFrame)
+{
+	// Protect FrameSnapshot access with read lock to allow concurrent reads from multiple animation threads
+	FReadScopeLock ReadLock(FrameSnapshotRWLock);
+	return ILiveLinkSubject::EvaluateFrame(InDesiredRole, OutFrame);
+}
+
 bool FLiveLinkSubject::HasValidFrameSnapshot() const
 {
+	FReadScopeLock ReadLock(FrameSnapshotRWLock);
 	return FrameSnapshot.StaticData.IsValid() && FrameSnapshot.FrameData.IsValid();
 }
 
@@ -1247,10 +1262,14 @@ void FLiveLinkSubject::CacheSettings(ULiveLinkSourceSettings* SourceSetting, ULi
 
 						if (NewRemapper->IsRemapperCompatible(StaticData))
 						{
-							if (HasValidFrameSnapshot())
+							// Protect FrameSnapshot access during remapping
 							{
-								NewRemapper->RemapStaticData(FrameSnapshot.StaticData);
-								NewRemapper->RemapFrameData(FrameSnapshot.StaticData, FrameSnapshot.FrameData);
+								FWriteScopeLock WriteLock(FrameSnapshotRWLock);
+								if (FrameSnapshot.StaticData.IsValid() && FrameSnapshot.FrameData.IsValid())
+								{
+									NewRemapper->RemapStaticData(FrameSnapshot.StaticData);
+									NewRemapper->RemapFrameData(FrameSnapshot.StaticData, FrameSnapshot.FrameData);
+								}
 							}
 
 							// Important: Because we changed the static data, we have to update the frame data to match the new static data.
diff --git a/Engine/Plugins/Animation/LiveLink/Source/LiveLink/Public/LiveLinkSubject.h b/Engine/Plugins/Animation/LiveLink/Source/LiveLink/Public/LiveLinkSubject.h
--- a/Engine/Plugins/Animation/LiveLink/Source/LiveLink/Public/LiveLinkSubject.h
+++ b/Engine/Plugins/Animation/LiveLink/Source/LiveLink/Public/LiveLinkSubject.h
@@ -102,6 +102,11 @@ public:
 protected:
 	virtual const FLiveLinkSubjectFrameData& GetFrameSnapshot() const override { return FrameSnapshot; }
 	friend class FLiveLinkClient;
+
+public:
+	UE_API virtual bool EvaluateFrame(TSubclassOf<ULiveLinkRole> InDesiredRole, FLiveLinkSubjectFrameData& OutFrame) override;
+
+protected:
 	//~ End ILiveLinkSubject Interface
 
 	//~Begin ITimedDataSource Interface
@@ -323,6 +328,9 @@ private:
 
 	/** Used to protect access to translators, preprocessors and interpolation processors since they can be set and accessed in different threads. */
 	mutable FCriticalSection SettingsCriticalSection;
+
+	/** Read-write lock to protect FrameSnapshot access. Allows multiple concurrent readers (EvaluateFrame_AnyThread) but exclusive writer access (Update). */
+	mutable FRWLock FrameSnapshotRWLock;
 };
 
 #undef UE_API
