LiveLink - Allow removing subject when disabling broadcast component

#Synced-CL 51448986

diff --git a/Engine/Plugins/Animation/LiveLink/Source/LiveLink/Private/LiveLinkBroadcastComponent.cpp b/Engine/Plugins/Animation/LiveLink/Source/LiveLink/Private/LiveLinkBroadcastComponent.cpp
--- a/Engine/Plugins/Animation/LiveLink/Source/LiveLink/Private/LiveLinkBroadcastComponent.cpp
+++ b/Engine/Plugins/Animation/LiveLink/Source/LiveLink/Private/LiveLinkBroadcastComponent.cpp
@@ -53,7 +53,19 @@ void ULiveLinkBroadcastComponent::PostEditChangeProperty(struct FPropertyChanged
 		|| PropertyChangedEvent.GetPropertyName() == GET_MEMBER_NAME_CHECKED(ULiveLinkBroadcastComponent, Role)
 		|| PropertyChangedEvent.GetPropertyName() == GET_MEMBER_NAME_CHECKED(ULiveLinkBroadcastComponent, AllowedCurveNames)
 		|| PropertyChangedEvent.GetPropertyName() == GET_MEMBER_NAME_CHECKED(ULiveLinkBroadcastComponent, AllowedBoneNames);
-	
+
+	if (PropertyChangedEvent.GetPropertyName() == GET_MEMBER_NAME_CHECKED(ULiveLinkBroadcastComponent, bEnable))
+	{
+		if (!bEnable && bRemoveSubjectWhenDisabled)
+		{
+			ULiveLinkBroadcastSubsystem* BroadcastSubsystem = GEngine->GetEngineSubsystem<ULiveLinkBroadcastSubsystem>();
+			BroadcastSubsystem->RemoveSubject(SubjectKey);
+		}
+		else
+		{
+			MarkDirty();
+		}
+	}
 }
 
 void ULiveLinkBroadcastComponent::OnObjectModified(UObject* ModifiedObject)
@@ -124,11 +136,7 @@ void ULiveLinkBroadcastComponent::TickComponent(float DeltaTime, ELevelTick Tick
 	{
 		if (SubjectName == NAME_None)
 		{
-#if WITH_EDITOR
-			SubjectName = *GetTypedOuter<AActor>()->GetActorLabel();
-#else
-			SubjectName = GetTypedOuter<AActor>()->GetFName();
-#endif
+			SubjectName = *GetTypedOuter<AActor>()->GetActorNameOrLabel();
 		}
 
 		CachedResolvedObject = SourceMesh.GetComponent(GetOwner());
@@ -143,6 +151,21 @@ void ULiveLinkBroadcastComponent::TickComponent(float DeltaTime, ELevelTick Tick
 	BroadcastFrameData();
 }
 
+void ULiveLinkBroadcastComponent::SetEnabled(bool bInEnabled)
+{
+	bEnable = bInEnabled;
+
+	if (bRemoveSubjectWhenDisabled && !bEnable)
+	{
+		ULiveLinkBroadcastSubsystem* BroadcastSubsystem = GEngine->GetEngineSubsystem<ULiveLinkBroadcastSubsystem>();
+		BroadcastSubsystem->RemoveSubject(SubjectKey);
+	}
+	else
+	{
+		MarkDirty();
+	}
+}
+
 void ULiveLinkBroadcastComponent::InitiateLiveLink()
 {
 	if (Role && (!SourceMesh.PathToComponent.IsEmpty() || !SourceMesh.ComponentProperty.IsNone()))
diff --git a/Engine/Plugins/Animation/LiveLink/Source/LiveLink/Public/LiveLinkBroadcastComponent.h b/Engine/Plugins/Animation/LiveLink/Source/LiveLink/Public/LiveLinkBroadcastComponent.h
--- a/Engine/Plugins/Animation/LiveLink/Source/LiveLink/Public/LiveLinkBroadcastComponent.h
+++ b/Engine/Plugins/Animation/LiveLink/Source/LiveLink/Public/LiveLinkBroadcastComponent.h
@@ -59,6 +59,10 @@ public:
 	UPROPERTY(BlueprintReadWrite, EditAnywhere, Category = "Live Link Broadcast", meta=(ToolTip="Curve data will be added to the broadcast data if listed here"))
 	TArray<FName> AllowedCurveNames;
 
+	/** Whether the livelink subject should be removed when broadcasting is disabled. */
+	UPROPERTY(BlueprintReadWrite, EditAnywhere, Category = "Live Link Broadcast")
+	bool bRemoveSubjectWhenDisabled = false;
+
 	//~ UActorComponent interface
 	UE_API virtual void TickComponent(float DeltaTime, ELevelTick TickType, FActorComponentTickFunction* ThisTickFunction) override;
 
@@ -68,6 +72,9 @@ public:
 		bIsDirty = true;
 	}
 
+	/** Sets whether this component should broadcast live link data. */
+	UE_API void SetEnabled(bool bInEnabled);
+
 private:
 	/** Creates the rebroadcast source if it doesn't exist. */
 	void InitiateLiveLink();
