LiveLinkBroadcastComponent - Fix static data not being broadcast when re-creating component with same subject name after removing the previous one

#Synced-CL 51719765

diff --git a/Engine/Plugins/Animation/LiveLink/Source/LiveLink/Private/LiveLinkBroadcastSubsystem.cpp b/Engine/Plugins/Animation/LiveLink/Source/LiveLink/Private/LiveLinkBroadcastSubsystem.cpp
--- a/Engine/Plugins/Animation/LiveLink/Source/LiveLink/Private/LiveLinkBroadcastSubsystem.cpp
+++ b/Engine/Plugins/Animation/LiveLink/Source/LiveLink/Private/LiveLinkBroadcastSubsystem.cpp
@@ -78,15 +78,15 @@ FLiveLinkSubjectKey ULiveLinkBroadcastSubsystem::CreateSubject(FName SubjectName
 
 		SubjectKey = FLiveLinkSubjectKey{ VirtualSourceGUID, SubjectName };
 
-		// Check that the subject exists before recreating it.
-		if (!LiveLinkClient->GetSubjectSettings(SubjectKey))
+		FLiveLinkSubjectPreset Preset;
+		Preset.bEnabled = true;
+		Preset.Key = SubjectKey;
+		Preset.Role = Role;
+
+		// CreateSubject handles the pending-kill case (removes and recreates),
+		// and returns false if the subject already exists and is alive.
+		if (LiveLinkClient->CreateSubject(Preset))
 		{
-			FLiveLinkSubjectPreset Preset;
-			Preset.bEnabled = true;
-			Preset.Key = SubjectKey;
-			Preset.Role = Role;
-
-			LiveLinkClient->CreateSubject(Preset);
 			ULiveLinkSubjectSettings* Settings = CastChecked<ULiveLinkSubjectSettings>(LiveLinkClient->GetSubjectSettings(SubjectKey));
 			Settings->bRebroadcastSubject = true;
 		}
