LiveLink - Add option to re-enable shadowed livelink subject if the other subject is removed

In certain cases, you can end up with two LiveLink subjects with the same name, resulting in one of them being disabled. This CL adds an option (Re-enable shadowed subjects on removal), which automatically re-enables subjects that were disabled because they were being shadowed once the original subject is removed.

#Synced-CL 51719765

diff --git a/Engine/Plugins/Animation/LiveLink/Source/LiveLink/Private/LiveLinkSourceCollection.cpp b/Engine/Plugins/Animation/LiveLink/Source/LiveLink/Private/LiveLinkSourceCollection.cpp
--- a/Engine/Plugins/Animation/LiveLink/Source/LiveLink/Private/LiveLinkSourceCollection.cpp
+++ b/Engine/Plugins/Animation/LiveLink/Source/LiveLink/Private/LiveLinkSourceCollection.cpp
@@ -5,6 +5,7 @@
 #include "EngineAnalytics.h"
 #include "Experimental/Async/MultiUniqueLock.h"
 #include "IAnalyticsProviderET.h"
+#include "LiveLinkSettings.h"
 #include "LiveLinkVirtualSource.h"
 #include "UObject/Package.h"
 
@@ -250,13 +251,38 @@ void FLiveLinkSourceCollection::AddSubject(FLiveLinkCollectionSubjectItem InSubj
 
 void FLiveLinkSourceCollection::RemoveSubject(FLiveLinkSubjectKey InSubjectKey)
 {
+	bool bWasEnabled = false;
+
+	TOptional<FLiveLinkSubjectKey> ReEnabledSubject;
 	{
 		UE::TUniqueLock Lock(SubjectsLock);
 		int32 IndexOf = Subjects.IndexOfByPredicate([InSubjectKey](const FLiveLinkCollectionSubjectItem& Other) { return Other.Key == InSubjectKey; });
 		if (IndexOf != INDEX_NONE)
 		{
+			bWasEnabled = Subjects[IndexOf].bEnabled;
 			Subjects.RemoveAtSwap(IndexOf);
 		}
+
+		// If the removed subject was enabled and the setting is on, re-enable another subject
+		// with the same name so it can take over. This prevents the case where removing a subject
+		// that was shadowing another (same name, different source) leaves the subject disabled.
+		if (bWasEnabled && GetDefault<ULiveLinkSettings>()->bReEnableSubjectOnRemoval)
+		{
+			for (FLiveLinkCollectionSubjectItem& SubjectItem : Subjects)
+			{
+				if (SubjectItem.Key.SubjectName == InSubjectKey.SubjectName && !SubjectItem.bPendingKill)
+				{
+					SubjectItem.bEnabled = true;
+					ReEnabledSubject = SubjectItem.Key;
+					break;
+				}
+			}
+		}
+	}
+
+	if (ReEnabledSubject)
+	{
+		BroadcastOnGameThread(OnLiveLinkSubjectEnabledChanged(), *ReEnabledSubject, true);
 	}
 
 	BroadcastOnGameThread(OnLiveLinkSubjectRemoved(), InSubjectKey);
diff --git a/Engine/Plugins/Animation/LiveLink/Source/LiveLink/Public/LiveLinkSettings.h b/Engine/Plugins/Animation/LiveLink/Source/LiveLink/Public/LiveLinkSettings.h
--- a/Engine/Plugins/Animation/LiveLink/Source/LiveLink/Public/LiveLinkSettings.h
+++ b/Engine/Plugins/Animation/LiveLink/Source/LiveLink/Public/LiveLinkSettings.h
@@ -129,6 +129,10 @@ public:
 	UPROPERTY(config, EditAnywhere, AdvancedDisplay, Category = "LiveLink", meta=(ForceUnits=s))
 	double MessageBusTimeBeforeRemovingInactiveSource;
 
+	/** When an enabled subject is removed and another subject with the same name exists, automatically enable it. */
+	UPROPERTY(config, EditAnywhere, AdvancedDisplay, Category = "LiveLink", meta = (DisplayName = "Re-enable shadowed subjects on removal"))
+	bool bReEnableSubjectOnRemoval = false;
+
 	/** Whether to Preprocess frames before rebroadcasting them. */
 	UPROPERTY(config)
 	bool bPreProcessRebroadcastFrames = false;
