LiveLinkHub - Refactor data sources creation code to simplify RefreshSubjects code

#Synced-CL 51719765

diff --git a/Engine/Plugins/Animation/LiveLink/Source/LiveLinkEditor/Private/LiveLinkClientPanelViews.cpp b/Engine/Plugins/Animation/LiveLink/Source/LiveLinkEditor/Private/LiveLinkClientPanelViews.cpp
--- a/Engine/Plugins/Animation/LiveLink/Source/LiveLinkEditor/Private/LiveLinkClientPanelViews.cpp
+++ b/Engine/Plugins/Animation/LiveLink/Source/LiveLinkEditor/Private/LiveLinkClientPanelViews.cpp
@@ -1591,14 +1591,8 @@ bool FLiveLinkSubjectsView::HasAnySources() const
 
 void FLiveLinkSubjectsView::RefreshSubjects()
 {
-	TArray<TPair<FLiveLinkSubjectKey, bool>> SavedSelection;
-	{
-		TArray<FLiveLinkSubjectUIEntryPtr> SelectedItems = SubjectsTreeView->GetSelectedItems();
-		for (const FLiveLinkSubjectUIEntryPtr& SelectedItem : SelectedItems)
-		{
-			SavedSelection.Add(TPair<FLiveLinkSubjectKey, bool>(SelectedItem->SubjectKey, SelectedItem->IsSubject()));
-		}
-	}
+	TArray<FSavedSelection> SavedSelection;
+	SaveSelection(SavedSelection);
 
 	if (IModularFeatures::Get().IsModularFeatureAvailable(ILiveLinkClient::ModularFeatureName))
 	{
@@ -1611,69 +1605,51 @@ void FLiveLinkSubjectsView::RefreshSubjects()
 			TArray<FGuid> Sources = Client->GetSources();
 
 			TMap<FGuid, FLiveLinkSubjectUIEntryPtr> SourceItems;
-
-			TArray<FLiveLinkSubjectUIEntryPtr> AllItems;
-			AllItems.Reserve(SubjectKeys.Num());
+			FLiveLinkClient* ClientPtr = static_cast<FLiveLinkClient*>(Client);
 
 			SubjectData.Reset();
-			
-			// Add Subjects & Sources
-			for (const FLiveLinkSubjectKey& SubjectKey : SubjectKeys)
-			{
-				FName SourceNameOverride = *Client->GetSourceNameOverride(SubjectKey).ToString();
-
-				FLiveLinkSubjectUIEntryPtr Source = SourceItems.FindRef(SubjectKey.Source);
 
-				if (!Source)
+			// Lazily create a header (if it doesn't exist yet) and add Child under it.
+			auto AddChildUnderHeader = [this](FLiveLinkSubjectUIEntryPtr& Header, ELiveLinkEntryType HeaderType, FLiveLinkClient* InClient, const FLiveLinkSubjectUIEntryPtr& Child)
+			{
+				if (!Header)
 				{
-					constexpr bool bIsSource = true;
-					Source = MakeShared<FLiveLinkSubjectUIEntry>(SubjectKey, static_cast<FLiveLinkClient*>(Client), bIsSource);
-
-					if (!SourcesHeader)
-					{
-						// It's possible that no source were created such as in Playback.
-						SourcesHeader = MakeShared<FLiveLinkSubjectUIEntry>(static_cast<FLiveLinkClient*>(Client), ELiveLinkEntryType::SourcesHeader);
-						SubjectData.Add(SourcesHeader);
-					}
-
-					SourcesHeader->Children.Add(Source);
-					SourceItems.Add(SubjectKey.Source) = Source;
-
-					SubjectsTreeView->SetItemExpansion(Source, true);
-					AllItems.Add(Source);
+					Header = MakeShared<FLiveLinkSubjectUIEntry>(InClient, HeaderType);
+					SubjectData.Add(Header);
 				}
+				Header->Children.Add(Child);
+			};
 
-				FLiveLinkSubjectUIEntryPtr SubjectEntry = MakeShared<FLiveLinkSubjectUIEntry>(SubjectKey, static_cast<FLiveLinkClient*>(Client));
-				Source->Children.Add(SubjectEntry);
-				AllItems.Add(SubjectEntry);
-			}
-
-			// Add sources that haven't been added through subjects.
-			// @Note: This pass must come after subjects were added, because subjects can override their source name, so we add them 
-			// in the previous pass to avoid having to rename them.
-			for (const FGuid& SourceId : Sources)
+			// Find an existing source entry or create a new one under the SourcesHeader.
+			auto GetOrCreateSource = [&](const FGuid& SourceGuid, const FLiveLinkSubjectKey& SubjectKey) -> FLiveLinkSubjectUIEntryPtr
 			{
-				if (SourceItems.Contains(SourceId))
+				if (FLiveLinkSubjectUIEntryPtr* Existing = SourceItems.Find(SourceGuid))
 				{
-					continue;
+					return *Existing;
 				}
-
 				constexpr bool bIsSource = true;
-				FLiveLinkSubjectUIEntryPtr Source = MakeShared<FLiveLinkSubjectUIEntry>(FLiveLinkSubjectKey{ SourceId, NAME_None }, static_cast<FLiveLinkClient*>(Client), bIsSource);
-
-				if (!SourcesHeader)
-				{
-					SourcesHeader = MakeShared<FLiveLinkSubjectUIEntry>(static_cast<FLiveLinkClient*>(Client), ELiveLinkEntryType::SourcesHeader);
-					SubjectData.Add(SourcesHeader);
-				}
+				FLiveLinkSubjectUIEntryPtr Source = MakeShared<FLiveLinkSubjectUIEntry>(SubjectKey, ClientPtr, bIsSource);
+				AddChildUnderHeader(SourcesHeader, ELiveLinkEntryType::SourcesHeader, ClientPtr, Source);
+				SourceItems.Add(SourceGuid, Source);
+				SubjectsTreeView->SetItemExpansion(Source, true);
+				return Source;
+			};
 
-				SourcesHeader->Children.Add(Source);
+			// Add Subjects & Sources
+			for (const FLiveLinkSubjectKey& SubjectKey : SubjectKeys)
+			{
+				FLiveLinkSubjectUIEntryPtr Source = GetOrCreateSource(SubjectKey.Source, SubjectKey);
 
-				SubjectsTreeView->SetItemExpansion(Source, true);
-				FText SourceType = Client->GetSourceType(SourceId);
+				FLiveLinkSubjectUIEntryPtr SubjectEntry = MakeShared<FLiveLinkSubjectUIEntry>(SubjectKey, ClientPtr);
+				Source->Children.Add(SubjectEntry);
+			}
 
-				SourceItems.Add(SourceId, Source);
-				AllItems.Add(Source);
+			// Add sources that haven't been added through subjects.
+			// @Note: This pass must come after subjects were added, because subjects can override
+			// their source name, so we add them in the previous pass to avoid having to rename them.
+			for (const FGuid& SourceId : Sources)
+			{
+				GetOrCreateSource(SourceId, FLiveLinkSubjectKey{ SourceId, NAME_None });
 			}
 
 			// Add Devices (Only visible in LiveLinkHub)
@@ -1684,42 +1660,26 @@ void FLiveLinkSubjectsView::RefreshSubjects()
 					const TMap<FGuid, TObjectPtr<ULiveLinkDevice>>& DeviceMap = Subsystem->GetDeviceMap();
 					for (const TPair<FGuid, TObjectPtr<ULiveLinkDevice>>& DevicePair : DeviceMap)
 					{
-						if (!DevicesHeader)
-						{
-							DevicesHeader = MakeShared<FLiveLinkSubjectUIEntry>(static_cast<FLiveLinkClient*>(Client), ELiveLinkEntryType::DevicesHeader);
-							SubjectData.Add(DevicesHeader);
-						}
-
-						FLiveLinkSubjectUIEntryPtr Entry = MakeShared<FLiveLinkSubjectUIEntry>(static_cast<FLiveLinkClient*>(Client), ELiveLinkEntryType::Device);
+						FLiveLinkSubjectUIEntryPtr Entry = MakeShared<FLiveLinkSubjectUIEntry>(ClientPtr, ELiveLinkEntryType::Device);
 						Entry->WeakDevice = DevicePair.Value;
-
-						DevicesHeader->Children.Add(Entry);
+						AddChildUnderHeader(DevicesHeader, ELiveLinkEntryType::DevicesHeader, ClientPtr, Entry);
 					}
 				}
 			}
 
-			auto SortPredicate = [](const FLiveLinkSubjectUIEntryPtr& LHS, const FLiveLinkSubjectUIEntryPtr& RHS) {return LHS->GetItemText().CompareTo(RHS->GetItemText()) < 0; };
+			// Sort all levels of the tree.
+			auto SortPredicate = [](const FLiveLinkSubjectUIEntryPtr& LHS, const FLiveLinkSubjectUIEntryPtr& RHS) { return LHS->GetItemText().CompareTo(RHS->GetItemText()) < 0; };
 			SubjectData.Sort(SortPredicate);
-			for (FLiveLinkSubjectUIEntryPtr& Subject : SubjectData)
+			for (FLiveLinkSubjectUIEntryPtr& Header : SubjectData)
 			{
-				Subject->Children.Sort(SortPredicate);
-				for (const FLiveLinkSubjectUIEntryPtr& Children : Subject->Children)
+				Header->Children.Sort(SortPredicate);
+				for (const FLiveLinkSubjectUIEntryPtr& Child : Header->Children)
 				{
-					Children->Children.Sort(SortPredicate);
+					Child->Children.Sort(SortPredicate);
 				}
 			}
 
-			for (const FLiveLinkSubjectUIEntryPtr& Item : AllItems)
-			{
-				for (const TPair<FLiveLinkSubjectKey, bool>& Selection : SavedSelection)
-				{
-					if (Item->SubjectKey == Selection.Key && Item->IsSubject() == Selection.Value)
-					{
-						SubjectsTreeView->SetItemSelection(Item, true);
-						break;
-					}
-				}
-			}
+			RestoreSelection(SavedSelection);
 
 			if (FilterSearchBox)
 			{
@@ -1729,6 +1689,68 @@ void FLiveLinkSubjectsView::RefreshSubjects()
 	}
 }
 
+void FLiveLinkSubjectsView::SaveSelection(TArray<FSavedSelection>& OutSavedSelection) const
+{
+	TArray<FLiveLinkSubjectUIEntryPtr> SelectedItems = SubjectsTreeView->GetSelectedItems();
+	OutSavedSelection.Reserve(SelectedItems.Num());
+	for (const FLiveLinkSubjectUIEntryPtr& SelectedItem : SelectedItems)
+	{
+		FSavedSelection Entry;
+		Entry.SubjectKey = SelectedItem->SubjectKey;
+		Entry.Type = SelectedItem->Type;
+		Entry.WeakDevice = SelectedItem->WeakDevice;
+		OutSavedSelection.Add(MoveTemp(Entry));
+	}
+}
+
+void FLiveLinkSubjectsView::RestoreSelection(const TArray<FSavedSelection>& InSavedSelection)
+{
+	SubjectsTreeView->ClearSelection();
+
+	for (const FLiveLinkSubjectUIEntryPtr& HeaderEntry : SubjectData)
+	{
+		for (const FLiveLinkSubjectUIEntryPtr& Child : HeaderEntry->Children)
+		{
+			RestoreSelectionForItem(Child, InSavedSelection);
+			for (const FLiveLinkSubjectUIEntryPtr& GrandChild : Child->Children)
+			{
+				RestoreSelectionForItem(GrandChild, InSavedSelection);
+			}
+		}
+	}
+}
+
+void FLiveLinkSubjectsView::RestoreSelectionForItem(const FLiveLinkSubjectUIEntryPtr& Item, const TArray<FSavedSelection>& InSavedSelection)
+{
+	for (const FSavedSelection& Selection : InSavedSelection)
+	{
+		if (Item->Type != Selection.Type)
+		{
+			continue;
+		}
+
+		bool bMatch = false;
+		if (Item->IsSubject())
+		{
+			bMatch = (Item->SubjectKey == Selection.SubjectKey);
+		}
+		else if (Item->IsSource())
+		{
+			bMatch = (Item->SubjectKey.Source == Selection.SubjectKey.Source);
+		}
+		else if (Item->IsDevice())
+		{
+			bMatch = (Item->WeakDevice == Selection.WeakDevice);
+		}
+
+		if (bMatch)
+		{
+			SubjectsTreeView->SetItemSelection(Item, true);
+			break;
+		}
+	}
+}
+
 bool FLiveLinkSubjectsView::CanPauseSubject() const
 {
 	TArray<FLiveLinkSubjectUIEntryPtr> Selected;
diff --git a/Engine/Plugins/Animation/LiveLink/Source/LiveLinkEditor/Public/LiveLinkClientPanelViews.h b/Engine/Plugins/Animation/LiveLink/Source/LiveLinkEditor/Public/LiveLinkClientPanelViews.h
--- a/Engine/Plugins/Animation/LiveLink/Source/LiveLinkEditor/Public/LiveLinkClientPanelViews.h
+++ b/Engine/Plugins/Animation/LiveLink/Source/LiveLinkEditor/Public/LiveLinkClientPanelViews.h
@@ -272,6 +272,24 @@ public:
 	UE_API TSharedRef<SWidget> GetWidget();
 
 private:
+	// Captures the identity of a selected item so it can be matched after the list is rebuilt.
+	struct FSavedSelection
+	{
+		// Full subject key (used for subject and source matching).
+		FLiveLinkSubjectKey SubjectKey;
+		// Type of the entry.
+		ELiveLinkEntryType Type;
+		// Weak reference to the device, used to match device entries.
+		TWeakObjectPtr<class ULiveLinkDevice> WeakDevice;
+	};
+
+	// Saves the current tree view selection into OutSavedSelection.
+	void SaveSelection(TArray<FSavedSelection>& OutSavedSelection) const;
+	// Restores a previously saved selection by matching entries in the rebuilt tree.
+	void RestoreSelection(const TArray<FSavedSelection>& InSavedSelection);
+	// Selects Item if it matches any entry in InSavedSelection.
+	void RestoreSelectionForItem(const FLiveLinkSubjectUIEntryPtr& Item, const TArray<FSavedSelection>& InSavedSelection);
+
 	void CreateCombinedTreeView(const TSharedPtr<FUICommandList>& InCommandList);
 	// Get the label for the pause subject context menu entry.
 	FText GetPauseSubjectLabel() const;
