LiveLink - Add a SetSubjectEnabled method that takes a FLiveLinkSubjectName parameter

#rnx
#Synced-CL 51719765

diff --git a/Engine/Plugins/Animation/LiveLink/Source/LiveLink/Private/LiveLinkClient.cpp b/Engine/Plugins/Animation/LiveLink/Source/LiveLink/Private/LiveLinkClient.cpp
--- a/Engine/Plugins/Animation/LiveLink/Source/LiveLink/Private/LiveLinkClient.cpp
+++ b/Engine/Plugins/Animation/LiveLink/Source/LiveLink/Private/LiveLinkClient.cpp
@@ -427,16 +427,16 @@ bool FLiveLinkClient::HandleSubjectRebroadcast(ILiveLinkSubject* InSubject, cons
 			{
 				RebroadcastLiveLinkProvider = GetRebroadcastLiveLinkProvider();
 			}
-				
+
 			if (RebroadcastLiveLinkProvider.IsValid())
 			{
 				// Make a copy of the data for use by the rebroadcaster
 				FLiveLinkFrameDataStruct FrameDataCopy;
 				FrameDataCopy.InitializeWith(InFrameData);
-				
+
 				FLiveLinkStaticDataStruct StaticDataCopy;
 				StaticDataCopy.InitializeWith(InSubject->GetStaticData());
-				
+
 				if (bPreProcessRebroadcastFrames)
 				{
 					InSubject->PreprocessFrame(StaticDataCopy, FrameDataCopy);
@@ -1699,6 +1699,14 @@ void FLiveLinkClient::SetSubjectEnabled(const FLiveLinkSubjectKey& InSubjectKey,
 	Collection->SetSubjectEnabled(InSubjectKey, bInEnabled);
 }
 
+void FLiveLinkClient::SetSubjectEnabled(FLiveLinkSubjectName SubjectName, bool bInEnabled)
+{
+	if (const FLiveLinkSubjectKey* FoundSubjectKey = EnabledSubjects.Find(SubjectName))
+	{
+		Collection->SetSubjectEnabled(*FoundSubjectKey, bInEnabled);
+	}
+}
+
 bool FLiveLinkClient::IsSubjectTimeSynchronized(const FLiveLinkSubjectKey& InSubjectKey) const
 {
 	if (const FLiveLinkCollectionSubjectItem* SubjectItem = Collection->FindSubject(InSubjectKey))
diff --git a/Engine/Plugins/Animation/LiveLink/Source/LiveLink/Public/LiveLinkClient.h b/Engine/Plugins/Animation/LiveLink/Source/LiveLink/Public/LiveLinkClient.h
--- a/Engine/Plugins/Animation/LiveLink/Source/LiveLink/Public/LiveLinkClient.h
+++ b/Engine/Plugins/Animation/LiveLink/Source/LiveLink/Public/LiveLinkClient.h
@@ -100,6 +100,7 @@ public:
 	UE_API virtual bool IsSubjectEnabled(const FLiveLinkSubjectKey& SubjectKey, bool bForThisFrame) const override;
 	UE_API virtual bool IsSubjectEnabled(FLiveLinkSubjectName SubjectName) const override;
 	UE_API virtual void SetSubjectEnabled(const FLiveLinkSubjectKey& SubjectKey, bool bEnabled) override;
+	UE_API virtual void SetSubjectEnabled(FLiveLinkSubjectName SubjectName, bool bEnabled) override;
 	UE_API virtual bool IsSubjectTimeSynchronized(const FLiveLinkSubjectKey& SubjectKey) const override;
 	UE_API virtual bool IsSubjectTimeSynchronized(FLiveLinkSubjectName SubjectName) const override;
 	UE_API virtual bool IsVirtualSubject(const FLiveLinkSubjectKey& SubjectKey) const override;
diff --git a/Engine/Source/Runtime/LiveLinkInterface/Public/ILiveLinkClient.h b/Engine/Source/Runtime/LiveLinkInterface/Public/ILiveLinkClient.h
--- a/Engine/Source/Runtime/LiveLinkInterface/Public/ILiveLinkClient.h
+++ b/Engine/Source/Runtime/LiveLinkInterface/Public/ILiveLinkClient.h
@@ -203,6 +203,15 @@ public:
 	 * SetSubjectEnabled will take effect on the next frame.
 	 */
 	virtual void SetSubjectEnabled(const FLiveLinkSubjectKey& SubjectKey, bool bEnabled) = 0;
+
+	/**
+	 * Set the subject's from a specific source to enabled, disabling the other in the process.
+	 * Only 1 subject with the same name can be enabled.
+	 * At the start of the frame, a snapshot of the enabled subjects will be made.
+	 * That snapshot dictate which subject will be used for the duration of that frame.
+	 * SetSubjectEnabled will take effect on the next frame.
+	 */
+	virtual void SetSubjectEnabled(FLiveLinkSubjectName SubjectName, bool bEnabled) = 0;
 	
 	/** Whether or not the subject's data, from a specific source, is time synchronized or not */
 	virtual bool IsSubjectTimeSynchronized(const FLiveLinkSubjectKey& SubjectKey) const = 0;
