Make the billboard component non editable in the details panel.

#Synced-CL 52173051

diff --git a/Engine/Plugins/Animation/LiveLink/Source/LiveLink/Public/Visualizers/LiveLinkDataPreview.h b/Engine/Plugins/Animation/LiveLink/Source/LiveLink/Public/Visualizers/LiveLinkDataPreview.h
--- a/Engine/Plugins/Animation/LiveLink/Source/LiveLink/Public/Visualizers/LiveLinkDataPreview.h
+++ b/Engine/Plugins/Animation/LiveLink/Source/LiveLink/Public/Visualizers/LiveLinkDataPreview.h
@@ -53,7 +53,7 @@ public:
 	TArray<ULiveLinkDataPreviewComponent> Visualizers;
 
 	/** The billboard component showing the actor's root position. */
-	UPROPERTY(EditAnywhere, BlueprintReadWrite, Category = Components)
+	UPROPERTY(VisibleAnywhere, BlueprintReadWrite, Category = Components)
 	TObjectPtr<UBillboardComponent> BillboardComponent;
 
 protected:
@@ -88,4 +88,4 @@ protected:
 
 private:
 	TArray<FLiveLinkSubjectKey> CachedSubjects;
-};
\ No newline at end of file
+};
